Horde: Add additional artifact parameters to BuildGraph scripts.

#rnx

diff --git a/Engine/Source/Programs/AutomationTool/BuildGraph/BgSchema.cs b/Engine/Source/Programs/AutomationTool/BuildGraph/BgSchema.cs
--- a/Engine/Source/Programs/AutomationTool/BuildGraph/BgSchema.cs
+++ b/Engine/Source/Programs/AutomationTool/BuildGraph/BgSchema.cs
@@ -1,13 +1,13 @@
 // Copyright Epic Games, Inc. All Rights Reserved.
 
+using EpicGames.BuildGraph;
+using EpicGames.Core;
 using System;
 using System.Collections.Generic;
 using System.Diagnostics;
 using System.Linq;
 using System.Xml;
 using System.Xml.Schema;
-using EpicGames.BuildGraph;
-using EpicGames.Core;
 
 #nullable enable
 
@@ -647,6 +647,9 @@ namespace AutomationTool
 			XmlSchemaComplexType artifactType = new XmlSchemaComplexType();
 			artifactType.Name = GetTypeName(ScriptSchemaStandardType.Artifact);
 			artifactType.Attributes.Add(CreateSchemaAttribute("Name", ScriptSchemaStandardType.Name, XmlSchemaUse.Required));
+			artifactType.Attributes.Add(CreateSchemaAttribute("Type", ScriptSchemaStandardType.Name, XmlSchemaUse.Optional));
+			artifactType.Attributes.Add(CreateSchemaAttribute("Description", ScriptSchemaStandardType.BalancedString, XmlSchemaUse.Optional));
+			artifactType.Attributes.Add(CreateSchemaAttribute("BasePath", ScriptSchemaStandardType.BalancedString, XmlSchemaUse.Optional));
 			artifactType.Attributes.Add(CreateSchemaAttribute("Tag", ScriptSchemaStandardType.Tag, XmlSchemaUse.Optional));
 			artifactType.Attributes.Add(CreateSchemaAttribute("Keys", ScriptSchemaStandardType.BalancedString, XmlSchemaUse.Optional));
 			artifactType.Attributes.Add(CreateSchemaAttribute("If", ScriptSchemaStandardType.BalancedString, XmlSchemaUse.Optional));
diff --git a/Engine/Source/Programs/AutomationTool/BuildGraph/BgScriptReader.cs b/Engine/Source/Programs/AutomationTool/BuildGraph/BgScriptReader.cs
--- a/Engine/Source/Programs/AutomationTool/BuildGraph/BgScriptReader.cs
+++ b/Engine/Source/Programs/AutomationTool/BuildGraph/BgScriptReader.cs
@@ -1,5 +1,8 @@
 // Copyright Epic Games, Inc. All Rights Reserved.
 
+using EpicGames.BuildGraph;
+using EpicGames.Core;
+using Microsoft.Extensions.Logging;
 using System;
 using System.Collections.Generic;
 using System.Diagnostics.CodeAnalysis;
@@ -9,12 +12,7 @@ using System.Text;
 using System.Text.RegularExpressions;
 using System.Threading.Tasks;
 using System.Xml;
-using System.Xml.Linq;
 using System.Xml.Schema;
-using EpicGames.BuildGraph;
-using EpicGames.Core;
-using Microsoft.Extensions.Logging;
-using Newtonsoft.Json.Linq;
 using UnrealBuildBase;
 
 #nullable enable
@@ -582,14 +580,14 @@ namespace AutomationTool
 					LogError(element, "Property '{Propertyname}' has CreateInParentScope=\"true\" but has no parent scope.", name);
 					return;
 				}
-				else 
+				else
 				{
 					scopeIdx--;
 				}
 			}
 
 			if (_shadowProperties[scopeIdx].Contains(name))
-			{	
+			{
 				// Make sure this property name was not already used in a child scope; it likely indicates an error.
 				LogError(element, "Property '{PropertyName}' was already used in a child scope. Move this definition before the previous usage if they are intended to share scope, or use a different name.", name);
 			}
@@ -903,9 +901,9 @@ namespace AutomationTool
 				string input = ReadAttribute(element, "Input");
 				string method = ReadAttribute(element, "Method");
 				string output = ReadAttribute(element, "Output");
-				
+
 				string operationResult = string.Empty;
-				
+
 				string[] arguments = { };
 
 				const string ArgumentsName = "Arguments";
@@ -916,7 +914,7 @@ namespace AutomationTool
 				}
 
 				// Supply more string operations here
-				switch(method)
+				switch (method)
 				{
 					case "ToLower": operationResult = input.ToLower(); break;
 					case "ToUpper": operationResult = input.ToUpper(); break;
@@ -927,14 +925,14 @@ namespace AutomationTool
 						}
 						operationResult = input.Replace(arguments[0], arguments[1]);
 						break;
-					case "SplitFirst": 
+					case "SplitFirst":
 						if (arguments.Length != 1)
 						{
 							throw new AutomationException($"String operation 'SplitFirst' requires exactly 1 argument.");
 						}
 						operationResult = input.Split(arguments[0]).First();
 						break;
-					case "SplitLast": 
+					case "SplitLast":
 						if (arguments.Length != 1)
 						{
 							throw new AutomationException($"String operation 'SplitLast' requires exactly 1 argument.");
@@ -1223,6 +1221,24 @@ namespace AutomationTool
 			string? name;
 			if (await EvaluateConditionAsync(element) && TryReadObjectName(element, out name))
 			{
+				string? type = ReadAttribute(element, "Type");
+				if (String.IsNullOrEmpty(type))
+				{
+					type = null;
+				}
+
+				string? description = ReadAttribute(element, "Description");
+				if (String.IsNullOrEmpty(description))
+				{
+					description = null;
+				}
+
+				string? basePath = ReadAttribute(element, "BasePath");
+				if (String.IsNullOrEmpty(basePath))
+				{
+					basePath = null;
+				}
+
 				string tag = ReadAttribute(element, "Tag");
 				if (String.IsNullOrEmpty(tag))
 				{
@@ -1231,7 +1247,7 @@ namespace AutomationTool
 
 				string[] keys = ReadListAttribute(element, "Keys");
 
-				BgArtifactDef newArtifact = new BgArtifactDef(name, tag, keys);
+				BgArtifactDef newArtifact = new BgArtifactDef(name, type, description, basePath, tag, keys);
 				_graph.Artifacts.Add(newArtifact);
 			}
 		}
diff --git a/Engine/Source/Programs/Shared/EpicGames.BuildGraph/BgArtifactDef.cs b/Engine/Source/Programs/Shared/EpicGames.BuildGraph/BgArtifactDef.cs
--- a/Engine/Source/Programs/Shared/EpicGames.BuildGraph/BgArtifactDef.cs
+++ b/Engine/Source/Programs/Shared/EpicGames.BuildGraph/BgArtifactDef.cs
@@ -14,6 +14,21 @@ namespace EpicGames.BuildGraph
 		/// </summary>
 		public string Name { get; }
 
+		/// <summary>
+		/// Type of this artifact
+		/// </summary>
+		public string? Type { get; }
+
+		/// <summary>
+		/// Description for the artifact
+		/// </summary>
+		public string? Description { get; }
+
+		/// <summary>
+		/// Base path for files included in the artifact. Will be detected from the files specified if not set.
+		/// </summary>
+		public string? BasePath { get; }
+
 		/// <summary>
 		/// Tag to use for the artifact. Uses the artifact name by default.
 		/// </summary>
@@ -28,11 +43,17 @@ namespace EpicGames.BuildGraph
 		/// Constructor
 		/// </summary>
 		/// <param name="name">Name of the artifact</param>
+		/// <param name="type">Type of the artifact</param>
+		/// <param name="description">Description for the artifact</param>
+		/// <param name="basePath">Base path for files included in the artifact</param>
 		/// <param name="tag">Name of the tag producing this artifact</param>
 		/// <param name="keys">Keys that can be used to find the artifact</param>
-		public BgArtifactDef(string name, string tag, IReadOnlyList<string> keys)
+		public BgArtifactDef(string name, string? type, string? description, string? basePath, string tag, IReadOnlyList<string> keys)
 		{
 			Name = name;
+			Type = type;
+			Description = description;
+			BasePath = basePath;
 			Tag = tag;
 			Keys = keys;
 		}
diff --git a/Engine/Source/Programs/Shared/EpicGames.BuildGraph/BgGraphDef.cs b/Engine/Source/Programs/Shared/EpicGames.BuildGraph/BgGraphDef.cs
--- a/Engine/Source/Programs/Shared/EpicGames.BuildGraph/BgGraphDef.cs
+++ b/Engine/Source/Programs/Shared/EpicGames.BuildGraph/BgGraphDef.cs
@@ -433,6 +433,18 @@ namespace EpicGames.BuildGraph
 				{
 					jsonWriter.WriteObjectStart();
 					jsonWriter.WriteValue("Name", artifact.Name);
+					if (!String.IsNullOrEmpty(artifact.Type))
+					{
+						jsonWriter.WriteValue("Type", artifact.Type);
+					}
+					if (!String.IsNullOrEmpty(artifact.Description))
+					{
+						jsonWriter.WriteValue("Description", artifact.Description);
+					}
+					if (!String.IsNullOrEmpty(artifact.BasePath))
+					{
+						jsonWriter.WriteValue("BasePath", artifact.BasePath);
+					}
 					jsonWriter.WriteValue("Tag", artifact.Tag);
 					jsonWriter.WriteObjectEnd();
 				}
