Horde: Add output artifacts to the graph from the exported BuildGraph definition.

diff --git a/Engine/Source/Programs/AutomationTool/BuildGraph/BuildGraph.cs b/Engine/Source/Programs/AutomationTool/BuildGraph/BuildGraph.cs
--- a/Engine/Source/Programs/AutomationTool/BuildGraph/BuildGraph.cs
+++ b/Engine/Source/Programs/AutomationTool/BuildGraph/BuildGraph.cs
@@ -1,28 +1,23 @@
 // Copyright Epic Games, Inc. All Rights Reserved.
 
 using EpicGames.BuildGraph;
+using EpicGames.BuildGraph.Expressions;
+using EpicGames.Core;
+using EpicGames.MCP.Automation;
+using Microsoft.Extensions.Logging;
+using OpenTracing;
+using OpenTracing.Util;
 using System;
 using System.Collections.Generic;
-using System.ComponentModel;
-using System.Diagnostics;
 using System.IO;
 using System.Linq;
 using System.Reflection;
+using System.Runtime.InteropServices;
 using System.Text;
+using System.Threading.Tasks;
 using System.Xml;
-using EpicGames.Core;
-using EpicGames.MCP.Automation;
-using Microsoft.Extensions.Logging;
-using OpenTracing;
-using OpenTracing.Util;
 using UnrealBuildBase;
 using UnrealBuildTool;
-using System.Threading.Tasks;
-using EpicGames.BuildGraph.Expressions;
-using AutomationTool.Tasks;
-
-using static AutomationTool.CommandUtils;
-using System.Runtime.InteropServices;
 
 namespace AutomationTool
 {
@@ -243,11 +238,11 @@ namespace AutomationTool
 			string BranchOverride = ParseParamValue("Branch", null);
 
 			GraphPrintOptions PrintOptions = GraphPrintOptions.ShowCommandLineOptions;
-			if(ParseParam("ShowDeps"))
+			if (ParseParam("ShowDeps"))
 			{
 				PrintOptions |= GraphPrintOptions.ShowDependencies;
 			}
-			if(ParseParam("ShowNotifications"))
+			if (ParseParam("ShowNotifications"))
 			{
 				PrintOptions |= GraphPrintOptions.ShowNotifications;
 			}
@@ -259,9 +254,9 @@ namespace AutomationTool
 
 			// Parse any specific nodes to clean
 			List<string> CleanNodes = new List<string>();
-			foreach(string NodeList in ParseParamValues("CleanNode"))
+			foreach (string NodeList in ParseParamValues("CleanNode"))
 			{
-				foreach(string NodeName in NodeList.Split('+', ';'))
+				foreach (string NodeName in NodeList.Split('+', ';'))
 				{
 					CleanNodes.Add(NodeName);
 				}
@@ -273,9 +268,9 @@ namespace AutomationTool
 			int? CodeChange = P4Enabled ? P4Env.CodeChangelist : GetEnvVarIntOrNull(EnvVarNames.CodeChangelist);
 
 			// Set up the standard properties which build scripts might need
-			Dictionary<string, string> DefaultProperties = new Dictionary<string,string>(StringComparer.InvariantCultureIgnoreCase);
+			Dictionary<string, string> DefaultProperties = new Dictionary<string, string>(StringComparer.InvariantCultureIgnoreCase);
 			DefaultProperties["Branch"] = Branch ?? "Unknown";
-			DefaultProperties["Depot"] = (Branch != null && Branch.StartsWith("//", StringComparison.Ordinal))? Branch.Substring(2).Split('/').First() : "Unknown";
+			DefaultProperties["Depot"] = (Branch != null && Branch.StartsWith("//", StringComparison.Ordinal)) ? Branch.Substring(2).Split('/').First() : "Unknown";
 			DefaultProperties["EscapedBranch"] = String.IsNullOrEmpty(Branch) ? "Unknown" : CommandUtils.EscapePath(Branch);
 			DefaultProperties["Change"] = (Change ?? 0).ToString();
 			DefaultProperties["CodeChange"] = (CodeChange ?? 0).ToString();
@@ -327,10 +322,10 @@ namespace AutomationTool
 			foreach (string Param in Params)
 			{
 				const string SetPrefix = "set:";
-				if(Param.StartsWith(SetPrefix, StringComparison.InvariantCultureIgnoreCase))
+				if (Param.StartsWith(SetPrefix, StringComparison.InvariantCultureIgnoreCase))
 				{
 					int EqualsIdx = Param.IndexOf('=');
-					if(EqualsIdx >= 0)
+					if (EqualsIdx >= 0)
 					{
 						Arguments[Param.Substring(SetPrefix.Length, EqualsIdx - SetPrefix.Length)] = Param.Substring(EqualsIdx + 1);
 					}
@@ -341,10 +336,10 @@ namespace AutomationTool
 				}
 
 				const string AppendPrefix = "append:";
-				if(Param.StartsWith(AppendPrefix, StringComparison.InvariantCultureIgnoreCase))
+				if (Param.StartsWith(AppendPrefix, StringComparison.InvariantCultureIgnoreCase))
 				{
 					int EqualsIdx = Param.IndexOf('=');
-					if(EqualsIdx >= 0)
+					if (EqualsIdx >= 0)
 					{
 						string Property = Param.Substring(AppendPrefix.Length, EqualsIdx - AppendPrefix.Length);
 						string Value = Param.Substring(EqualsIdx + 1);
@@ -366,13 +361,13 @@ namespace AutomationTool
 
 			// Find all the tasks from the loaded assemblies
 			Dictionary<string, ScriptTaskBinding> NameToTask = new Dictionary<string, ScriptTaskBinding>();
-			if(!FindAvailableTasks(NameToTask, bPublicTasksOnly))
+			if (!FindAvailableTasks(NameToTask, bPublicTasksOnly))
 			{
 				return ExitCode.Error_Unknown;
 			}
 
 			// Generate documentation
-			if(DocumentationFileName != null)
+			if (DocumentationFileName != null)
 			{
 				WriteDocumentation(NameToTask, new FileReference(DocumentationFileName));
 				return ExitCode.Success;
@@ -458,26 +453,26 @@ namespace AutomationTool
 			}
 
 			// Export the graph for Horde
-			if(HordeExportFileName != null)
+			if (HordeExportFileName != null)
 			{
 				Graph.ExportForHorde(new FileReference(HordeExportFileName));
 			}
 
 			// Create the temp storage handler
 			DirectoryReference RootDir = new DirectoryReference(CommandUtils.CmdEnv.LocalRoot);
-			TempStorage Storage = new TempStorage(RootDir, DirectoryReference.Combine(RootDir, "Engine", "Saved", "BuildGraph"), (SharedStorageDir == null)? null : new DirectoryReference(SharedStorageDir), bWriteToSharedStorage);
-			if(!bResume)
+			TempStorage Storage = new TempStorage(RootDir, DirectoryReference.Combine(RootDir, "Engine", "Saved", "BuildGraph"), (SharedStorageDir == null) ? null : new DirectoryReference(SharedStorageDir), bWriteToSharedStorage);
+			if (!bResume)
 			{
 				Storage.CleanLocal();
 			}
-			foreach(string CleanNode in CleanNodes)
+			foreach (string CleanNode in CleanNodes)
 			{
 				Storage.CleanLocalNode(CleanNode);
 			}
 
 			// Convert the supplied target references into nodes 
 			HashSet<BgNodeDef> TargetNodes = new HashSet<BgNodeDef>();
-			if(TargetNames.Length == 0)
+			if (TargetNames.Length == 0)
 			{
 				if (!bListOnly && SingleNodeName == null)
 				{
@@ -504,7 +499,7 @@ namespace AutomationTool
 				foreach (string TargetName in NodesToResolve)
 				{
 					BgNodeDef[] Nodes;
-					if(!Graph.TryResolveReference(TargetName, out Nodes))
+					if (!Graph.TryResolveReference(TargetName, out Nodes))
 					{
 						Logger.LogError("Target '{TargetName}' is not in graph", TargetName);
 						return ExitCode.Error_Unknown;
@@ -514,7 +509,7 @@ namespace AutomationTool
 			}
 
 			// Try to acquire tokens for all the target nodes we want to build
-			if(TokenSignature != null)
+			if (TokenSignature != null)
 			{
 				// Find all the lock files
 				HashSet<FileReference> RequiredTokens = new HashSet<FileReference>(TargetNodes.SelectMany(x => x.RequiredTokens));
@@ -523,9 +518,9 @@ namespace AutomationTool
 				if (SingleNodeName == null)
 				{
 					Logger.LogInformation("Required tokens:");
-					foreach(BgNodeDef Node in TargetNodes)
+					foreach (BgNodeDef Node in TargetNodes)
 					{
-						foreach(FileReference RequiredToken in Node.RequiredTokens)
+						foreach (FileReference RequiredToken in Node.RequiredTokens)
 						{
 							Logger.LogInformation("  '{Node}' requires {RequiredToken}", Node, RequiredToken);
 						}
@@ -534,35 +529,35 @@ namespace AutomationTool
 
 				// Try to create all the lock files
 				List<FileReference> CreatedTokens = new List<FileReference>();
-				if(!bListOnly)
+				if (!bListOnly)
 				{
 					CreatedTokens.AddRange(RequiredTokens.Where(x => WriteTokenFile(x, TokenSignature)));
 				}
 
 				// Find all the tokens that we don't have
 				Dictionary<FileReference, string> MissingTokens = new Dictionary<FileReference, string>();
-				foreach(FileReference RequiredToken in RequiredTokens)
+				foreach (FileReference RequiredToken in RequiredTokens)
 				{
 					string CurrentOwner = ReadTokenFile(RequiredToken);
-					if(CurrentOwner != null && CurrentOwner != TokenSignature)
+					if (CurrentOwner != null && CurrentOwner != TokenSignature)
 					{
 						MissingTokens.Add(RequiredToken, CurrentOwner);
 					}
 				}
 
 				// If we want to skip all the nodes with missing locks, adjust the target nodes to account for it
-				if(MissingTokens.Count > 0)
+				if (MissingTokens.Count > 0)
 				{
-					if(bSkipTargetsWithoutTokens)
+					if (bSkipTargetsWithoutTokens)
 					{
 						// Find all the nodes we're going to skip
 						HashSet<BgNodeDef> SkipNodes = new HashSet<BgNodeDef>();
-						foreach(IGrouping<string, FileReference> MissingTokensForBuild in MissingTokens.GroupBy(x => x.Value, x => x.Key))
+						foreach (IGrouping<string, FileReference> MissingTokensForBuild in MissingTokens.GroupBy(x => x.Value, x => x.Key))
 						{
 							Logger.LogInformation("Skipping the following nodes due to {Arg0}:", MissingTokensForBuild.Key);
-							foreach(FileReference MissingToken in MissingTokensForBuild)
+							foreach (FileReference MissingToken in MissingTokensForBuild)
 							{
-								foreach(BgNodeDef SkipNode in TargetNodes.Where(x => x.RequiredTokens.Contains(MissingToken) && SkipNodes.Add(x)))
+								foreach (BgNodeDef SkipNode in TargetNodes.Where(x => x.RequiredTokens.Contains(MissingToken) && SkipNodes.Add(x)))
 								{
 									Logger.LogInformation("    {SkipNode}", SkipNode);
 								}
@@ -570,15 +565,15 @@ namespace AutomationTool
 						}
 
 						// Write a list of everything left over
-						if(SkipNodes.Count > 0)
+						if (SkipNodes.Count > 0)
 						{
 							TargetNodes.ExceptWith(SkipNodes);
 							Logger.LogInformation("Remaining target nodes:");
-							foreach(BgNodeDef TargetNode in TargetNodes)
+							foreach (BgNodeDef TargetNode in TargetNodes)
 							{
 								Logger.LogInformation("    {TargetNode}", TargetNode);
 							}
-							if(TargetNodes.Count == 0)
+							if (TargetNodes.Count == 0)
 							{
 								Logger.LogInformation("    None.");
 							}
@@ -586,12 +581,12 @@ namespace AutomationTool
 					}
 					else
 					{
-						foreach(KeyValuePair<FileReference, string> Pair in MissingTokens)
+						foreach (KeyValuePair<FileReference, string> Pair in MissingTokens)
 						{
 							List<BgNodeDef> SkipNodes = TargetNodes.Where(x => x.RequiredTokens.Contains(Pair.Key)).ToList();
 							Logger.LogError("Cannot run {Arg0} due to previous build: {Arg1}", String.Join(", ", SkipNodes), Pair.Value);
 						}
-						foreach(FileReference CreatedToken in CreatedTokens)
+						foreach (FileReference CreatedToken in CreatedTokens)
 						{
 							FileReference.Delete(CreatedToken);
 						}
@@ -616,36 +611,36 @@ namespace AutomationTool
 			{
 				FileReference PreprocessedFileLocation = new FileReference(PreprocessedFileName);
 				Logger.LogInformation("Writing {PreprocessedFileLocation}...", PreprocessedFileLocation);
-				Graph.Write(PreprocessedFileLocation, (SchemaFileName != null)? new FileReference(SchemaFileName) : null);
+				Graph.Write(PreprocessedFileLocation, (SchemaFileName != null) ? new FileReference(SchemaFileName) : null);
 				bListOnly = true;
 			}
 
 			// If we're just building a single node, find it 
 			BgNodeDef SingleNode = null;
-			if(SingleNodeName != null && !Graph.NameToNode.TryGetValue(SingleNodeName, out SingleNode))
+			if (SingleNodeName != null && !Graph.NameToNode.TryGetValue(SingleNodeName, out SingleNode))
 			{
 				Logger.LogError("Node '{SingleNodeName}' is not in the trimmed graph", SingleNodeName);
 				return ExitCode.Error_Unknown;
 			}
 
 			// If we just want to show the contents of the graph, do so and exit.
-			if(bListOnly)
-			{ 
+			if (bListOnly)
+			{
 				HashSet<BgNodeDef> CompletedNodes = FindCompletedNodes(Graph, Storage);
 				Graph.Print(CompletedNodes, PrintOptions, Log.Logger);
 			}
 
 			// Print out all the diagnostic messages which still apply, unless we're running a step as part of a build system or just listing the contents of the file. 
-			if(SingleNode == null && (!bListOnly || bShowDiagnostics))
+			if (SingleNode == null && (!bListOnly || bShowDiagnostics))
 			{
 				List<BgDiagnosticDef> Diagnostics = Graph.GetAllDiagnostics();
 				foreach (BgDiagnosticDef Diagnostic in Diagnostics)
 				{
-					if(Diagnostic.Level == LogLevel.Information)
+					if (Diagnostic.Level == LogLevel.Information)
 					{
 						Logger.LogInformation("{Arg0}({Arg1}): {Arg2}", Diagnostic.File, Diagnostic.Line, Diagnostic.Message);
 					}
-					else if(Diagnostic.Level == LogLevel.Warning)
+					else if (Diagnostic.Level == LogLevel.Warning)
 					{
 						Logger.LogWarning("{Arg0}({Arg1}): warning: {Arg2}", Diagnostic.File, Diagnostic.Line, Diagnostic.Message);
 					}
@@ -661,7 +656,7 @@ namespace AutomationTool
 			}
 
 			// Export the graph to a file
-			if(ExportFileName != null)
+			if (ExportFileName != null)
 			{
 				HashSet<BgNodeDef> CompletedNodes = FindCompletedNodes(Graph, Storage);
 				Graph.Print(CompletedNodes, PrintOptions, Log.Logger);
@@ -689,16 +684,16 @@ namespace AutomationTool
 			// Execute the command
 			if (!bListOnly)
 			{
-				if(SingleNode != null)
+				if (SingleNode != null)
 				{
-					if(!await BuildNodeAsync(Graph, SingleNode, NodeToExecutor, Storage, bWithBanner: true))
+					if (!await BuildNodeAsync(Graph, SingleNode, NodeToExecutor, Storage, bWithBanner: true))
 					{
 						return ExitCode.Error_Unknown;
 					}
 				}
 				else
 				{
-					if(!await BuildAllNodesAsync(Graph, NodeToExecutor, Storage))
+					if (!await BuildAllNodesAsync(Graph, NodeToExecutor, Storage))
 					{
 						return ExitCode.Error_Unknown;
 					}
@@ -774,7 +769,7 @@ namespace AutomationTool
 		{
 			IEnumerable<Assembly> LoadedScriptAssemblies = ScriptManager.AllScriptAssemblies;
 
-			if(bPublicTasksOnly)
+			if (bPublicTasksOnly)
 			{
 				LoadedScriptAssemblies = LoadedScriptAssemblies.Where(x => IsPublicAssembly(new FileReference(x.Location)));
 			}
@@ -791,16 +786,16 @@ namespace AutomationTool
 					continue;
 				}
 
-				foreach(Type Type in Types)
+				foreach (Type Type in Types)
 				{
-					foreach(TaskElementAttribute ElementAttribute in Type.GetCustomAttributes<TaskElementAttribute>())
+					foreach (TaskElementAttribute ElementAttribute in Type.GetCustomAttributes<TaskElementAttribute>())
 					{
-						if(!Type.IsSubclassOf(typeof(BgTaskImpl)))
+						if (!Type.IsSubclassOf(typeof(BgTaskImpl)))
 						{
 							Logger.LogError("Class '{Arg0}' has TaskElementAttribute, but is not derived from 'BgTaskImpl'", Type.Name);
 							return false;
 						}
-						if(NameToTask.ContainsKey(ElementAttribute.Name))
+						if (NameToTask.ContainsKey(ElementAttribute.Name))
 						{
 							Logger.LogError("Found multiple handlers for task elements called '{Arg0}'", ElementAttribute.Name);
 							return false;
@@ -818,7 +813,7 @@ namespace AutomationTool
 		/// <returns>Contents of the token, or null if it does not exist</returns>
 		public string ReadTokenFile(FileReference Location)
 		{
-			return FileReference.Exists(Location)? File.ReadAllText(Location.FullName) : null;
+			return FileReference.Exists(Location) ? File.ReadAllText(Location.FullName) : null;
 		}
 
 		/// <summary>
@@ -828,7 +823,7 @@ namespace AutomationTool
 		public bool WriteTokenFile(FileReference Location, string Signature)
 		{
 			// Check it doesn't already exist
-			if(FileReference.Exists(Location))
+			if (FileReference.Exists(Location))
 			{
 				return false;
 			}
@@ -845,7 +840,7 @@ namespace AutomationTool
 
 			// Create a temp file containing the owner name
 			string TempFileName;
-			for(int Idx = 0;;Idx++)
+			for (int Idx = 0; ; Idx++)
 			{
 				TempFileName = String.Format("{0}.{1}.tmp", Location.FullName, Idx);
 				try
@@ -857,9 +852,9 @@ namespace AutomationTool
 					}
 					break;
 				}
-				catch(IOException)
+				catch (IOException)
 				{
-					if(!File.Exists(TempFileName))
+					if (!File.Exists(TempFileName))
 					{
 						throw;
 					}
@@ -874,7 +869,7 @@ namespace AutomationTool
 			}
 			catch
 			{
-				if(!File.Exists(TempFileName))
+				if (!File.Exists(TempFileName))
 				{
 					throw;
 				}
@@ -890,10 +885,10 @@ namespace AutomationTool
 		static bool IsPublicAssembly(FileReference File)
 		{
 			DirectoryReference EngineDirectory = Unreal.EngineDirectory;
-			if(File.IsUnderDirectory(EngineDirectory))
+			if (File.IsUnderDirectory(EngineDirectory))
 			{
 				string[] PathFragments = File.MakeRelativeTo(EngineDirectory).Split(Path.DirectorySeparatorChar, Path.AltDirectorySeparatorChar);
-				if(PathFragments.All(x => !x.Equals("NotForLicensees", StringComparison.InvariantCultureIgnoreCase) && !x.Equals("NoRedist", StringComparison.InvariantCultureIgnoreCase)))
+				if (PathFragments.All(x => !x.Equals("NotForLicensees", StringComparison.InvariantCultureIgnoreCase) && !x.Equals("NoRedist", StringComparison.InvariantCultureIgnoreCase)))
 				{
 					return true;
 				}
@@ -909,9 +904,9 @@ namespace AutomationTool
 		HashSet<BgNodeDef> FindCompletedNodes(BgGraphDef Graph, TempStorage Storage)
 		{
 			HashSet<BgNodeDef> CompletedNodes = new HashSet<BgNodeDef>();
-			foreach(BgNodeDef Node in Graph.Agents.SelectMany(x => x.Nodes))
+			foreach (BgNodeDef Node in Graph.Agents.SelectMany(x => x.Nodes))
 			{
-				if(Storage.IsComplete(Node.Name))
+				if (Storage.IsComplete(Node.Name))
 				{
 					CompletedNodes.Add(Node);
 				}
@@ -934,9 +929,9 @@ namespace AutomationTool
 			// Check the integrity of any local nodes that have been completed. It's common to run formal builds locally between regular development builds, so we may have 
 			// stale local state. Rather than failing later, detect and clean them up now.
 			HashSet<BgNodeDef> CleanedNodes = new HashSet<BgNodeDef>();
-			foreach(BgNodeDef NodeToExecute in NodesToExecute)
+			foreach (BgNodeDef NodeToExecute in NodesToExecute)
 			{
-				if(NodeToExecute.InputDependencies.Any(x => CleanedNodes.Contains(x)) || !Storage.CheckLocalIntegrity(NodeToExecute.Name, NodeToExecute.Outputs.Select(x => x.TagName)))
+				if (NodeToExecute.InputDependencies.Any(x => CleanedNodes.Contains(x)) || !Storage.CheckLocalIntegrity(NodeToExecute.Name, NodeToExecute.Outputs.Select(x => x.TagName)))
 				{
 					Storage.CleanLocalNode(NodeToExecute.Name);
 					CleanedNodes.Add(NodeToExecute);
@@ -945,16 +940,16 @@ namespace AutomationTool
 
 			// Execute them in order
 			int NodeIdx = 0;
-			foreach(BgNodeDef NodeToExecute in NodesToExecute)
+			foreach (BgNodeDef NodeToExecute in NodesToExecute)
 			{
 				Logger.LogInformation("****** [{Arg0}/{Arg1}] {Arg2}", ++NodeIdx, NodesToExecute.Length, NodeToExecute.Name);
-				if(!Storage.IsComplete(NodeToExecute.Name))
+				if (!Storage.IsComplete(NodeToExecute.Name))
 				{
 					Logger.LogInformation("");
-					if(!await BuildNodeAsync(Graph, NodeToExecute, NodeToExecutor, Storage, bWithBanner: false))
+					if (!await BuildNodeAsync(Graph, NodeToExecute, NodeToExecutor, Storage, bWithBanner: false))
 					{
 						return false;
-					} 
+					}
 					Logger.LogInformation("");
 				}
 			}
@@ -1021,11 +1016,11 @@ namespace AutomationTool
 			using CleanupScriptRunner CleanupRunner = new CleanupScriptRunner(Logger);
 
 			// Create the mapping of tag names to file sets
-			Dictionary<string, HashSet<FileReference>> TagNameToFileSet = new Dictionary<string,HashSet<FileReference>>();
+			Dictionary<string, HashSet<FileReference>> TagNameToFileSet = new Dictionary<string, HashSet<FileReference>>();
 
 			// Read all the input tags for this node, and build a list of referenced input storage blocks
 			HashSet<TempStorageBlock> InputStorageBlocks = new HashSet<TempStorageBlock>();
-			foreach(BgNodeOutput Input in Node.Inputs)
+			foreach (BgNodeOutput Input in Node.Inputs)
 			{
 				TempStorageTagManifest FileList = Storage.ReadFileList(Input.ProducingNode.Name, Input.TagName);
 				TagNameToFileSet[Input.TagName] = FileList.ToFileSet(RootDir);
@@ -1047,13 +1042,13 @@ namespace AutomationTool
 
 			// Read all the input storage blocks, keeping track of which block each file came from
 			Dictionary<FileReference, TempStorageBlock> FileToStorageBlock = new Dictionary<FileReference, TempStorageBlock>();
-			foreach(KeyValuePair<TempStorageBlock, TempStorageManifest> Pair in InputManifests)
+			foreach (KeyValuePair<TempStorageBlock, TempStorageManifest> Pair in InputManifests)
 			{
 				TempStorageBlock InputStorageBlock = Pair.Key;
-				foreach(FileReference File in Pair.Value.Files.Select(x => x.ToFileReference(RootDir)))
+				foreach (FileReference File in Pair.Value.Files.Select(x => x.ToFileReference(RootDir)))
 				{
 					TempStorageBlock CurrentStorageBlock;
-					if(FileToStorageBlock.TryGetValue(File, out CurrentStorageBlock) && !TempStorage.IsDuplicateBuildProduct(File))
+					if (FileToStorageBlock.TryGetValue(File, out CurrentStorageBlock) && !TempStorage.IsDuplicateBuildProduct(File))
 					{
 						Logger.LogError("File '{File}' was produced by {InputStorageBlock} and {CurrentStorageBlock}", File, InputStorageBlock, CurrentStorageBlock);
 					}
@@ -1062,22 +1057,22 @@ namespace AutomationTool
 			}
 
 			// Add placeholder outputs for the current node
-			foreach(BgNodeOutput Output in Node.Outputs)
+			foreach (BgNodeOutput Output in Node.Outputs)
 			{
 				TagNameToFileSet.Add(Output.TagName, new HashSet<FileReference>());
 			}
 
 			// Execute the node
-			if(bWithBanner)
+			if (bWithBanner)
 			{
 				Console.WriteLine();
 				Logger.LogInformation("========== Starting: {Arg0} ==========", Node.Name);
 			}
-			if(!await NodeToExecutor[Node].ExecuteAsync(new JobContext(Node.Name, this), TagNameToFileSet))
+			if (!await NodeToExecutor[Node].ExecuteAsync(new JobContext(Node.Name, this), TagNameToFileSet))
 			{
 				return false;
 			}
-			if(bWithBanner)
+			if (bWithBanner)
 			{
 				Logger.LogInformation("========== Finished: {Arg0} ==========", Node.Name);
 				Console.WriteLine();
@@ -1085,24 +1080,24 @@ namespace AutomationTool
 
 			// Check that none of the inputs have been clobbered
 			Dictionary<string, string> ModifiedFiles = new Dictionary<string, string>(StringComparer.InvariantCultureIgnoreCase);
-			foreach(TempStorageFile File in InputManifests.Values.SelectMany(x => x.Files))
+			foreach (TempStorageFile File in InputManifests.Values.SelectMany(x => x.Files))
 			{
 				string Message;
-				if(!ModifiedFiles.ContainsKey(File.RelativePath) && !File.Compare(Unreal.RootDirectory, out Message))
+				if (!ModifiedFiles.ContainsKey(File.RelativePath) && !File.Compare(Unreal.RootDirectory, out Message))
 				{
 					// look up the previous nodes to help with error diagnosis
-					List<string> PreviousNodeNames = InputManifests.Where( x => x.Value.Files.Contains(File) )
-						.Select( x => x.Key.NodeName )
+					List<string> PreviousNodeNames = InputManifests.Where(x => x.Value.Files.Contains(File))
+						.Select(x => x.Key.NodeName)
 						.ToList();
 					if (PreviousNodeNames.Count > 0)
 					{
-						Message += $" (previous {(PreviousNodeNames.Count==1?"step":"steps")}: {string.Join(" + ", PreviousNodeNames)})";
+						Message += $" (previous {(PreviousNodeNames.Count == 1 ? "step" : "steps")}: {string.Join(" + ", PreviousNodeNames)})";
 					}
 
 					ModifiedFiles.Add(File.RelativePath, Message);
 				}
 			}
-			if(ModifiedFiles.Count > 0)
+			if (ModifiedFiles.Count > 0)
 			{
 				string modifiedFileList = "";
 				if (ModifiedFiles.Count < 100)
@@ -1114,19 +1109,19 @@ namespace AutomationTool
 					modifiedFileList = String.Join("\n", ModifiedFiles.Take(100).Select(x => x.Value));
 					modifiedFileList += $"{Environment.NewLine}And {ModifiedFiles.Count - 100} more.";
 				}
-				throw new AutomationException("Build {0} from a previous step have been modified:\n{1}", (ModifiedFiles.Count == 1)? "product" : "products", modifiedFileList);
+				throw new AutomationException("Build {0} from a previous step have been modified:\n{1}", (ModifiedFiles.Count == 1) ? "product" : "products", modifiedFileList);
 			}
 
 			// Determine all the output files which are required to be copied to temp storage (because they're referenced by nodes in another agent)
 			HashSet<FileReference> ReferencedOutputFiles = new HashSet<FileReference>();
-			foreach(BgAgentDef Agent in Graph.Agents)
+			foreach (BgAgentDef Agent in Graph.Agents)
 			{
 				bool bSameAgent = Agent.Nodes.Contains(Node);
-				foreach(BgNodeDef OtherNode in Agent.Nodes)
+				foreach (BgNodeDef OtherNode in Agent.Nodes)
 				{
-					if(!bSameAgent)
+					if (!bSameAgent)
 					{
-						foreach(BgNodeOutput Input in OtherNode.Inputs.Where(x => x.ProducingNode == Node))
+						foreach (BgNodeOutput Input in OtherNode.Inputs.Where(x => x.ProducingNode == Node))
 						{
 							ReferencedOutputFiles.UnionWith(TagNameToFileSet[Input.TagName]);
 						}
@@ -1136,16 +1131,16 @@ namespace AutomationTool
 
 			// Find a block name for all new outputs
 			Dictionary<FileReference, string> FileToOutputName = new Dictionary<FileReference, string>();
-			foreach(BgNodeOutput Output in Node.Outputs)
+			foreach (BgNodeOutput Output in Node.Outputs)
 			{
-				HashSet<FileReference> Files = TagNameToFileSet[Output.TagName]; 
-				foreach(FileReference File in Files)
+				HashSet<FileReference> Files = TagNameToFileSet[Output.TagName];
+				foreach (FileReference File in Files)
 				{
-					if(!FileToStorageBlock.ContainsKey(File) && File.IsUnderDirectory(RootDir))
+					if (!FileToStorageBlock.ContainsKey(File) && File.IsUnderDirectory(RootDir))
 					{
-						if(Output == Node.DefaultOutput)
+						if (Output == Node.DefaultOutput)
 						{
-							if(!FileToOutputName.ContainsKey(File))
+							if (!FileToOutputName.ContainsKey(File))
 							{
 								FileToOutputName[File] = "";
 							}
@@ -1153,7 +1148,7 @@ namespace AutomationTool
 						else
 						{
 							string OutputName;
-							if(FileToOutputName.TryGetValue(File, out OutputName) && OutputName.Length > 0)
+							if (FileToOutputName.TryGetValue(File, out OutputName) && OutputName.Length > 0)
 							{
 								FileToOutputName[File] = String.Format("{0}+{1}", OutputName, Output.TagName.Substring(1));
 							}
@@ -1168,10 +1163,10 @@ namespace AutomationTool
 
 			// Invert the dictionary to make a mapping of storage block to the files each contains
 			Dictionary<string, HashSet<FileReference>> OutputStorageBlockToFiles = new Dictionary<string, HashSet<FileReference>>();
-			foreach(KeyValuePair<FileReference, string> Pair in FileToOutputName)
+			foreach (KeyValuePair<FileReference, string> Pair in FileToOutputName)
 			{
 				HashSet<FileReference> Files;
-				if(!OutputStorageBlockToFiles.TryGetValue(Pair.Value, out Files))
+				if (!OutputStorageBlockToFiles.TryGetValue(Pair.Value, out Files))
 				{
 					Files = new HashSet<FileReference>();
 					OutputStorageBlockToFiles.Add(Pair.Value, Files);
@@ -1193,7 +1188,7 @@ namespace AutomationTool
 				}
 
 				// Find all the output tags that are published artifacts
-				Dictionary<string, BgArtifactDef> outputNameToArtifact = Graph.Artifacts.ToDictionary(x => x.Tag, x => x);
+				Dictionary<string, BgArtifactDef> outputNameToArtifact = Graph.Artifacts.ToDictionary(x => x.TagName, x => x);
 
 				// Publish all the output tags
 				foreach (BgNodeOutput Output in Node.Outputs)
@@ -1209,7 +1204,7 @@ namespace AutomationTool
 							StorageBlocks.Add(StorageBlock);
 						}
 					}
-					
+
 					IEnumerable<string> Keys = Enumerable.Empty<string>();
 					if (outputNameToArtifact.TryGetValue(Output.TagName, out BgArtifactDef artifact))
 					{
@@ -1259,17 +1254,17 @@ namespace AutomationTool
 
 			// Read documentation for each of them
 			Dictionary<string, XmlElement> MemberNameToElement = new Dictionary<string, XmlElement>();
-			foreach(Assembly TaskAssembly in TaskAssemblies)
+			foreach (Assembly TaskAssembly in TaskAssemblies)
 			{
 				string XmlFileName = Path.ChangeExtension(TaskAssembly.Location, ".xml");
-				if(File.Exists(XmlFileName))
+				if (File.Exists(XmlFileName))
 				{
 					// Read the document
 					XmlDocument Document = new XmlDocument();
 					Document.Load(XmlFileName);
 
 					// Parse all the members, and add them to the map
-					foreach(XmlElement Element in Document.SelectNodes("/doc/members/member"))
+					foreach (XmlElement Element in Document.SelectNodes("/doc/members/member"))
 					{
 						string Name = Element.GetAttribute("name");
 						MemberNameToElement.Add(Name, Element);
@@ -1278,7 +1273,7 @@ namespace AutomationTool
 			}
 
 			// Create the output directory
-			if(FileReference.Exists(OutputFile))
+			if (FileReference.Exists(OutputFile))
 			{
 				FileReference.MakeWriteable(OutputFile);
 			}
@@ -1429,14 +1424,14 @@ namespace AutomationTool
 				Writer.WriteLine("  </head>");
 				Writer.WriteLine("  <body>");
 				Writer.WriteLine("    <h1>BuildGraph Tasks</h1>");
-				foreach(string TaskName in NameToTask.Keys.OrderBy(x => x))
+				foreach (string TaskName in NameToTask.Keys.OrderBy(x => x))
 				{
 					// Get the task object
 					ScriptTaskBinding Task = NameToTask[TaskName];
 
 					// Get the documentation for this task
 					XmlElement TaskElement;
-					if(MemberNameToElement.TryGetValue("T:" + Task.TaskClass.FullName, out TaskElement))
+					if (MemberNameToElement.TryGetValue("T:" + Task.TaskClass.FullName, out TaskElement))
 					{
 						// Write the task heading
 						Writer.WriteLine("    <h2>{0}</h2>", TaskName);
@@ -1452,22 +1447,22 @@ namespace AutomationTool
 						Writer.WriteLine("      </tr>");
 
 						// Document the parameters
-						foreach(string ParameterName in Task.NameToParameter.Keys)
+						foreach (string ParameterName in Task.NameToParameter.Keys)
 						{
 							// Get the parameter data
 							ScriptTaskParameterBinding Parameter = Task.NameToParameter[ParameterName];
 
 							// Get the documentation for this parameter
 							XmlElement ParameterElement;
-							if(MemberNameToElement.TryGetValue("F:" + Parameter.FieldInfo.DeclaringType.FullName + "." + Parameter.Name, out ParameterElement))
+							if (MemberNameToElement.TryGetValue("F:" + Parameter.FieldInfo.DeclaringType.FullName + "." + Parameter.Name, out ParameterElement))
 							{
 								string TypeName = Parameter.FieldInfo.FieldType.Name;
-								if(Parameter.ValidationType != TaskParameterValidationType.Default)
+								if (Parameter.ValidationType != TaskParameterValidationType.Default)
 								{
 									StringBuilder NewTypeName = new StringBuilder(Parameter.ValidationType.ToString());
-									for(int Idx = 1; Idx < NewTypeName.Length; Idx++)
+									for (int Idx = 1; Idx < NewTypeName.Length; Idx++)
 									{
-										if(Char.IsLower(NewTypeName[Idx - 1]) && Char.IsUpper(NewTypeName[Idx]))
+										if (Char.IsLower(NewTypeName[Idx - 1]) && Char.IsUpper(NewTypeName[Idx]))
 										{
 											NewTypeName.Insert(Idx, ' ');
 										}
@@ -1478,7 +1473,7 @@ namespace AutomationTool
 								Writer.WriteLine("      <tr>");
 								Writer.WriteLine("         <td>{0}</td>", ParameterName);
 								Writer.WriteLine("         <td>{0}</td>", TypeName);
-								Writer.WriteLine("         <td>{0}</td>", Parameter.Optional? "Optional" : "Required");
+								Writer.WriteLine("         <td>{0}</td>", Parameter.Optional ? "Optional" : "Required");
 								Writer.WriteLine("         <td>{0}</td>", ParameterElement.SelectSingleNode("summary").InnerXml.Trim());
 								Writer.WriteLine("      </tr>");
 							}
@@ -1511,12 +1506,12 @@ namespace AutomationTool
 			string Text = Node.InnerXml;
 
 			StringBuilder Result = new StringBuilder();
-			for(int Idx = 0; Idx < Text.Length; Idx++)
+			for (int Idx = 0; Idx < Text.Length; Idx++)
 			{
-				if(Char.IsWhiteSpace(Text[Idx]))
+				if (Char.IsWhiteSpace(Text[Idx]))
 				{
 					Result.Append(' ');
-					while(Idx + 1 < Text.Length && Char.IsWhiteSpace(Text[Idx + 1]))
+					while (Idx + 1 < Text.Length && Char.IsWhiteSpace(Text[Idx + 1]))
 					{
 						Idx++;
 					}
diff --git a/Engine/Source/Programs/Horde/Horde.Agent/Execution/JobExecutor.cs b/Engine/Source/Programs/Horde/Horde.Agent/Execution/JobExecutor.cs
--- a/Engine/Source/Programs/Horde/Horde.Agent/Execution/JobExecutor.cs
+++ b/Engine/Source/Programs/Horde/Horde.Agent/Execution/JobExecutor.cs
@@ -27,6 +27,7 @@ using HordeCommon.Rpc.Tasks;
 using Microsoft.Extensions.Logging;
 using OpenTracing;
 using OpenTracing.Util;
+using CreateArtifactRequest = HordeCommon.Rpc.CreateArtifactRequest;
 
 namespace Horde.Agent.Execution
 {
@@ -41,11 +42,12 @@ namespace Horde.Agent.Execution
 		bool? Warnings,
 		IReadOnlyList<string> Inputs,
 		IReadOnlyList<string> OutputNames,
-		IList<int> PublishOutputs
+		IList<int> PublishOutputs,
+		IReadOnlyList<CreateArtifactRequest> Artifacts
 	)
 	{
 		public JobStepInfo(BeginStepResponse response)
-			: this(JobStepId.Parse(response.StepId), LogId.Parse(response.LogId), response.Name, response.Credentials, response.Properties, response.EnvVars, response.Warnings, response.Inputs, response.OutputNames, response.PublishOutputs)
+			: this(JobStepId.Parse(response.StepId), LogId.Parse(response.LogId), response.Name, response.Credentials, response.Properties, response.EnvVars, response.Warnings, response.Inputs, response.OutputNames, response.PublishOutputs, response.Artifacts)
 		{
 		}
 	}
@@ -124,12 +126,22 @@ namespace Horde.Agent.Execution
 			public string? Dependencies { get; set; }
 		}
 
+		protected class ExportedArtifact
+		{
+			public string Name { get; set; } = String.Empty;
+			public string? Type { get; set; }
+			public string? Description { get; set; }
+			public string? BasePath { get; set; }
+			public string OutputName { get; set; } = String.Empty;
+		}
+
 		protected class ExportedGraph
 		{
 			public List<ExportedGroup> Groups { get; set; } = new List<ExportedGroup>();
 			public List<ExportedAggregate> Aggregates { get; set; } = new List<ExportedAggregate>();
 			public List<ExportedLabel> Labels { get; set; } = new List<ExportedLabel>();
 			public List<ExportedBadge> Badges { get; set; } = new List<ExportedBadge>();
+			public List<ExportedArtifact> Artifacts { get; set; } = new List<ExportedArtifact>();
 		}
 
 		class TraceEvent
@@ -738,6 +750,19 @@ namespace Horde.Agent.Execution
 				updateGraph.Labels.Add(createLabel);
 			}
 
+			foreach (ExportedArtifact exportedArtifact in graph.Artifacts)
+			{
+				CreateArtifactRequest createArtifact = new CreateArtifactRequest();
+
+				createArtifact.Name = exportedArtifact.Name;
+				createArtifact.Type = exportedArtifact.Type ?? String.Empty;
+				createArtifact.Description = exportedArtifact.Description ?? String.Empty;
+				createArtifact.BasePath = exportedArtifact.BasePath ?? String.Empty;
+				createArtifact.OutputName = exportedArtifact.OutputName;
+
+				updateGraph.Artifacts.Add(createArtifact);
+			}
+
 			return updateGraph;
 		}
 
@@ -1113,6 +1138,54 @@ namespace Horde.Agent.Execution
 				logger.LogInformation("Upload took {Time:n1}s", timer.Elapsed.TotalSeconds);
 			}
 
+			// Create all the named artifacts. TODO: Merge this with regular temp storage artifacts?
+			foreach (CreateArtifactRequest graphArtifact in step.Artifacts)
+			{
+				HashSet<FileReference>? files;
+				if (!tagNameToFileSet.TryGetValue(graphArtifact.OutputName, out files))
+				{
+					logger.LogWarning("Missing output fileset for artifact {Name} (output={OutputName})", graphArtifact.Name, graphArtifact.OutputName);
+					continue;
+				}
+
+				using (GlobalTracer.Instance.BuildSpan("Artifact").WithTag("name", graphArtifact.Name).WithTag("resource", "Write").StartActive())
+				{
+					// Create the artifact
+					CreateJobArtifactRequestV2 artifactRequest = new CreateJobArtifactRequestV2();
+					artifactRequest.JobId = JobId.ToString();
+					artifactRequest.StepId = step.StepId.ToString();
+					artifactRequest.Name = graphArtifact.Name;
+					artifactRequest.Type = graphArtifact.Type;
+					artifactRequest.Description = graphArtifact.Description;
+
+					CreateJobArtifactResponseV2 artifact = await jobRpc.Client.CreateArtifactV2Async(artifactRequest, cancellationToken: cancellationToken);
+					logger.LogInformation("Created artifact {ArtifactId} '{ArtifactName}' ({ArtifactType}) with ref {RefName} ({RefUrl})", artifact.Id, artifactRequest.Name, ArtifactType.StepOutput, artifact.RefName, $"{Session.ServerUrl.ToString().TrimEnd('/')}/api/v1/storage/{artifact.NamespaceId}/refs/{artifact.RefName}");
+
+					using IStorageClient storage = CreateStorageClient(new NamespaceId(artifact.NamespaceId), artifact.Token);
+
+					// Upload the data
+					Stopwatch timer = Stopwatch.StartNew();
+
+					IBlobRef<DirectoryNode> outputNodeRef;
+					await using (IBlobWriter blobWriter = storage.CreateBlobWriter(artifact.RefName))
+					{
+						DirectoryNode outputNode = new DirectoryNode();
+
+						DirectoryReference baseDir = DirectoryReference.Combine(workspaceDir, graphArtifact.BasePath);
+						if (DirectoryReference.Exists(baseDir))
+						{
+							await outputNode.AddFilesAsync(baseDir, files, blobWriter, cancellationToken: cancellationToken);
+						}
+
+						outputNodeRef = await blobWriter.WriteBlobAsync(outputNode, cancellationToken: cancellationToken);
+					}
+
+					// Write the final node
+					await storage.WriteRefAsync(artifact.RefName, outputNodeRef, new RefOptions(), cancellationToken: cancellationToken);
+					logger.LogInformation("Upload took {Time:n1}s", timer.Elapsed.TotalSeconds);
+				}
+			}
+
 			return true;
 		}
 
diff --git a/Engine/Source/Programs/Horde/Horde.Server.Tests/Fixture.cs b/Engine/Source/Programs/Horde/Horde.Server.Tests/Fixture.cs
--- a/Engine/Source/Programs/Horde/Horde.Server.Tests/Fixture.cs
+++ b/Engine/Source/Programs/Horde/Horde.Server.Tests/Fixture.cs
@@ -131,6 +131,7 @@ namespace Horde.Server.Tests
 			public IReadOnlyList<INodeGroup> Groups { get; set; } = null!;
 			public IReadOnlyList<IAggregate> Aggregates { get; set; } = null!;
 			public IReadOnlyList<ILabel> Labels { get; set; } = null!;
+			public IReadOnlyList<IGraphArtifact> Artifacts { get; set; } = null!;
 		}
 	}
 }
\ No newline at end of file
diff --git a/Engine/Source/Programs/Horde/Horde.Server.Tests/Jobs/JobCollectionTests.cs b/Engine/Source/Programs/Horde/Horde.Server.Tests/Jobs/JobCollectionTests.cs
--- a/Engine/Source/Programs/Horde/Horde.Server.Tests/Jobs/JobCollectionTests.cs
+++ b/Engine/Source/Programs/Horde/Horde.Server.Tests/Jobs/JobCollectionTests.cs
@@ -9,6 +9,7 @@ using EpicGames.Horde.Agents;
 using EpicGames.Horde.Agents.Leases;
 using EpicGames.Horde.Agents.Pools;
 using EpicGames.Horde.Agents.Sessions;
+using EpicGames.Horde.Artifacts;
 using EpicGames.Horde.Jobs;
 using EpicGames.Horde.Jobs.Templates;
 using EpicGames.Horde.Logs;
@@ -811,5 +812,25 @@ namespace Horde.Server.Tests.Jobs
 			Assert.AreEqual(1, job.Batches[1].GroupIdx);
 			Assert.AreEqual(2, job.Batches[2].GroupIdx);
 		}
+
+		[TestMethod]
+		public async Task AddArtifactsAsync()
+		{
+			Mock<ITemplate> templateMock = new Mock<ITemplate>(MockBehavior.Strict);
+			templateMock.SetupGet(x => x.InitialAgentType).Returns((string?)null);
+
+			IGraph baseGraph = await GraphCollection.AddAsync(templateMock.Object, null);
+
+			List<NewGraphArtifact> newArtifacts = new List<NewGraphArtifact>();
+			newArtifacts.Add(new NewGraphArtifact(new ArtifactName("foo"), new ArtifactType("type"), "hello world", "Engine/Source", "fileset"));
+
+			IGraph graph = await GraphCollection.AppendAsync(baseGraph, newArtifactRequests: newArtifacts);
+			Assert.AreEqual(1, graph.Artifacts.Count);
+			Assert.AreEqual("foo", graph.Artifacts[0].Name.ToString());
+			Assert.AreEqual("type", graph.Artifacts[0].Type.ToString());
+			Assert.AreEqual("hello world", graph.Artifacts[0].Description);
+			Assert.AreEqual("Engine/Source", graph.Artifacts[0].BasePath);
+			Assert.AreEqual("fileset", graph.Artifacts[0].OutputName);
+		}
 	}
 }
diff --git a/Engine/Source/Programs/Horde/Horde.Server/Jobs/Graphs/GraphCollection.cs b/Engine/Source/Programs/Horde/Horde.Server/Jobs/Graphs/GraphCollection.cs
--- a/Engine/Source/Programs/Horde/Horde.Server/Jobs/Graphs/GraphCollection.cs
+++ b/Engine/Source/Programs/Horde/Horde.Server/Jobs/Graphs/GraphCollection.cs
@@ -7,6 +7,7 @@ using System.Linq;
 using System.Threading;
 using System.Threading.Tasks;
 using EpicGames.Core;
+using EpicGames.Horde.Artifacts;
 using Horde.Server.Jobs.Templates;
 using Horde.Server.Server;
 using HordeCommon;
@@ -189,9 +190,34 @@ namespace Horde.Server.Jobs.Graphs
 			}
 		}
 
+		class GraphArtifact : IGraphArtifact
+		{
+			public ArtifactName Name { get; set; }
+			public ArtifactType Type { get; set; }
+			public string Description { get; set; }
+			public string BasePath { get; set; }
+			public string OutputName { get; set; }
+
+			private GraphArtifact()
+			{
+				Description = String.Empty;
+				BasePath = String.Empty;
+				OutputName = String.Empty;
+			}
+
+			public GraphArtifact(ArtifactName name, ArtifactType type, string description, string basePath, string outputName)
+			{
+				Name = name;
+				Type = type;
+				Description = description;
+				BasePath = basePath;
+				OutputName = outputName;
+			}
+		}
+
 		class GraphDocument : IGraph
 		{
-			public static GraphDocument Empty { get; } = new GraphDocument(new List<NodeGroup>(), new List<Aggregate>(), new List<Label>());
+			public static GraphDocument Empty { get; } = new GraphDocument(new List<NodeGroup>(), new List<Aggregate>(), new List<Label>(), new List<GraphArtifact>());
 
 			[BsonRequired, BsonId]
 			public ContentHash Id { get; private set; } = ContentHash.Empty;
@@ -200,6 +226,7 @@ namespace Horde.Server.Jobs.Graphs
 			public List<NodeGroup> Groups { get; private set; } = new List<NodeGroup>();
 			public List<Aggregate> Aggregates { get; private set; } = new List<Aggregate>();
 			public List<Label> Labels { get; private set; } = new List<Label>();
+			public List<GraphArtifact> Artifacts { get; private set; } = new List<GraphArtifact>();
 
 			[BsonIgnore]
 			IReadOnlyDictionary<string, NodeRef>? _cachedNodeNameToRef;
@@ -210,21 +237,23 @@ namespace Horde.Server.Jobs.Graphs
 			IReadOnlyList<INodeGroup> IGraph.Groups => Groups;
 			IReadOnlyList<IAggregate> IGraph.Aggregates => Aggregates;
 			IReadOnlyList<ILabel> IGraph.Labels => Labels;
+			IReadOnlyList<IGraphArtifact> IGraph.Artifacts => Artifacts;
 
 			[BsonConstructor]
 			private GraphDocument()
 			{
 			}
 
-			public GraphDocument(List<NodeGroup> groups, List<Aggregate> aggregates, List<Label> labels)
+			public GraphDocument(List<NodeGroup> groups, List<Aggregate> aggregates, List<Label> labels, List<GraphArtifact> artifacts)
 			{
 				Groups = groups;
 				Aggregates = aggregates;
 				Labels = labels;
+				Artifacts = artifacts;
 				Id = ContentHash.SHA1(BsonExtensionMethods.ToBson(this));
 			}
 
-			public GraphDocument(GraphDocument baseGraph, List<NewGroup>? newGroupRequests, List<NewAggregate>? newAggregateRequests, List<NewLabel>? newLabelRequests)
+			public GraphDocument(GraphDocument baseGraph, List<NewGroup>? newGroupRequests, List<NewAggregate>? newAggregateRequests, List<NewLabel>? newLabelRequests, List<NewGraphArtifact>? newArtifactRequests)
 			{
 				Dictionary<string, NodeRef> nodeNameToRef = new Dictionary<string, NodeRef>(baseGraph.GetNodeNameToRef(), StringComparer.OrdinalIgnoreCase);
 				Dictionary<string, NodeOutputRef> nodeOutputNameToRef = new Dictionary<string, NodeOutputRef>(baseGraph.GetNodeOutputNameToRef(), StringComparer.OrdinalIgnoreCase);
@@ -295,10 +324,21 @@ namespace Horde.Server.Jobs.Graphs
 					}
 				}
 
+				// Update the list of artifacts
+				List<GraphArtifact> newArtifacts = new List<GraphArtifact>(baseGraph.Artifacts);
+				if (newArtifactRequests != null)
+				{
+					foreach (NewGraphArtifact newArtifactRequest in newArtifactRequests)
+					{
+						newArtifacts.Add(new GraphArtifact(newArtifactRequest.Name, newArtifactRequest.Type, newArtifactRequest.Description, newArtifactRequest.BasePath, newArtifactRequest.OutputName));
+					}
+				}
+
 				// Create the new arrays
 				Groups = newGroups;
 				Aggregates = newAggregates;
 				Labels = newLabels;
+				Artifacts = newArtifacts;
 
 				// Create the new graph, and save the generated node lookup into it
 				_cachedNodeNameToRef = nodeNameToRef;
@@ -413,15 +453,15 @@ namespace Horde.Server.Jobs.Graphs
 			Node node = new Node(IJob.SetupNodeName, null, null, Array.Empty<NodeRef>(), Array.Empty<NodeRef>(), Priority.High, true, false, true, null, null, null);
 			NodeGroup group = new NodeGroup(template.InitialAgentType ?? streamInitialAgentType ?? "Win64", new List<Node> { node });
 
-			GraphDocument graph = new GraphDocument(new List<NodeGroup> { group }, new List<Aggregate>(), new List<Label>());
+			GraphDocument graph = new GraphDocument(new List<NodeGroup> { group }, new List<Aggregate>(), new List<Label>(), new List<GraphArtifact>());
 			await AddAsync(graph, cancellationToken);
 			return graph;
 		}
 
 		/// <inheritdoc/>
-		public async Task<IGraph> AppendAsync(IGraph? baseGraph, List<NewGroup>? newGroupRequests, List<NewAggregate>? newAggregateRequests, List<NewLabel>? newLabelRequests, CancellationToken cancellationToken)
+		public async Task<IGraph> AppendAsync(IGraph? baseGraph, List<NewGroup>? newGroupRequests, List<NewAggregate>? newAggregateRequests, List<NewLabel>? newLabelRequests, List<NewGraphArtifact>? newArtifactRequests, CancellationToken cancellationToken)
 		{
-			GraphDocument graph = new GraphDocument((GraphDocument?)baseGraph ?? GraphDocument.Empty, newGroupRequests, newAggregateRequests, newLabelRequests);
+			GraphDocument graph = new GraphDocument((GraphDocument?)baseGraph ?? GraphDocument.Empty, newGroupRequests, newAggregateRequests, newLabelRequests, newArtifactRequests);
 			await AddAsync(graph, cancellationToken);
 			return graph;
 		}
diff --git a/Engine/Source/Programs/Horde/Horde.Server/Jobs/Graphs/IGraph.cs b/Engine/Source/Programs/Horde/Horde.Server/Jobs/Graphs/IGraph.cs
--- a/Engine/Source/Programs/Horde/Horde.Server/Jobs/Graphs/IGraph.cs
+++ b/Engine/Source/Programs/Horde/Horde.Server/Jobs/Graphs/IGraph.cs
@@ -6,6 +6,7 @@ using System.Diagnostics;
 using System.Diagnostics.CodeAnalysis;
 using System.Linq;
 using EpicGames.Core;
+using EpicGames.Horde.Artifacts;
 using HordeCommon;
 using MongoDB.Bson.Serialization.Attributes;
 
@@ -267,6 +268,37 @@ namespace Horde.Server.Jobs.Graphs
 		}
 	}
 
+	/// <summary>
+	/// Artifact produced by a graph
+	/// </summary>
+	public interface IGraphArtifact
+	{
+		/// <summary>
+		/// Name of the artifact
+		/// </summary>
+		public ArtifactName Name { get; }
+
+		/// <summary>
+		/// Type of the artifact
+		/// </summary>
+		public ArtifactType Type { get; }
+
+		/// <summary>
+		/// Description for the artifact
+		/// </summary>
+		public string Description { get; }
+
+		/// <summary>
+		/// Base path for files in the artifact
+		/// </summary>
+		public string BasePath { get; }
+
+		/// <summary>
+		/// Tag for the artifact files
+		/// </summary>
+		public string OutputName { get; }
+	}
+
 	/// <summary>
 	/// A unique dependency graph instance
 	/// </summary>
@@ -296,6 +328,11 @@ namespace Horde.Server.Jobs.Graphs
 		/// Status labels for this graph
 		/// </summary>
 		public IReadOnlyList<ILabel> Labels { get; }
+
+		/// <summary>
+		/// Artifacts for this graph
+		/// </summary>
+		public IReadOnlyList<IGraphArtifact> Artifacts { get; }
 	}
 
 	/// <summary>
@@ -636,4 +673,9 @@ namespace Horde.Server.Jobs.Graphs
 			Nodes = nodes;
 		}
 	}
+
+	/// <summary>
+	/// Information about an artifact
+	/// </summary>
+	public record class NewGraphArtifact(ArtifactName Name, ArtifactType Type, string Description, string BasePath, string OutputName);
 }
diff --git a/Engine/Source/Programs/Horde/Horde.Server/Jobs/Graphs/IGraphCollection.cs b/Engine/Source/Programs/Horde/Horde.Server/Jobs/Graphs/IGraphCollection.cs
--- a/Engine/Source/Programs/Horde/Horde.Server/Jobs/Graphs/IGraphCollection.cs
+++ b/Engine/Source/Programs/Horde/Horde.Server/Jobs/Graphs/IGraphCollection.cs
@@ -29,9 +29,10 @@ namespace Horde.Server.Jobs.Graphs
 		/// <param name="newGroupRequests">List of group requests</param>
 		/// <param name="newAggregateRequests">List of aggregate requests</param>
 		/// <param name="newLabelRequests">List of label requests</param>
+		/// <param name="newArtifactRequests">List of artifact requests</param>
 		/// <param name="cancellationToken">Cancellation token for the operation</param>
 		/// <returns>The new graph definition</returns>
-		Task<IGraph> AppendAsync(IGraph? baseGraph, List<NewGroup>? newGroupRequests = null, List<NewAggregate>? newAggregateRequests = null, List<NewLabel>? newLabelRequests = null, CancellationToken cancellationToken = default);
+		Task<IGraph> AppendAsync(IGraph? baseGraph, List<NewGroup>? newGroupRequests = null, List<NewAggregate>? newAggregateRequests = null, List<NewLabel>? newLabelRequests = null, List<NewGraphArtifact>? newArtifactRequests = null, CancellationToken cancellationToken = default);
 
 		/// <summary>
 		/// Gets the graph for a job
diff --git a/Engine/Source/Programs/Horde/Horde.Server/Jobs/JobMessages.cs b/Engine/Source/Programs/Horde/Horde.Server/Jobs/JobMessages.cs
--- a/Engine/Source/Programs/Horde/Horde.Server/Jobs/JobMessages.cs
+++ b/Engine/Source/Programs/Horde/Horde.Server/Jobs/JobMessages.cs
@@ -418,14 +418,14 @@ namespace Horde.Server.Jobs
 	/// <summary>
 	/// Response describing an artifact produced during a job
 	/// </summary>
-	/// <param name="Id">Identifier for this artifact</param>
+	/// <param name="Id">Identifier for this artifact, if it has been produced</param>
 	/// <param name="Name">Name of the artifact</param>
 	/// <param name="Type">Artifact type</param>
 	/// <param name="Description">Description to display for the artifact on the dashboard</param>
 	/// <param name="StepId">Step producing the artifact</param>
 	public record class GetJobArtifactResponse
 	(
-		ArtifactId Id,
+		ArtifactId? Id,
 		ArtifactName Name,
 		ArtifactType Type,
 		string? Description,
diff --git a/Engine/Source/Programs/Horde/Horde.Server/Jobs/JobsController.cs b/Engine/Source/Programs/Horde/Horde.Server/Jobs/JobsController.cs
--- a/Engine/Source/Programs/Horde/Horde.Server/Jobs/JobsController.cs
+++ b/Engine/Source/Programs/Horde/Horde.Server/Jobs/JobsController.cs
@@ -494,6 +494,10 @@ namespace Horde.Server.Jobs
 			}
 			if (includeArtifacts)
 			{
+				response.Artifacts = new List<GetJobArtifactResponse>();
+
+				HashSet<(ArtifactName, JobStepId)> addedArtifacts = new HashSet<(ArtifactName, JobStepId)>();
+
 				string artifactKey = $"job:{job.Id}";
 				string artifactStepKeyPrefix = $"job:{job.Id}/step:";
 				await foreach (IArtifact artifact in _artifactCollection.FindAsync(keys: new[] { artifactKey }, cancellationToken: cancellationToken))
@@ -503,8 +507,31 @@ namespace Horde.Server.Jobs
 						string? stepKey = artifact.Keys.FirstOrDefault(x => x.StartsWith(artifactStepKeyPrefix, StringComparison.Ordinal));
 						if (stepKey != null && JobStepId.TryParse(stepKey.Substring(artifactStepKeyPrefix.Length), out JobStepId jobStepId))
 						{
-							response.Artifacts ??= new List<GetJobArtifactResponse>();
 							response.Artifacts.Add(new GetJobArtifactResponse(artifact.Id, artifact.Name, artifact.Type, artifact.Description, jobStepId));
+							addedArtifacts.Add((artifact.Name, jobStepId));
+						}
+					}
+				}
+
+				Dictionary<string, IGraphArtifact> outputNameToArtifact = new Dictionary<string, IGraphArtifact>(StringComparer.OrdinalIgnoreCase);
+				foreach (IGraphArtifact artifact in graph.Artifacts)
+				{
+					outputNameToArtifact[artifact.OutputName] = artifact;
+				}
+
+				foreach (IJobStepBatch batch in job.Batches)
+				{
+					INodeGroup group = graph.Groups[batch.GroupIdx];
+					foreach (IJobStep step in batch.Steps)
+					{
+						INode node = group.Nodes[step.NodeIdx];
+						foreach (string outputName in node.OutputNames)
+						{
+							IGraphArtifact? graphArtifact;
+							if (outputNameToArtifact.TryGetValue(outputName, out graphArtifact) && !addedArtifacts.Contains((graphArtifact.Name, step.Id)))
+							{
+								response.Artifacts.Add(new GetJobArtifactResponse(null, graphArtifact.Name, graphArtifact.Type, graphArtifact.Description, step.Id));
+							}
 						}
 					}
 				}
@@ -911,7 +938,7 @@ namespace Horde.Server.Jobs
 				}
 
 				IGraph oldGraph = await _jobService.GetGraphAsync(job, cancellationToken);
-				IGraph newGraph = await _graphs.AppendAsync(oldGraph, requests, null, null, cancellationToken);
+				IGraph newGraph = await _graphs.AppendAsync(oldGraph, requests, null, null, null, cancellationToken);
 
 				IJob? newJob = await _jobService.TryUpdateGraphAsync(job, oldGraph, newGraph, cancellationToken);
 				if (newJob != null)
diff --git a/Engine/Source/Programs/Shared/EpicGames.BuildGraph/BgArtifactDef.cs b/Engine/Source/Programs/Shared/EpicGames.BuildGraph/BgArtifactDef.cs
--- a/Engine/Source/Programs/Shared/EpicGames.BuildGraph/BgArtifactDef.cs
+++ b/Engine/Source/Programs/Shared/EpicGames.BuildGraph/BgArtifactDef.cs
@@ -32,7 +32,7 @@ namespace EpicGames.BuildGraph
 		/// <summary>
 		/// Tag to use for the artifact. Uses the artifact name by default.
 		/// </summary>
-		public string Tag { get; }
+		public string TagName { get; }
 
 		/// <summary>
 		/// Keys that can be used to find the artifact
@@ -46,15 +46,15 @@ namespace EpicGames.BuildGraph
 		/// <param name="type">Type of the artifact</param>
 		/// <param name="description">Description for the artifact</param>
 		/// <param name="basePath">Base path for files included in the artifact</param>
-		/// <param name="tag">Name of the tag producing this artifact</param>
+		/// <param name="tagName">Name of the tag producing this artifact</param>
 		/// <param name="keys">Keys that can be used to find the artifact</param>
-		public BgArtifactDef(string name, string? type, string? description, string? basePath, string tag, IReadOnlyList<string> keys)
+		public BgArtifactDef(string name, string? type, string? description, string? basePath, string tagName, IReadOnlyList<string> keys)
 		{
 			Name = name;
 			Type = type;
 			Description = description;
 			BasePath = basePath;
-			Tag = tag;
+			TagName = tagName;
 			Keys = keys;
 		}
 
diff --git a/Engine/Source/Programs/Shared/EpicGames.BuildGraph/BgGraphDef.cs b/Engine/Source/Programs/Shared/EpicGames.BuildGraph/BgGraphDef.cs
--- a/Engine/Source/Programs/Shared/EpicGames.BuildGraph/BgGraphDef.cs
+++ b/Engine/Source/Programs/Shared/EpicGames.BuildGraph/BgGraphDef.cs
@@ -445,7 +445,7 @@ namespace EpicGames.BuildGraph
 					{
 						jsonWriter.WriteValue("BasePath", artifact.BasePath);
 					}
-					jsonWriter.WriteValue("Tag", artifact.Tag);
+					jsonWriter.WriteValue("OutputName", artifact.TagName);
 					jsonWriter.WriteObjectEnd();
 				}
 				jsonWriter.WriteArrayEnd();
