- Fixed IsStaged with CrashReportClient living in a staged build

diff --git a/Engine/Source/Programs/AutomationTool/Scripts/CopyBuildToStagingDirectory.Automation.cs b/Engine/Source/Programs/AutomationTool/Scripts/CopyBuildToStagingDirectory.Automation.cs
--- a/Engine/Source/Programs/AutomationTool/Scripts/CopyBuildToStagingDirectory.Automation.cs
+++ b/Engine/Source/Programs/AutomationTool/Scripts/CopyBuildToStagingDirectory.Automation.cs
@@ -1262,10 +1262,11 @@ namespace AutomationScripts
 
 					// write out a StagedBuild.ini file that's only in Staged builds so runtime knows the build is staged
 					DirectoryReference IntermediateDir = DirectoryReference.Combine(SC.EngineRoot, "Intermediate", "Build", "BuildCookRun");
-					FileReference StagedBuildConfigFile = FileReference.Combine(IntermediateDir, "StagedBuild.ini");
+					FileReference StagedBuildConfigFile = FileReference.Combine(IntermediateDir, $"StagedBuild_{SC.ShortProjectName}.ini");
 					DirectoryReference.CreateDirectory(IntermediateDir);
 					File.WriteAllText(StagedBuildConfigFile.FullName, "");
-					SC.StageFile(StagedFileType.UFS, StagedBuildConfigFile, new StagedFileReference($"Engine/Config/{StagedBuildConfigFile.GetFileName()}"));
+					// has to be outside of the .pak files since we use to to set the ProjectDir, which is needed to load .pak files!
+					SC.StageFile(StagedFileType.NonUFS, StagedBuildConfigFile, new StagedFileReference($"Engine/Config/{StagedBuildConfigFile.GetFileName()}"));
 
 					// Stage ICU internationalization data from Engine.
 					var ICUDataVersion = SC.StageTargetPlatform.ICUDataVersion;
diff --git a/Engine/Source/Runtime/Core/Private/GenericPlatform/GenericPlatformMisc.cpp b/Engine/Source/Runtime/Core/Private/GenericPlatform/GenericPlatformMisc.cpp
--- a/Engine/Source/Runtime/Core/Private/GenericPlatform/GenericPlatformMisc.cpp
+++ b/Engine/Source/Runtime/Core/Private/GenericPlatform/GenericPlatformMisc.cpp
@@ -1248,6 +1248,10 @@ const TCHAR* FGenericPlatformMisc::ProjectDir()
 {
 	FString& ProjectDir = TLazySingleton<FStaticData>::Get().ProjectDir;
 
+	// always call this at ProjectDir() time and cache it (internally) because there could be pak file interaction that can disable
+	// ini file loading
+	bool bIsStaged = FPaths::IsStaged();
+
 	// track if last time we called this function the .ini was ready and had fixed the GameName case
 	static bool bWasIniReady = false;
 	bool bIsIniReady = GConfig && GConfig->IsReadyForUse();
@@ -1277,7 +1281,7 @@ const TCHAR* FGenericPlatformMisc::ProjectDir()
 		ProjectDir.Reserve(FPlatformMisc::GetMaxPathLength());
 		if (FPlatformProperties::IsProgram())
 		{
-			if (FPaths::IsStaged())
+			if (bIsStaged)
 			{
 				// if staged, use the remapped location
 				ProjectDir = FString::Printf(TEXT("../../../%s/"), FApp::GetProjectName());
diff --git a/Engine/Source/Runtime/Core/Private/Misc/Paths.cpp b/Engine/Source/Runtime/Core/Private/Misc/Paths.cpp
--- a/Engine/Source/Runtime/Core/Private/Misc/Paths.cpp
+++ b/Engine/Source/Runtime/Core/Private/Misc/Paths.cpp
@@ -146,9 +146,13 @@ namespace UE4Paths_Private
 
 bool FPaths::IsStaged()
 {
-	// if the platform requires cooked data, then we can assume staged, but if we don't require cooked data, then
+	// if it's not a program and the platform requires cooked data, then we can assume staged, but if not, then
 	// check if it went through the staging process
-	static bool bIsStaged = FPlatformProperties::RequiresCookedData() || FileExists(Combine(EngineConfigDir(), TEXT("StagedBuild.ini")));
+	static bool bIsStaged =
+#if !IS_PROGRAM
+		FPlatformProperties::RequiresCookedData() ||
+#endif
+		FileExists(Combine(EngineConfigDir(), FString::Printf(TEXT("StagedBuild_%s.ini"), FApp::GetProjectName())));
 	return bIsStaged;
 }
 
