Use logcat to get Android log at runtime + fix missing obb files copy (re-insert the code block removed previously and use Android SDK v29 obb folder)

#rnx

diff --git a/Engine/Source/Programs/AutomationTool/Gauntlet/Platform/Android/Gauntlet.TargetDeviceAndroid.cs b/Engine/Source/Programs/AutomationTool/Gauntlet/Platform/Android/Gauntlet.TargetDeviceAndroid.cs
--- a/Engine/Source/Programs/AutomationTool/Gauntlet/Platform/Android/Gauntlet.TargetDeviceAndroid.cs
+++ b/Engine/Source/Programs/AutomationTool/Gauntlet/Platform/Android/Gauntlet.TargetDeviceAndroid.cs
@@ -359,6 +359,8 @@ namespace Gauntlet
 			string Package = Build.AndroidPackageName;
 			KillRunningProcess(Package);
 
+			EnablePermissions(Package);
+
 			// Install apk
 			CopyFileToDevice(Package, Build.SourceApkPath, string.Empty);
 
@@ -370,8 +372,6 @@ namespace Gauntlet
 
 			// Now, create and push the UECommandline.txt file
 			CopyCommandlineFile(AppConfig.CommandLine, Package, AppConfig.ProjectName, Build.UsesExternalFilesDir);
-
-			EnablePermissions(Package);
 		}
 
 		public IAppInstall CreateAppInstall(UnrealAppConfig AppConfig)
@@ -729,7 +729,7 @@ namespace Gauntlet
 			RunAdbDeviceCommand(CommandLine, true, false, true);
 		}
 		/// <summary>
-		/// Enable Android permissions which would otherwise block automation with permimssion requests
+		/// Enable Android permissions which would otherwise block automation with permission requests
 		/// </summary>
 		public void EnablePermissions(string AndroidPackageName)
 		{
@@ -751,7 +751,7 @@ namespace Gauntlet
 
 		protected void CopyOBBFiles(Dictionary<string, string> OBBFiles, string Package)
 		{
-			string OBBRemoteDestination = string.Format("{0}/obb/{1}", StoragePath, Package);
+			string OBBRemoteDestination = string.Format("{0}/Android/obb/{1}", StoragePath, Package);
 
 			// Remove all existing obbs
 			RunAdbDeviceCommand(string.Format("shell rm {0}", OBBRemoteDestination));
@@ -777,7 +777,7 @@ namespace Gauntlet
 				if (OBBMatch.Success)
 				{
 					string StrippedObb = Path.GetFileName(SourceFile.Replace(".Client.obb", ".obb").Replace(OBBMatch.Groups[1].ToString(), PackageVersion));
-					DestinationFile = StoragePath + "/obb/" + Package + "/" + StrippedObb;
+					DestinationFile = StoragePath + "/Android/obb/" + Package + "/" + StrippedObb;
 				}
 
 				DestinationFile = Regex.Replace(DestinationFile, "%STORAGE%", StoragePath, RegexOptions.IgnoreCase);
@@ -863,7 +863,7 @@ namespace Gauntlet
 			DeviceArtifactPath = Build.UsesExternalFilesDir ? DeviceExternalFilesSavedPath : DeviceExternalStorageSavedPath;
 
 			// path for OBB files
-			string OBBRemoteDestination = string.Format("{0}/obb/{1}", StoragePath, Build.AndroidPackageName);
+			string OBBRemoteDestination = string.Format("{0}/Android/obb/{1}", StoragePath, Build.AndroidPackageName);
 
 			Log.Info("DeviceBaseDir: " + DeviceBaseDir);
 			Log.Info("DeviceExternalStorageSavedPath: " + DeviceExternalStorageSavedPath);
@@ -880,6 +880,7 @@ namespace Gauntlet
 				Log.Info("Fully cleaning console before install...");
 				RunAdbDeviceCommand(string.Format("shell rm -r {0}/UnrealGame/*", StoragePath));
 				RunAdbDeviceCommand(string.Format("shell rm -r {0}/Android/data/{1}/*", StoragePath, Build.AndroidPackageName));
+				RunAdbDeviceCommand(string.Format("shell rm -r {0}/obb/{1}/*", StoragePath, Build.AndroidPackageName));
 				RunAdbDeviceCommand(string.Format("shell rm -r {0}/Android/obb/{1}/*", StoragePath, Build.AndroidPackageName));
 				RunAdbDeviceCommand(string.Format("shell rm -r {0}/Download/*", StoragePath));
 			}
@@ -946,7 +947,6 @@ namespace Gauntlet
 					Console.WriteLine("Populating Directory");
 					PopulateDirectoryMappings(Path.GetDirectoryName(DeviceArtifactPath));
 				}
-				Console.WriteLine("trying to copy files over.");
 				foreach (UnrealFileToCopy FileToCopy in AppConfig.FilesToCopy)
 				{
 					string PathToCopyTo = Path.Combine(LocalDirectoryMappings[FileToCopy.TargetBaseDirectory], FileToCopy.TargetRelativeLocation);
@@ -988,7 +988,7 @@ namespace Gauntlet
 					Match OBBMatch = Regex.Match(SrcFile, @"\.(\d+)\.com.*\.obb");
 					if (OBBMatch.Success)
 					{
-						DestPath = StoragePath + "/obb/" + Build.AndroidPackageName + "/" + SrcFile.Replace(".Client.obb", ".obb").Replace(OBBMatch.Groups[1].ToString(), PackageVersion);
+						DestPath = StoragePath + "/Android/obb/" + Build.AndroidPackageName + "/" + SrcFile.Replace(".Client.obb", ".obb").Replace(OBBMatch.Groups[1].ToString(), PackageVersion);
 					}
 
 					DestPath = Regex.Replace(DestPath, "%STORAGE%", StoragePath, RegexOptions.IgnoreCase);
@@ -1028,6 +1028,16 @@ namespace Gauntlet
 				}
 
 				EnablePermissions(Build.AndroidPackageName);
+
+				// Copy other file dependencies (including OBB files)
+				foreach (var KV in FilesToInstall)
+				{
+					string LocalFile = KV.Key;
+					string RemoteFile = KV.Value;
+
+					Console.WriteLine("Copying {0} to {1}", LocalFile, RemoteFile);
+					CopyFileToDevice(Build.AndroidPackageName, LocalFile, RemoteFile);
+				}
 			}
 			else
 			{
@@ -1055,12 +1065,15 @@ namespace Gauntlet
 
 		public string CommandLine { get; protected set; }
 
-		public AndroidAppInstall(TargetDeviceAndroid InDevice, string InName, string InAndroidPackageName, string InCommandLine)
+		public string AppTag { get; set; }
+
+		public AndroidAppInstall(TargetDeviceAndroid InDevice, string InName, string InAndroidPackageName, string InCommandLine, string InAppTag = "UE")
 		{
 			AndroidDevice = InDevice;
 			Name = InName;
 			AndroidPackageName = InAndroidPackageName;
 			CommandLine = InCommandLine;
+			AppTag = InAppTag;
 		}
 
 		public IAppInstance Run()
@@ -1123,8 +1136,7 @@ namespace Gauntlet
 		{
 			get
 			{
-				UpdateCachedLog();
-				return String.IsNullOrEmpty(ActivityLogCached) ? String.Empty : ActivityLogCached;
+				return LogProcess.Output;
 			}
 		}
 
@@ -1138,6 +1150,8 @@ namespace Gauntlet
 
 		public IProcessResult LaunchProcess;
 
+		protected IProcessResult LogProcess;
+
 		protected TargetDeviceAndroid AndroidDevice;
 
 		protected AndroidAppInstall Install;
@@ -1148,15 +1162,12 @@ namespace Gauntlet
 		private DateTime ActivityCheckTime = DateTime.UtcNow;
 		private bool ActivityExited = false;
 
-		private static readonly TimeSpan ActivityLogDelta = TimeSpan.FromSeconds(15);
-		private DateTime ActivityLogTime = DateTime.UtcNow - ActivityLogDelta;
-		private string ActivityLogCached = string.Empty;
-
 		public AndroidAppInstance(TargetDeviceAndroid InDevice, AndroidAppInstall InInstall, IProcessResult InProcess)
 		{
 			AndroidDevice = InDevice;
 			Install = InInstall;
 			LaunchProcess = InProcess;
+			LogProcess = AndroidDevice.RunAdbDeviceCommand($"logcat -s {Install.AppTag} debug Debug DEBUG -v raw", false, false);
 		}
 
 		public int WaitForExit()
@@ -1164,6 +1175,7 @@ namespace Gauntlet
 			if (!HasExited)
 			{
 				LaunchProcess.WaitForExit();
+				LogProcess.StopProcess();
 			}
 
 			return ExitCode;
@@ -1175,6 +1187,7 @@ namespace Gauntlet
 			{
 				WasKilled = true;
 				Install.AndroidDevice.KillRunningProcess(Install.AndroidPackageName);
+				LogProcess.StopProcess();
 			}
 		}
 
@@ -1188,16 +1201,6 @@ namespace Gauntlet
 			{
 				Log.Warning(KnownLogEvents.Gauntlet_DeviceEvent, "Failed to retrieve artifacts. {Output}", PullCmd.Output);
 			}
-			else
-			{
-				// update final cached stdout property
-				string LogFilename = Install.AndroidDevice.DeviceLogPath;
-				if (File.Exists(LogFilename))
-				{
-					ActivityLogCached = File.ReadAllText(LogFilename);
-					ActivityLogTime = DateTime.MinValue;
-				}
-			}
 
 			// pull the logcat over from device.
 			IProcessResult LogcatResult = Install.AndroidDevice.RunAdbDeviceCommand("logcat -d", bShouldLogCommand: Log.IsVerbose);
@@ -1242,55 +1245,18 @@ namespace Gauntlet
 			if (bHasExited)
 			{
 				ActivityExited = true;
-				// The activity has exited, make sure entire activity log has been captured, sleep to allow time for the log to flush
-				Thread.Sleep(5000);
-				UpdateCachedLog(true);
+				if (!LogProcess.HasExited)
+				{
+					// The activity has exited, make sure entire activity log has been captured, sleep to allow time for the log to flush
+					Thread.Sleep(5000);
+					LogProcess.StopProcess();
+				}
 				Log.VeryVerbose("{0}: process exited, Activity running={1}, Activity in foreground={2} ", ToString(), bActivityPresent.ToString(), bActivityInForeground.ToString());
 			}
 
 			return bHasExited;
 
 		}
-
-		/// <summary>
-		/// Updates cached activity log by running a shell command returning the full log from device (possibly over wifi)
-		/// The result is cached and updated at ActivityLogDelta frequency
-		/// </summary>
-		private void UpdateCachedLog(bool ForceUpdate = false)
-		{
-			if (!ForceUpdate && (ActivityLogTime == DateTime.MinValue || ((DateTime.UtcNow - ActivityLogTime) < ActivityLogDelta)))
-			{
-				return;
-			}
-
-			if (Install.AndroidDevice != null && Install.AndroidDevice.Disposed)
-			{
-				Log.Warning(KnownLogEvents.Gauntlet_DeviceEvent, "Attempting to cache log using disposed Android device");
-				return;
-			}
-
-			// Note: DeviceLogPath has the correct location for the current app. Apps can be configured to override the default log storage path. 
-			string GetLogCommand = string.Format("shell cat {0}", Install.AndroidDevice.DeviceLogPath);
-			IProcessResult LogQuery = Install.AndroidDevice.RunAdbDeviceCommand(GetLogCommand, true);
-
-			if (LogQuery.ExitCode != 0)
-			{
-				Log.VeryVerbose("Unable to query activity stdout on device {0}", Install.AndroidDevice.Name);
-			}
-			else
-			{
-				ActivityLogCached = LogQuery.Output;
-			}
-
-			ActivityLogTime = DateTime.UtcNow;
-
-			// the activity has exited, mark final log sentinel
-			if (ActivityExited)
-			{
-				ActivityLogTime = DateTime.MinValue;
-			}
-
-		}
 	}
 
 	// Device data from json
