Avoid force exit with FRequestExitCommand and (client) Boot Test in order to catch on shutdown issue

#rnx

diff --git a/Engine/Source/Programs/AutomationTool/Gauntlet/Unreal/Automation/UE.BootTest.cs b/Engine/Source/Programs/AutomationTool/Gauntlet/Unreal/Automation/UE.BootTest.cs
--- a/Engine/Source/Programs/AutomationTool/Gauntlet/Unreal/Automation/UE.BootTest.cs
+++ b/Engine/Source/Programs/AutomationTool/Gauntlet/Unreal/Automation/UE.BootTest.cs
@@ -49,6 +49,10 @@ namespace UE
 			UnrealTestConfiguration Config = base.GetConfiguration();
 
 			UnrealTestRole Client = Config.RequireRole(UnrealTargetRole.Client);
+			if (string.IsNullOrEmpty(GetCompletionString()))
+			{
+				Client.CommandLineParams.Add("ExecCmds", "Automation SoftQuit");
+			}
 
 			return Config;
 		}
@@ -70,7 +74,7 @@ namespace UE
 			// track our starting condition
 			LastLogTime = DateTime.Now;
 			LogLinesLastTick = 0;
-			DidDetectLaunch = true;
+			DidDetectLaunch = false;
 
 			return true;
 		}
@@ -81,7 +85,7 @@ namespace UE
 		/// <returns></returns>
 		protected virtual string GetCompletionString()
 		{
-			return "Engine is initialized. Leaving FEngineLoop::Init()";
+			return null;
 		}
 
 		/// <summary>
@@ -118,7 +122,6 @@ namespace UE
 			}
 
 			string CompletionString = GetCompletionString();
-
 			if (!string.IsNullOrEmpty(CompletionString))
 			{
 				if (LogParser.GetLogLinesContaining(CompletionString).Any())
@@ -143,7 +146,7 @@ namespace UE
 		{
 			if (Result == TestResult.Passed)
 			{
-				if (!DidDetectLaunch)
+				if (!string.IsNullOrEmpty(GetCompletionString()) && !DidDetectLaunch)
 				{
 					ReportError("Failed to detect completion of launch");
 				}
diff --git a/Engine/Source/Programs/AutomationTool/Gauntlet/Unreal/Base/Gauntlet.UnrealSession.cs b/Engine/Source/Programs/AutomationTool/Gauntlet/Unreal/Base/Gauntlet.UnrealSession.cs
--- a/Engine/Source/Programs/AutomationTool/Gauntlet/Unreal/Base/Gauntlet.UnrealSession.cs
+++ b/Engine/Source/Programs/AutomationTool/Gauntlet/Unreal/Base/Gauntlet.UnrealSession.cs
@@ -1111,7 +1111,7 @@ namespace Gauntlet
 			{
 				try
 				{
-					ArtifactLogFilePath = Path.Combine(DestinationDirectory.FullName, RoleName + "Output.log");
+					ArtifactLogFilePath = Path.GetFullPath(Path.Combine(DestinationDirectory.FullName, RoleName + "Output.log"));
 
 					// Write a short gauntlet blurb before the entire process log
 					using (StreamWriter Writer = new(ArtifactLogFilePath, false))
@@ -1125,9 +1125,21 @@ namespace Gauntlet
 					Log.Info($"Wrote {RoleName} Log to {ArtifactLogFilePath}");
 
 					// On build machines, copy all role logs to Horde.
-					if (IsBuildMachine)
+					if (IsBuildMachine && Horde.IsHordeJob)
 					{
-						string HordeLogFilePath = Path.Combine(CommandUtils.CmdEnv.LogFolder, RoleName + "Output.log");
+						// Extract the log path portion that includes the Gauntlet test name. ie: UE.BootTest(Win64_Test_Client)\Client\ClientOutput.log
+						// That is to handle situation where multiple tests are run within the same Gauntlet Session and logs get overwritten.
+						string LogName = ArtifactLogFilePath.Replace(Path.GetFullPath(InContext.Options.LogDir), "").TrimStart(Path.DirectorySeparatorChar);
+						if (Path.IsPathFullyQualified(LogName))
+						{
+							// The path was expected to be relative to LogDir, however it appeared to be an absolute path.
+							// So we revert to default behavior and save the log directly in UAT log folder.
+							LogName = RoleName + "Output.log";
+						}
+						string HordeLogFilePath = Path.GetFullPath(Path.Combine(CommandUtils.CmdEnv.LogFolder, LogName));
+						Log.Verbose($"Copy log for Horde to {HordeLogFilePath}");
+						string TargetDirectry = Path.GetDirectoryName(HordeLogFilePath);
+						if (!Directory.Exists(TargetDirectry)) { Directory.CreateDirectory(TargetDirectry); }
 						File.Copy(ArtifactLogFilePath, HordeLogFilePath, true);
 					}
 				}
diff --git a/Engine/Source/Runtime/Engine/Private/Tests/AutomationCommon.cpp b/Engine/Source/Runtime/Engine/Private/Tests/AutomationCommon.cpp
--- a/Engine/Source/Runtime/Engine/Private/Tests/AutomationCommon.cpp
+++ b/Engine/Source/Runtime/Engine/Private/Tests/AutomationCommon.cpp
@@ -706,7 +706,17 @@ bool FExitGameCommand::Update()
 
 bool FRequestExitCommand::Update()
 {
-	FPlatformMisc::RequestExit(true, TEXT("FRequestExitCommand"));
+	GLog->FlushThreadedLogs();
+#if WITH_EDITOR
+	if (GIsEditor)
+	{
+		GEngine->DeferredCommands.Add(TEXT("QUIT_EDITOR"));
+		return true;
+	}
+#endif
+	// force exit only if platform doesn't support quitting
+	FPlatformMisc::RequestExitWithStatus(!FPlatformProperties::SupportsQuit(), 0);
+
 	return true;
 }
 
