[Gauntlet] Fix issue with fatal error capture sorting
OrderBy should be inline and not in-place

#rnx

diff --git a/Engine/Source/Programs/AutomationTool/Gauntlet/Unreal/Utils/Gauntlet.UnrealLogParser.cs b/Engine/Source/Programs/AutomationTool/Gauntlet/Unreal/Utils/Gauntlet.UnrealLogParser.cs
--- a/Engine/Source/Programs/AutomationTool/Gauntlet/Unreal/Utils/Gauntlet.UnrealLogParser.cs
+++ b/Engine/Source/Programs/AutomationTool/Gauntlet/Unreal/Utils/Gauntlet.UnrealLogParser.cs
@@ -959,8 +959,7 @@ namespace Gauntlet
 			var Traces = ParseTracedErrors(ErrorMsgMatches).Concat(GetASanErrors());
 
 			// Keep the one with the most information.
-			Traces.OrderBy(T => T.Callstack.Length);
-			return Traces.Count() > 0 ? Traces.Last() : null;
+			return Traces.Count() > 0 ? Traces.OrderBy(T => T.Callstack.Length).Last() : null;
 		}
 
 		/// <summary>
