Add comment following code review suggestion

#rnx

diff --git a/Engine/Source/Programs/AutomationTool/Gauntlet/Unreal/Automation/UE.BootTest.cs b/Engine/Source/Programs/AutomationTool/Gauntlet/Unreal/Automation/UE.BootTest.cs
--- a/Engine/Source/Programs/AutomationTool/Gauntlet/Unreal/Automation/UE.BootTest.cs
+++ b/Engine/Source/Programs/AutomationTool/Gauntlet/Unreal/Automation/UE.BootTest.cs
@@ -85,6 +85,7 @@ namespace UE
 		/// <returns></returns>
 		protected virtual string GetCompletionString()
 		{
+			// Intentionally setup with no completion string as UnrealTestNode verified initialization string already
 			return null;
 		}
 
