Calling the CustomCopyHandler during CreatePakResponseFileFromStagingManifest so that the copy output file can be correctly added to the PakResponse

diff --git a/Engine/Source/Programs/AutomationTool/AutomationUtils/CustomStageCopyHandler.cs b/Engine/Source/Programs/AutomationTool/AutomationUtils/CustomStageCopyHandler.cs
--- a/Engine/Source/Programs/AutomationTool/AutomationUtils/CustomStageCopyHandler.cs
+++ b/Engine/Source/Programs/AutomationTool/AutomationUtils/CustomStageCopyHandler.cs
@@ -11,6 +11,15 @@ using Microsoft.Extensions.Logging;
 
 public abstract class CustomStageCopyHandler
 {
+	/// <summary>
+	/// Called when copying files to the staging directory.
+	/// </summary>
+	/// <returns>true if the file copy can be handled.</returns>
+	public virtual bool CanCopyFile(string SourceName)
+	{
+		return false;
+	}
+
 	/// <summary>
 	/// Called when copying files to the staging directory.
 	/// </summary>
diff --git a/Engine/Source/Programs/AutomationTool/Scripts/CopyBuildToStagingDirectory.Automation.cs b/Engine/Source/Programs/AutomationTool/Scripts/CopyBuildToStagingDirectory.Automation.cs
--- a/Engine/Source/Programs/AutomationTool/Scripts/CopyBuildToStagingDirectory.Automation.cs
+++ b/Engine/Source/Programs/AutomationTool/Scripts/CopyBuildToStagingDirectory.Automation.cs
@@ -2764,6 +2764,19 @@ namespace AutomationScripts
 					Src = NewIniFilename;
 				}
 
+				// Allow the custom stage copy handler to copy the file.
+				if (SC.CustomStageCopyHandler != null && SC.CustomStageCopyHandler.CanCopyFile(Src.FullName))
+				{
+					string SubFolder = Pair.Key.Name.Replace('/', Path.DirectorySeparatorChar);
+					FileReference NewFilename = FileReference.Combine(SC.ProjectRoot, "Saved", "Temp", SC.PlatformDir, SubFolder);
+					InternalUtils.SafeCreateDirectory(NewFilename.Directory.FullName, true);
+
+					if (SC.CustomStageCopyHandler.StageFile(Logger, Src.FullName, NewFilename.FullName))
+					{
+						Src = NewFilename;
+					}
+				}
+
 				// there can be files that only differ in case only, we don't support that in paks as paks are case-insensitive
 				if (UnrealPakResponseFile.ContainsKey(Src.FullName))
 				{
