Horde: Only enable code signing by default when running on an Epic build machine. The IsBuildMachine flag is true for any Horde deployment, but UAT hard-codes the signing certificate name for Epic. Check for the EpicInternal.txt flag to indicate whether we're running on Epic's farm.

diff --git a/Engine/Source/Programs/AutomationTool/AutomationUtils/ProjectParams.cs b/Engine/Source/Programs/AutomationTool/AutomationUtils/ProjectParams.cs
--- a/Engine/Source/Programs/AutomationTool/AutomationUtils/ProjectParams.cs
+++ b/Engine/Source/Programs/AutomationTool/AutomationUtils/ProjectParams.cs
@@ -864,7 +864,7 @@ namespace AutomationTool
 			this.OptionalFileStagingDirectory = ParseParamValueIfNotSpecified(Command, OptionalFileStagingDirectory, "optionalfilestagingdirectory", String.Empty, true);
 			this.OptionalFileInputDirectory = ParseParamValueIfNotSpecified(Command, OptionalFileInputDirectory, "optionalfileinputdirectory", String.Empty, true);
 			this.CookerSupportFilesSubdirectory = ParseParamValueIfNotSpecified(Command, CookerSupportFilesSubdirectory, "CookerSupportFilesSubdirectory", String.Empty, true);
-			this.bCodeSign = GetOptionalParamValueIfNotSpecified(Command, CodeSign, CommandUtils.IsBuildMachine, "CodeSign", "NoCodeSign");
+			this.bCodeSign = GetOptionalParamValueIfNotSpecified(Command, CodeSign, IsEpicBuildMachine(), "CodeSign", "NoCodeSign");
 			this.bTreatNonShippingBinariesAsDebugFiles = GetParamValueIfNotSpecified(Command, TreatNonShippingBinariesAsDebugFiles, false, "TreatNonShippingBinariesAsDebugFiles");
 			this.bUseExtraFlavor = GetParamValueIfNotSpecified(Command, UseExtraFlavor, false, "UseExtraFlavor");
 			this.Manifests = GetParamValueIfNotSpecified(Command, Manifests, this.Manifests, "manifests");
@@ -1338,6 +1338,12 @@ namespace AutomationTool
 			return ConfigValue;
 		}
 
+		static bool IsEpicBuildMachine()
+		{
+			return CommandUtils.IsBuildMachine 
+				&& FileReference.Exists(FileReference.Combine(Unreal.EngineDirectory, "Restricted", "NotForLicensees", "Build", "EpicInternal.txt"));
+		}
+
 		/// <summary>
 		/// Shared: Full path to the .uproject file
 		/// </summary>
