Adding skip build as an option for the compile task in buildgraph.
Skip build flag is passed to UBT so that only target files can be generated without performing a full build operation.

Fixing an issue in the CompileTask where AllowAllCores would always be false

diff --git a/Engine/Source/Programs/AutomationTool/AutomationUtils/UnrealBuild.cs b/Engine/Source/Programs/AutomationTool/AutomationUtils/UnrealBuild.cs
--- a/Engine/Source/Programs/AutomationTool/AutomationUtils/UnrealBuild.cs
+++ b/Engine/Source/Programs/AutomationTool/AutomationUtils/UnrealBuild.cs
@@ -146,7 +146,7 @@ namespace AutomationTool
 			}
 		}
 
-		void BuildWithUBT(List<BuildTarget> Targets, Dictionary<BuildTarget, BuildManifest> TargetToManifest, bool DisableXGE, bool AllCores)
+		void BuildWithUBT(List<BuildTarget> Targets, Dictionary<BuildTarget, BuildManifest> TargetToManifest, bool DisableXGE, bool AllCores, bool SkipBuild)
 		{
 			List<FileReference> ManifestFiles = new List<FileReference>(Targets.Count);
 
@@ -175,6 +175,10 @@ namespace AutomationTool
 			{
 				FullCommandLine.Append(" -AllCores");
 			}
+			if (SkipBuild)
+			{
+				FullCommandLine.Append(" -SkipBuild");
+			}
 
 			PrepareUBT();
 
@@ -524,7 +528,7 @@ namespace AutomationTool
 		/// <param name="InAllCores">If true AND XGE not present or not being used then ensure UBT uses all available cores</param>
 		/// <param name="InChangelistNumberOverride"></param>
 		/// <param name="InTargetToManifest"></param>
-		public void Build(BuildAgenda Agenda, bool? InDeleteBuildProducts = null, bool InUpdateVersionFiles = true, bool InForceNoXGE = false, bool InAllCores = false, int? InChangelistNumberOverride = null, Dictionary<BuildTarget, BuildManifest> InTargetToManifest = null)
+		public void Build(BuildAgenda InAgenda, bool? InDeleteBuildProducts = null, bool InUpdateVersionFiles = true, bool InForceNoXGE = false, bool InAllCores = false, int? InChangelistNumberOverride = null, Dictionary<BuildTarget, BuildManifest> InTargetToManifest = null, bool InSkipBuild = false)
 		{
 			if (!CommandUtils.CmdEnv.HasCapabilityToCompile)
 			{
@@ -540,7 +544,7 @@ namespace AutomationTool
 
 			// make a set of unique platforms involved
 			var UniquePlatforms = new List<UnrealTargetPlatform>();
-			foreach (var Target in Agenda.Targets)
+			foreach (var Target in InAgenda.Targets)
 			{
 				if (!UniquePlatforms.Contains(Target.Platform))
 				{
@@ -548,21 +552,21 @@ namespace AutomationTool
 				}
 			}
 
-			if (Agenda.SwarmAgentProject != "")
+			if (InAgenda.SwarmAgentProject != "")
 			{
-				string SwarmAgentSolution = Path.Combine(CommandUtils.CmdEnv.LocalRoot, Agenda.SwarmAgentProject);
+				string SwarmAgentSolution = Path.Combine(CommandUtils.CmdEnv.LocalRoot, InAgenda.SwarmAgentProject);
 				CommandUtils.BuildSolution(CommandUtils.CmdEnv, SwarmAgentSolution, "Development", "Mixed Platforms");
 				AddSwarmBuildProducts();
 			}
 
-			if (Agenda.SwarmCoordinatorProject != "")
+			if (InAgenda.SwarmCoordinatorProject != "")
 			{
-				string SwarmCoordinatorSolution = Path.Combine(CommandUtils.CmdEnv.LocalRoot, Agenda.SwarmCoordinatorProject);
+				string SwarmCoordinatorSolution = Path.Combine(CommandUtils.CmdEnv.LocalRoot, InAgenda.SwarmCoordinatorProject);
 				CommandUtils.BuildSolution(CommandUtils.CmdEnv, SwarmCoordinatorSolution, "Development", "Mixed Platforms");
 				AddSwarmBuildProducts();
 			}
 				
-			foreach (var DotNetProject in Agenda.DotNetProjects)
+			foreach (var DotNetProject in InAgenda.DotNetProjects)
 			{
 				string CsProj = Path.Combine(CommandUtils.CmdEnv.LocalRoot, DotNetProject);
 				CommandUtils.BuildCSharpProject(CommandUtils.CmdEnv, CsProj);
@@ -577,7 +581,7 @@ namespace AutomationTool
 			Logger.LogDebug("************************* UseXGE: {bCanUseXGE}", bCanUseXGE);
 
 			// Clean all the targets
-			foreach (BuildTarget Target in Agenda.Targets)
+			foreach (BuildTarget Target in InAgenda.Targets)
 			{
 				bool bClean = Target.Clean ?? DeleteBuildProducts;
 				if (bClean)
@@ -586,7 +590,7 @@ namespace AutomationTool
 				}
 			}
 
-			List<BuildTarget> Targets = new List<BuildTarget>(Agenda.Targets);
+			List<BuildTarget> Targets = new List<BuildTarget>(InAgenda.Targets);
 
 			// Temporary hack: iOS & tvOS configs need to build separately
 			if (Targets.Any(x => x.Platform == UnrealTargetPlatform.IOS || x.Platform == UnrealTargetPlatform.TVOS))
@@ -601,7 +605,7 @@ namespace AutomationTool
 						if (ConfigTargets.Count > 0)
 						{
 							// Build all the targets
-							BuildWithUBT(ConfigTargets, InTargetToManifest, bDisableXGE, InAllCores);
+							BuildWithUBT(ConfigTargets, InTargetToManifest, bDisableXGE, InAllCores, InSkipBuild);
 						}
 					}
 				}
@@ -615,7 +619,7 @@ namespace AutomationTool
 			}
 
 			// Build all the targets
-			BuildWithUBT(Targets, InTargetToManifest, bDisableXGE, InAllCores);
+			BuildWithUBT(Targets, InTargetToManifest, bDisableXGE, InAllCores, InSkipBuild);
 		}
 
 		/// <summary>
diff --git a/Engine/Source/Programs/AutomationTool/BuildGraph/Tasks/CompileTask.cs b/Engine/Source/Programs/AutomationTool/BuildGraph/Tasks/CompileTask.cs
--- a/Engine/Source/Programs/AutomationTool/BuildGraph/Tasks/CompileTask.cs
+++ b/Engine/Source/Programs/AutomationTool/BuildGraph/Tasks/CompileTask.cs
@@ -74,6 +74,12 @@ namespace AutomationTool.Tasks
 		[TaskParameter(Optional = true)]
 		public bool? Clean = null;
 
+		/// <summary>
+		/// Global flag passed to UBT that can be used to generate target files without fully compiling.
+		/// </summary>
+		[TaskParameter(Optional = true)]
+		public bool SkipBuild = false;
+
 		/// <summary>
 		/// Tag to be applied to build products of this task.
 		/// </summary>
@@ -106,6 +112,11 @@ namespace AutomationTool.Tasks
 		/// </summary>
 		bool bAllowAllCores = false;
 
+		/// <summary>
+		/// Should SkipBuild be passed to UBT so that only .target files are generated.
+		/// </summary>
+		bool bSkipBuild = false;
+
 		/// <summary>
 		/// Constructor
 		/// </summary>
@@ -128,15 +139,22 @@ namespace AutomationTool.Tasks
 				return false;
 			}
 
-			if(Targets.Count > 0)
+			CompileTaskParameters Parameters = CompileTask.Parameters;
+			if (Targets.Count > 0)
 			{
-				if (bAllowXGE != CompileTask.Parameters.AllowXGE)
+				if (bAllowXGE != Parameters.AllowXGE
+					|| bSkipBuild != Parameters.SkipBuild)
 				{
 					return false;
 				}
 			}
+			else
+			{
+				bAllowXGE = Parameters.AllowXGE;
+				bAllowAllCores = Parameters.AllowAllCores;
+				bSkipBuild = Parameters.SkipBuild;
+			}
 
-			CompileTaskParameters Parameters = CompileTask.Parameters;
 			bAllowXGE &= Parameters.AllowXGE;
 			bAllowAllCores &= Parameters.AllowAllCores;
 
@@ -167,8 +185,8 @@ namespace AutomationTool.Tasks
 			Dictionary<UnrealBuild.BuildTarget, BuildManifest> TargetToManifest = new Dictionary<UnrealBuild.BuildTarget,BuildManifest>();
 			UnrealBuild Builder = new UnrealBuild(Job.OwnerCommand);
 
-			bool bAllCores = (CommandUtils.IsBuildMachine || bAllowAllCores);	// Enable using all cores if this is a build agent or the flag was passed in to the task and XGE is disabled.
-			Builder.Build(Agenda, InDeleteBuildProducts: null, InUpdateVersionFiles: false, InForceNoXGE: !bAllowXGE, InAllCores: bAllCores, InTargetToManifest: TargetToManifest);
+			bool bAllCores = (CommandUtils.IsBuildMachine || bAllowAllCores);   // Enable using all cores if this is a build agent or the flag was passed in to the task and XGE is disabled.
+			Builder.Build(Agenda, InDeleteBuildProducts: null, InUpdateVersionFiles: false, InForceNoXGE: !bAllowXGE, InAllCores: bAllCores, InTargetToManifest: TargetToManifest, InSkipBuild: bSkipBuild);
 
 			UnrealBuild.CheckBuildProducts(Builder.BuildProductFiles);
 
