Add an option to force when exporting zen snapshots.

diff --git a/Engine/Source/Programs/AutomationTool/BuildGraph/Tasks/ZenExportSnapshotTask.cs b/Engine/Source/Programs/AutomationTool/BuildGraph/Tasks/ZenExportSnapshotTask.cs
--- a/Engine/Source/Programs/AutomationTool/BuildGraph/Tasks/ZenExportSnapshotTask.cs
+++ b/Engine/Source/Programs/AutomationTool/BuildGraph/Tasks/ZenExportSnapshotTask.cs
@@ -136,6 +136,13 @@ namespace AutomationTool.Tasks
 		/// </summary>
 		[TaskParameter(Optional = true)]
 		public string OverridePlatformCookedDir;
+
+		/// <summary>
+		/// Optional. Whether to force export of data even if the destination claims to have them.
+		/// </summary>
+		[TaskParameter(Optional = true)]
+		public bool Force = false;
+
 	}
 
 	/// <summary>
@@ -434,7 +441,11 @@ namespace AutomationTool.Tasks
 
 			// Format the command line
 			StringBuilder OplogExportCommandline = new StringBuilder();
-			OplogExportCommandline.AppendFormat("oplog-export");
+			OplogExportCommandline.Append("oplog-export");
+			if (Parameters.Force)
+			{
+				OplogExportCommandline.Append(" --force");
+			}
 
 			switch (DestinationStorageType)
 			{
