Adding systemic matcher to "There is not enough space on the disk." errors during LogChannelEvents

#rnx

diff --git a/Engine/Source/Programs/AutomationTool/AutomationUtils/Matchers/LogChannelEventMatcher.cs b/Engine/Source/Programs/AutomationTool/AutomationUtils/Matchers/LogChannelEventMatcher.cs
--- a/Engine/Source/Programs/AutomationTool/AutomationUtils/Matchers/LogChannelEventMatcher.cs
+++ b/Engine/Source/Programs/AutomationTool/AutomationUtils/Matchers/LogChannelEventMatcher.cs
@@ -25,6 +25,12 @@ namespace AutomationUtils.Matchers
 		{
 			if (input.TryMatch(s_pattern, out Match? match))
 			{
+
+				if (input.Contains("There is not enough space on the disk."))
+				{
+					return new LogEventBuilder(input).ToMatch(LogEventPriority.Low, LogLevel.Error, KnownLogEvents.Systemic_OutOfDiskSpace);
+				}
+
 				LogEventBuilder builder = new LogEventBuilder(input);
 				builder.Annotate(match.Groups["channel"], LogEventMarkup.Channel);
 				builder.Annotate(match.Groups["severity"], LogEventMarkup.Severity);
