Adding rate control for LogHeartbeatCategories to UnrealTestNode

diff --git a/Engine/Source/Programs/AutomationTool/Gauntlet/Unreal/Base/Gauntlet.UnrealTestConfiguration.cs b/Engine/Source/Programs/AutomationTool/Gauntlet/Unreal/Base/Gauntlet.UnrealTestConfiguration.cs
--- a/Engine/Source/Programs/AutomationTool/Gauntlet/Unreal/Base/Gauntlet.UnrealTestConfiguration.cs
+++ b/Engine/Source/Programs/AutomationTool/Gauntlet/Unreal/Base/Gauntlet.UnrealTestConfiguration.cs
@@ -513,13 +513,19 @@ namespace Gauntlet
 		/// </summary>
 		public float TimeoutBetweenAnyHeartbeats;
 
-		public UnrealHeartbeatOptions(float InHeartbeatPeriod = 30f, bool bShouldExpectHeartbeats = false, float InTimeoutBeforeFirstActiveHeartbeat = 0f, float InTimeoutBetweenActiveHeartbeats = 0f, float InTimeoutBetweenAnyHeartbeats = 90f)
+		/// <summary>
+		/// The minimum time interval between retrieval of heartbeat logs
+		/// </summary>
+		public float LogHeartbeatInterval;
+		
+		public UnrealHeartbeatOptions(float InHeartbeatPeriod = 30f, bool bShouldExpectHeartbeats = false, float InTimeoutBeforeFirstActiveHeartbeat = 0f, float InTimeoutBetweenActiveHeartbeats = 0f, float InTimeoutBetweenAnyHeartbeats = 90f, float InLogHeartbeatInterval = 0f)
 		{
 			HeartbeatPeriod = InHeartbeatPeriod;
 			bExpectHeartbeats = bShouldExpectHeartbeats;
 			TimeoutBeforeFirstActiveHeartbeat = InTimeoutBeforeFirstActiveHeartbeat;
 			TimeoutBetweenActiveHeartbeats = InTimeoutBetweenActiveHeartbeats;
 			TimeoutBetweenAnyHeartbeats = InTimeoutBetweenAnyHeartbeats;
+			LogHeartbeatInterval = InLogHeartbeatInterval;
 		}
 
 	}
diff --git a/Engine/Source/Programs/AutomationTool/Gauntlet/Unreal/Base/Gauntlet.UnrealTestNode.cs b/Engine/Source/Programs/AutomationTool/Gauntlet/Unreal/Base/Gauntlet.UnrealTestNode.cs
--- a/Engine/Source/Programs/AutomationTool/Gauntlet/Unreal/Base/Gauntlet.UnrealTestNode.cs
+++ b/Engine/Source/Programs/AutomationTool/Gauntlet/Unreal/Base/Gauntlet.UnrealTestNode.cs
@@ -341,7 +341,8 @@
 		
 		protected DateTime LastHeartbeatTime = DateTime.MinValue;
 		protected DateTime LastActiveHeartbeatTime = DateTime.MinValue;
-
+		private DateTime LastHeartbeatLogTime = DateTime.MinValue;
+		
 		// End  UnrealTestNode properties and members 
 
 		// artifact paths that have been used in this run
@@ -1020,62 +1021,11 @@
 		/// <param name="InInstance">The Unrealnstance being tested</param>
 		public virtual void TickTest(UnrealSessionInstance InInstance)
 		{
-			List<string> LogCategories = new List<string>();
-			LogCategories.Add("Gauntlet");
-			{
-				// get the categories used to monitor process (this needs rethought).
-				IEnumerable<string> HeartbeatCategories = GetHeartbeatLogCategories().Union(GetCachedConfiguration().LogCategoriesForEvents);
-				LogCategories.AddRange(HeartbeatCategories);
-			}
-			LogCategories = LogCategories.Distinct().ToList();
-
-			void LogHeartbeatCategories(ref int LastLogCount, IAppInstance App, string AppPrefix, bool bUpdateHeartbeatTime)
-			{
-				if (App != null)
-				{
-					UnrealLogStreamParser Parser = new UnrealLogStreamParser();
-					LastLogCount += Parser.ReadStream(App.StdOut, LastLogCount);
-
-					foreach (string TestLine in Parser.GetLogFromShortNameChannels(LogCategories))
-					{
-						Log.Info(string.Format("{0}: {1}", AppPrefix, TestLine));
-
-						if (bUpdateHeartbeatTime)
-						{
-							if (Regex.IsMatch(TestLine, @".*GauntletHeartbeat\: Active.*"))
-							{
-								LastHeartbeatTime = DateTime.Now;
-								LastActiveHeartbeatTime = DateTime.Now;
-							}
-							else if (Regex.IsMatch(TestLine, @".*GauntletHeartbeat\: Idle.*"))
-							{
-								LastHeartbeatTime = DateTime.Now;
-							}
-						}
-					}
-
-				}
-			}
-
-			if (InInstance.ServerApp != null)
-			{
-				bool bUpdateHeartbeat = InInstance.ClientApps == null;
-				LogHeartbeatCategories(ref LastServerLogCount, InInstance.ServerApp, "Server", bUpdateHeartbeat);
-			}
-			if (InInstance.ClientApps.Length > 0)
-			{
-				bool bUpdateHeartbeat = true;
-				LogHeartbeatCategories(ref LastClientLogCount, InInstance.ClientApps.First(), "Client", bUpdateHeartbeat);
-			}
-
-			if (InInstance.EditorApp != null)
-			{
-				bool bUpdateHeartbeat = false;
-				LogHeartbeatCategories(ref LastEditorLogCount, InInstance.EditorApp, "Editor", bUpdateHeartbeat);
-			}
-
+			// Retrieve and log heartbeat
+			LogHeartbeat(InInstance);
+			
 			// Detect missed heartbeats and fail the test
-			CheckHeartbeat();		
+			CheckHeartbeat();
 		}
 
 		/// <summary>
@@ -1802,6 +1752,76 @@
 		}
 
 
+		private void LogHeartbeat(UnrealSessionInstance InInstance)
+		{
+			if (CachedConfig == null ||
+			    GetTestStatus() != TestStatus.InProgress)
+			{
+				return;
+			}
+
+			UnrealHeartbeatOptions HeartbeatOptions = CachedConfig.HeartbeatOptions;
+			if (DateTime.Now.Subtract(LastHeartbeatLogTime).TotalSeconds < HeartbeatOptions.LogHeartbeatInterval)
+			{
+				return;
+			}
+			
+			List<string> LogCategories = new List<string>();
+			LogCategories.Add("Gauntlet");
+			{
+				// get the categories used to monitor process (this needs rethought).
+				IEnumerable<string> HeartbeatCategories = GetHeartbeatLogCategories().Union(GetCachedConfiguration().LogCategoriesForEvents);
+				LogCategories.AddRange(HeartbeatCategories);
+			}
+			LogCategories = LogCategories.Distinct().ToList();
+
+			void LogHeartbeatCategories(ref int LastLogCount, IAppInstance App, string AppPrefix, bool bUpdateHeartbeatTime)
+			{
+				if (App != null)
+				{
+					UnrealLogStreamParser Parser = new UnrealLogStreamParser();
+					LastLogCount += Parser.ReadStream(App.StdOut, LastLogCount);
+
+					foreach (string TestLine in Parser.GetLogFromShortNameChannels(LogCategories))
+					{
+						Log.Info(string.Format("{0}: {1}", AppPrefix, TestLine));
+
+						if (bUpdateHeartbeatTime)
+						{
+							if (Regex.IsMatch(TestLine, @".*GauntletHeartbeat\: Active.*"))
+							{
+								LastHeartbeatTime = DateTime.Now;
+								LastActiveHeartbeatTime = DateTime.Now;
+							}
+							else if (Regex.IsMatch(TestLine, @".*GauntletHeartbeat\: Idle.*"))
+							{
+								LastHeartbeatTime = DateTime.Now;
+							}
+						}
+					}
+				}
+			}
+
+			if (InInstance.ServerApp != null)
+			{
+				bool bUpdateHeartbeat = InInstance.ClientApps == null;
+				LogHeartbeatCategories(ref LastServerLogCount, InInstance.ServerApp, "Server", bUpdateHeartbeat);
+			}
+			if (InInstance.ClientApps.Length > 0)
+			{
+				bool bUpdateHeartbeat = true;
+				LogHeartbeatCategories(ref LastClientLogCount, InInstance.ClientApps.First(), "Client", bUpdateHeartbeat);
+			}
+
+			if (InInstance.EditorApp != null)
+			{
+				bool bUpdateHeartbeat = false;
+				LogHeartbeatCategories(ref LastEditorLogCount, InInstance.EditorApp, "Editor", bUpdateHeartbeat);
+			}
+
+			LastHeartbeatLogTime = DateTime.Now;
+		}
+
 		private void CheckHeartbeat()
 		{
 			if (CachedConfig == null 
