Fix problem making chunked builds due to CombinePaths combining quoted and unquoted paths together

diff --git a/Engine/Source/Programs/AutomationTool/AutomationUtils/CommandUtils.cs b/Engine/Source/Programs/AutomationTool/AutomationUtils/CommandUtils.cs
--- a/Engine/Source/Programs/AutomationTool/AutomationUtils/CommandUtils.cs
+++ b/Engine/Source/Programs/AutomationTool/AutomationUtils/CommandUtils.cs
@@ -1259,10 +1259,10 @@ namespace AutomationTool
 			var CombinedPath = new StringBuilder(CombinePathMaxLength);
 
 			// Combine all paths
-			CombinedPath.Append(Paths[0]);
+			CombinedPath.Append(Paths[0].Replace("\"", ""));
 			for (int PathIndex = 1; PathIndex < Paths.Length; ++PathIndex)
 			{
-				var NextPath = Paths[PathIndex];
+				var NextPath = Paths[PathIndex].Replace("\"", "");
 				if (String.IsNullOrEmpty(NextPath) == false)
 				{
 					int NextPathStartIndex = 0;
