Fix CombinePaths dereferencing null

diff --git a/Engine/Source/Programs/AutomationTool/AutomationUtils/CommandUtils.cs b/Engine/Source/Programs/AutomationTool/AutomationUtils/CommandUtils.cs
--- a/Engine/Source/Programs/AutomationTool/AutomationUtils/CommandUtils.cs
+++ b/Engine/Source/Programs/AutomationTool/AutomationUtils/CommandUtils.cs
@@ -1262,9 +1262,9 @@ namespace AutomationTool
 			CombinedPath.Append(Paths[0].Replace("\"", ""));
 			for (int PathIndex = 1; PathIndex < Paths.Length; ++PathIndex)
 			{
-				var NextPath = Paths[PathIndex].Replace("\"", "");
-				if (String.IsNullOrEmpty(NextPath) == false)
+				if (String.IsNullOrEmpty(Paths[PathIndex]) == false)
 				{
+					var NextPath = Paths[PathIndex].Replace("\"", "");
 					int NextPathStartIndex = 0;
 					if (CombinedPath.Length != 0)
 					{
