Fix Package on Host target platform
Agent should be based on host target platform
+ Support .app packaging for IOS in Gauntlet

#rnx

diff --git a/Engine/Source/Programs/AutomationTool/Gauntlet/Platform/IOS/Gauntlet.IOSBuildSource.cs b/Engine/Source/Programs/AutomationTool/Gauntlet/Platform/IOS/Gauntlet.IOSBuildSource.cs
--- a/Engine/Source/Programs/AutomationTool/Gauntlet/Platform/IOS/Gauntlet.IOSBuildSource.cs
+++ b/Engine/Source/Programs/AutomationTool/Gauntlet/Platform/IOS/Gauntlet.IOSBuildSource.cs
@@ -18,7 +18,9 @@ namespace Gauntlet
 
 		public UnrealTargetConfiguration Configuration { get; protected set; }
 
-		public string SourceIPAPath;
+		public string SourcePath;
+
+		public bool IsIPAFile;
 
 		public Dictionary<string, string> FilesToInstall;
 
@@ -32,14 +34,15 @@ namespace Gauntlet
 
 		public bool SupportsAdditionalFileCopy { get; }
 
-		public IOSBuild(UnrealTargetConfiguration InConfig, string InPackageName, string InIPAPath, Dictionary<string, string> InFilesToInstall, BuildFlags InFlags)
+		public IOSBuild(UnrealTargetConfiguration InConfig, string InPackageName, string InSourcePath, Dictionary<string, string> InFilesToInstall, BuildFlags InFlags)
 		{
 			Configuration = InConfig;
 			PackageName = InPackageName;
-			SourceIPAPath = InIPAPath;
+			SourcePath = InSourcePath;
 			FilesToInstall = InFilesToInstall;
 			Flags = InFlags;
 			SupportsAdditionalFileCopy = true;
+			IsIPAFile = Path.GetExtension(InSourcePath).Equals(".ipa", StringComparison.OrdinalIgnoreCase);
 		}
 
 		public bool CanSupportRole(UnrealTargetRole RoleType)
@@ -95,7 +98,7 @@ namespace Gauntlet
 					Log.Info(String.Format("unzip encountered an issue procesing IPA, possibly due to Zip64. Future steps may fail."));
 				}
 			}
-			
+
 			return true;
 		}
 
@@ -123,40 +126,60 @@ namespace Gauntlet
 
 				}
 			}
-			
+
 			return true;
 		}
 
 
-		private static string GetBundleIdentifier(string SourceIPA)
+		private static string GetBundleIdentifier(string Source)
 		{
-			string Output;
+			string PlistInfo = string.Empty;
+			bool IsIPAFile = Path.GetExtension(Source).Equals(".ipa", StringComparison.OrdinalIgnoreCase);
 
-			// Get a list of files in the IPA
-			if (!ExecuteIPAZipCommand(String.Format("-Z1 {0}", SourceIPA), out Output))
+			if (IsIPAFile)
 			{
-				Log.Warning(String.Format("Unable to list files for IPA {0}", SourceIPA));
-				return null;
-			}
+				string Output;
 
+				// Get a list of files in the IPA
+				if (!ExecuteIPAZipCommand(String.Format("-Z1 {0}", Source), out Output))
+				{
+					Log.Warning(String.Format("Unable to list files for IPA {0}", Source));
+					return null;
+				}
 
-			string[] Filenames = Regex.Split(Output, "\r\n|\r|\n");			
-			string PList = Filenames.Where(F => Regex.IsMatch(F.ToLower().Trim(), @"(payload\/)([^\/]+)(\/info\.plist)")).FirstOrDefault();
 
-			if (String.IsNullOrEmpty(PList))
-			{
-				Log.Warning(String.Format("Unable to find plist for IPA {0}", SourceIPA));
-				return null;
-			}
+				string[] Filenames = Regex.Split(Output, "\r\n|\r|\n");
+				string PList = Filenames.Where(F => Regex.IsMatch(F.ToLower().Trim(), @"(payload\/)([^\/]+)(\/info\.plist)")).FirstOrDefault();
 
-			// Get the plist info
-			if (!ExecuteIPAZipCommand(String.Format("-p '{0}' '{1}'", SourceIPA, PList), out Output))
-			{
-				Log.Warning(String.Format("Unable to extract plist data for IPA {0}", SourceIPA));
-				return null;
+				if (String.IsNullOrEmpty(PList))
+				{
+					Log.Warning(String.Format("Unable to find plist for IPA {0}", Source));
+					return null;
+				}
+
+				// Get the plist info
+				if (!ExecuteIPAZipCommand(String.Format("-p '{0}' '{1}'", Source, PList), out Output))
+				{
+					Log.Warning(String.Format("Unable to extract plist data for IPA {0}", Source));
+					return null;
+				}
+
+				PlistInfo = Output;
 			}
+			else
+			{
+				// Find plist file
+				DirectoryInfo Di = new DirectoryInfo(Source);
+				FileInfo PlistFile = new FileInfo(Path.Combine(Di.FullName, "Info.plist"));
+				if (!PlistFile.Exists)
+				{
+					Log.Warning(String.Format("Unable to find plist from {0}", Source));
+					return null;
+				}
 
-			string PlistInfo = Output;
+				StreamReader PListStream = new StreamReader(PlistFile.FullName);
+				PlistInfo = PListStream.ReadToEnd();
+			}
 
 			// todo: plist parsing, could be better
 			string PackageName = null;
@@ -174,7 +197,7 @@ namespace Gauntlet
 
 			if (String.IsNullOrEmpty(PackageName))
 			{
-				Log.Warning(String.Format("Unable to find CFBundleIdentifier in plist info for IPA {0}", SourceIPA));
+				Log.Warning(String.Format("Unable to find CFBundleIdentifier in plist info for App {0}", Source));
 				return null;
 			}
 
@@ -184,54 +207,50 @@ namespace Gauntlet
 
 		}
 
-		public static IEnumerable<IOSBuild> CreateFromPath(string InProjectName, string InPath)
+		public static IOSBuild CreateFromPath(string InProjectName, string InPath)
 		{
 			string BuildPath = InPath;
 
-			List<IOSBuild> DiscoveredBuilds = new List<IOSBuild>();
+			IOSBuild DiscoveredBuilds = null;
 
 			DirectoryInfo Di = new DirectoryInfo(BuildPath);
 
-			// find all install batchfiles
-			FileInfo[] InstallFiles = Di.GetFiles("*.ipa");
-
-			foreach (FileInfo Fi in InstallFiles)
-			{	
-
-				var UnrealConfig = UnrealHelpers.GetConfigurationFromExecutableName(InProjectName, Fi.Name);
-
-				Log.Verbose("Pulling package data from {0}", Fi.FullName);
+			var UnrealConfig = UnrealHelpers.GetConfigurationFromExecutableName(InProjectName, Di.Name);
 
-				string AbsPath = Fi.Directory.FullName;
-
-				// IOS builds are always packaged, and can always replace the command line and executable as we cache the unzip'd IPA
-				BuildFlags Flags = BuildFlags.Packaged | BuildFlags.CanReplaceCommandLine | BuildFlags.CanReplaceExecutable;
-
-				if (AbsPath.Contains("Bulk"))
-				{
-					Flags |= BuildFlags.Bulk;
-				}
-				else
+			if (UnrealConfig != UnrealTargetConfiguration.Unknown)
+			{
+				// check there's an executable with the right name 
+				string ShortName = Regex.Replace(InProjectName, "Game", "", RegexOptions.IgnoreCase);
+				FileInfo Executable = new DirectoryInfo(InPath).GetFiles().Where(Fi => Fi.Name.StartsWith(ShortName, StringComparison.OrdinalIgnoreCase)).FirstOrDefault();
+				if (Executable != null)
 				{
-					Flags |= BuildFlags.NotBulk;
-				}
+					Log.Verbose("Pulling package data from {0}", Di.FullName);
 
-				string SourceIPAPath = Fi.FullName;
-				string PackageName = GetBundleIdentifier(SourceIPAPath);
+					string AbsPath = Di.FullName;
 
-				if (String.IsNullOrEmpty(PackageName))
-				{
-					continue;
-				}
+					// IOS builds are always packaged, and can always replace the command line and executable as we cache the unzip'd IPA
+					BuildFlags Flags = BuildFlags.Packaged | BuildFlags.CanReplaceCommandLine | BuildFlags.CanReplaceExecutable;
 
-				Dictionary<string, string> FilesToInstall = new Dictionary<string, string>();
+					if (AbsPath.Contains("Bulk"))
+					{
+						Flags |= BuildFlags.Bulk;
+					}
+					else
+					{
+						Flags |= BuildFlags.NotBulk;
+					}
 
-				IOSBuild NewBuild = new IOSBuild(UnrealConfig, PackageName, SourceIPAPath, FilesToInstall, Flags);
+					string PackageName = GetBundleIdentifier(Di.FullName);
 
-				DiscoveredBuilds.Add(NewBuild);
+					if (!String.IsNullOrEmpty(PackageName))
+					{
+						Dictionary<string, string> FilesToInstall = new Dictionary<string, string>();
 
-				Log.Verbose("Found {0} {1} build at {2}", UnrealConfig, ((Flags & BuildFlags.Bulk) == BuildFlags.Bulk) ? "(bulk)" : "(not bulk)", AbsPath);
+						DiscoveredBuilds = new IOSBuild(UnrealConfig, PackageName, Di.FullName, FilesToInstall, Flags);
 
+						Log.Verbose("Found {0} {1} build at {2}", UnrealConfig, ((Flags & BuildFlags.Bulk) == BuildFlags.Bulk) ? "(bulk)" : "(not bulk)", AbsPath);
+					}
+				}
 			}
 
 			return DiscoveredBuilds;
@@ -261,8 +280,6 @@ namespace Gauntlet
 				return new List<IBuild>();
 			}
 
-			List<DirectoryInfo> AllDirs = new List<DirectoryInfo>();
-
 			List<IBuild> Builds = new List<IBuild>();
 
 			// c:\path\to\build
@@ -270,47 +287,48 @@ namespace Gauntlet
 
 			if (PathDI.Exists)
 			{
+				List<DirectoryInfo> SearchDirs = new List<DirectoryInfo>();
+
 				if (PathDI.Name.IndexOf("IOS", StringComparison.OrdinalIgnoreCase) >= 0)
 				{
-					AllDirs.Add(PathDI);
+					SearchDirs.Add(PathDI);
+				}
+				else
+				{
+					// find all directories that begin with IOS
+					SearchDirs.AddRange(PathDI.GetDirectories("IOS*", SearchOption.TopDirectoryOnly));
 				}
 
-				// find all directories that begin with IOS
-				DirectoryInfo[] IOSDirs = PathDI.GetDirectories("IOS*", SearchOption.TopDirectoryOnly);
-
-				AllDirs.AddRange(IOSDirs);
+				IEnumerable<DirectoryInfo> DirsToRecurse = new List<DirectoryInfo>(SearchDirs);
 
-				List<DirectoryInfo> DirsToRecurse = AllDirs;
+				List<DirectoryInfo> AllDirs = new List<DirectoryInfo>();
 
-				// now get subdirs
+				// get subdirs
 				while (MaxRecursion-- > 0)
 				{
-					List<DirectoryInfo> DiscoveredDirs = new List<DirectoryInfo>();
+					IEnumerable<DirectoryInfo> DiscoveredDirs = DirsToRecurse.SelectMany(D => D.GetDirectories("*", SearchOption.TopDirectoryOnly));
 
-					DirsToRecurse.ToList().ForEach((D) =>
-					{
-						DiscoveredDirs.AddRange(D.GetDirectories("*", SearchOption.TopDirectoryOnly));
-					});
+					// mac packages are folders so we only want things that end with .app
+					IEnumerable<DirectoryInfo> Packages = DiscoveredDirs.Where(D => Path.GetExtension(D.Name).Equals(".app", StringComparison.OrdinalIgnoreCase));
 
-					AllDirs.AddRange(DiscoveredDirs);
-					DirsToRecurse = DiscoveredDirs;
-				}
+					AllDirs.AddRange(Packages);
 
-				//IOSBuildSource BuildSource = null;
+					// don't recurse into the dirs that end with .app since those are builds we'll look at
+					DirsToRecurse = DiscoveredDirs.Except(Packages);
+				}
 
 				string IOSBuildFilter = Globals.Params.ParseValue("IOSBuildFilter", "");
 				foreach (DirectoryInfo Di in AllDirs)
 				{
-					IEnumerable<IOSBuild> FoundBuilds = IOSBuild.CreateFromPath(InProjectName, Di.FullName);
+					IOSBuild FoundBuild = IOSBuild.CreateFromPath(InProjectName, Di.FullName);
 
-					if (FoundBuilds != null)
+					if (FoundBuild != null)
 					{
-						if (!string.IsNullOrEmpty(IOSBuildFilter))
+						if (!string.IsNullOrEmpty(IOSBuildFilter) && FoundBuild.SourcePath.IndexOf(IOSBuildFilter, StringComparison.OrdinalIgnoreCase) >= 0)
 						{
-							//IndexOf used because Contains must be case-sensitive
-							FoundBuilds = FoundBuilds.Where(B => B.SourceIPAPath.IndexOf(IOSBuildFilter, StringComparison.OrdinalIgnoreCase) >= 0);
+							continue;
 						}
-						Builds.AddRange(FoundBuilds);
+						Builds.Add(FoundBuild);
 					}
 				}
 			}
diff --git a/Engine/Source/Programs/AutomationTool/Gauntlet/Platform/IOS/Gauntlet.TargetDeviceIOS.cs b/Engine/Source/Programs/AutomationTool/Gauntlet/Platform/IOS/Gauntlet.TargetDeviceIOS.cs
--- a/Engine/Source/Programs/AutomationTool/Gauntlet/Platform/IOS/Gauntlet.TargetDeviceIOS.cs
+++ b/Engine/Source/Programs/AutomationTool/Gauntlet/Platform/IOS/Gauntlet.TargetDeviceIOS.cs
@@ -386,15 +386,15 @@ namespace Gauntlet
 		}
 
 		/// <summary>
-		/// Checks whether version of deployed bundle matches local IPA
+		/// Checks whether version of deployed bundle matches local Build
 		/// </summary>
-		bool CheckDeployedIPA(IOSBuild Build)
+		bool CheckDeployedApp(IOSBuild Build)
 		{
 			try
 			{
-				Log.Verbose("Checking deployed IPA hash");
+				Log.Verbose("Checking deployed App hash");
 
-				string CommandLine = String.Format("--bundle_id {0} --download={1} --to {2}", Build.PackageName, "/Documents/IPAHash.txt", LocalCachePath);
+				string CommandLine = String.Format("--bundle_id {0} --download={1} --to {2}", Build.PackageName, "/Documents/AppHash.txt", LocalCachePath);
 				IProcessResult Result = ExecuteIOSDeployCommand(CommandLine, 120);
 
 				if (Result.ExitCode != 0)
@@ -402,12 +402,12 @@ namespace Gauntlet
 					return false;
 				}
 
-				string Hash = File.ReadAllText(LocalCachePath + "/Documents/IPAHash.txt").Trim();
-				string StoredHash = File.ReadAllText(IPAHashFilename).Trim();
+				string Hash = File.ReadAllText(LocalCachePath + "/Documents/AppHash.txt").Trim();
+				string StoredHash = File.ReadAllText(AppHashFilename).Trim();
 
 				if (Hash == StoredHash)
 				{
-					Log.Verbose("Deployed app hash matched cached IPA hash");
+					Log.Verbose("Deployed app hash matched cached App hash");
 					return true;
 				}
 
@@ -416,11 +416,11 @@ namespace Gauntlet
 			{
 				if (!Ex.Message.Contains("is denied"))
 				{
-					Log.Verbose("Unable to pull cached IPA cache from device, cached file may not exist: {0}", Ex.Message);
+					Log.Verbose("Unable to pull cached App cache from device, cached file may not exist: {0}", Ex.Message);
 				}
 			}
 
-			Log.Verbose("Deployed app hash doesn't match, IPA will be installed");
+			Log.Verbose("Deployed app hash doesn't match, App will be installed");
 			return false;
 
 		}
@@ -509,8 +509,8 @@ namespace Gauntlet
 			}
 		}
 
-		// We need to lock around setting up the IPA
-		static object IPALock = new object();
+		// We need to lock around setting up the App
+		static object AppLock = new object();
 		public IAppInstall InstallApplication(UnrealAppConfig AppConfig)
 		{
             IOSBuild Build = AppConfig.Build as IOSBuild;
@@ -521,107 +521,138 @@ namespace Gauntlet
 				throw new AutomationException("Invalid build for IOS!");
 			}
 
-			bool CacheResigned = false;
-			bool UseLocalExecutable = Globals.Params.ParseParam("dev");
-
-			lock(IPALock)
+			if (!AppConfig.SkipInstall)
 			{
-				Log.Info("Installing using IPA {0}", Build.SourceIPAPath);
+				bool CacheResigned = false;
+				bool UseLocalExecutable = Globals.Params.ParseParam("dev");
+
+				lock (AppLock)
+				{
+					Log.Info("Installing using build source {0}", Build.SourcePath);
 
-				// device artifact path
-				DeviceArtifactPath = string.Format("/Documents/{0}/Saved", AppConfig.ProjectName);
+					// device artifact path
+					DeviceArtifactPath = string.Format("/Documents/{0}/Saved", AppConfig.ProjectName);
 
-				CacheResigned = File.Exists(CacheResignedFilename);
+					CacheResigned = File.Exists(CacheResignedFilename);
 
-				if (CacheResigned && !UseLocalExecutable)
-				{
-					if (File.Exists(IPAHashFilename))
+					if (CacheResigned && !UseLocalExecutable)
 					{
-						Log.Verbose("App was resigned, invalidating app cache");
-						File.Delete(IPAHashFilename);
+						if (File.Exists(AppHashFilename))
+						{
+							Log.Verbose("App was resigned, invalidating app cache");
+							File.Delete(AppHashFilename);
+						}
 					}
-				}
 
-				PrepareIPA(Build);
+					if (Build.IsIPAFile)
+					{
+						PrepareIPA(Build);
+					}
+					else
+					{
+						LocalAppBundle = Build.SourcePath;
+					}
 
-				// local executable support
-				if (UseLocalExecutable)
-				{
-					ResignApplication(AppConfig);
+					// local executable support
+					if (UseLocalExecutable)
+					{
+						ResignApplication(AppConfig);
+					}
 				}
-			}
-
-			if (CacheResigned || UseLocalExecutable || !CheckDeployedIPA(Build))
-			{
-				// uninstall will clean all device artifacts
-				ExecuteIOSDeployCommand(String.Format("--uninstall -b \"{0}\"", LocalAppBundle), 20 * 60);
-			}
-			else
-			{
-				// remove device artifacts
-				CleanDeviceArtifacts(Build);
-			}
 
-			// parallel iOS tests use same app install folder, so lock it as setup is quick
-			lock (Globals.MainLock)
-			{
-				// local app install with additional files, this directory will be mirrored to device in a single operation
-				string AppInstallPath;
-
-				AppInstallPath = Path.Combine(Globals.TempDir, "iOSAppInstall");
-
-				if (Directory.Exists(AppInstallPath))
+				if (CacheResigned || UseLocalExecutable || !CheckDeployedApp(Build))
 				{
-					Directory.Delete(AppInstallPath, true);
+					// uninstall will clean all device artifacts
+					ExecuteIOSDeployCommand(String.Format("--uninstall -b \"{0}\"", LocalAppBundle), 20 * 60);
 				}
-
-				Directory.CreateDirectory(AppInstallPath);
-
-				if (LocalDirectoryMappings.Count == 0)
+				else
 				{
-					PopulateDirectoryMappings(AppInstallPath);
+					// remove device artifacts
+					CleanDeviceArtifacts(Build);
 				}
 
-				//@todo: Combine Build and AppConfig files, this should be done in higher level code, not per device implementation
-
-				if (AppConfig.FilesToCopy != null)
+				// parallel iOS tests use same app install folder, so lock it as setup is quick
+				lock (Globals.MainLock)
 				{
-					foreach (UnrealFileToCopy FileToCopy in AppConfig.FilesToCopy)
+					// local app install with additional files, this directory will be mirrored to device in a single operation
+					string AppInstallPath;
+
+					AppInstallPath = Path.Combine(Globals.TempDir, "iOSAppInstall");
+
+					if (Directory.Exists(AppInstallPath))
+					{
+						Directory.Delete(AppInstallPath, true);
+					}
+
+					Directory.CreateDirectory(AppInstallPath);
+
+					if (LocalDirectoryMappings.Count == 0)
 					{
-						string PathToCopyTo = Path.Combine(LocalDirectoryMappings[FileToCopy.TargetBaseDirectory], FileToCopy.TargetRelativeLocation);
+						PopulateDirectoryMappings(AppInstallPath);
+					}
 
-						if (File.Exists(FileToCopy.SourceFileLocation))
+					//@todo: Combine Build and AppConfig files, this should be done in higher level code, not per device implementation
+
+					if (AppConfig.FilesToCopy != null)
+					{
+						foreach (UnrealFileToCopy FileToCopy in AppConfig.FilesToCopy)
 						{
-							FileInfo SrcInfo = new FileInfo(FileToCopy.SourceFileLocation);
-							SrcInfo.IsReadOnly = false;
-							string DirectoryToCopyTo = Path.GetDirectoryName(PathToCopyTo);
-							if (!Directory.Exists(DirectoryToCopyTo))
+							string PathToCopyTo = Path.Combine(LocalDirectoryMappings[FileToCopy.TargetBaseDirectory], FileToCopy.TargetRelativeLocation);
+
+							if (File.Exists(FileToCopy.SourceFileLocation))
 							{
-								Directory.CreateDirectory(DirectoryToCopyTo);
+								FileInfo SrcInfo = new FileInfo(FileToCopy.SourceFileLocation);
+								SrcInfo.IsReadOnly = false;
+								string DirectoryToCopyTo = Path.GetDirectoryName(PathToCopyTo);
+								if (!Directory.Exists(DirectoryToCopyTo))
+								{
+									Directory.CreateDirectory(DirectoryToCopyTo);
+								}
+								if (File.Exists(PathToCopyTo))
+								{
+									FileInfo ExistingFile = new FileInfo(PathToCopyTo);
+									ExistingFile.IsReadOnly = false;
+								}
+
+								SrcInfo.CopyTo(PathToCopyTo, true);
+								Log.Verbose("Copying app install: {0} to {1}", FileToCopy, DirectoryToCopyTo);
 							}
-							if (File.Exists(PathToCopyTo))
+							else
 							{
-								FileInfo ExistingFile = new FileInfo(PathToCopyTo);
-								ExistingFile.IsReadOnly = false;
+								Log.Warning(KnownLogEvents.Gauntlet_DeviceEvent, "File to copy {File} not found", FileToCopy);
 							}
-
-							SrcInfo.CopyTo(PathToCopyTo, true);
-							Log.Verbose("Copying app install: {0} to {1}", FileToCopy, DirectoryToCopyTo);
-						}
-						else
-						{
-							Log.Warning(KnownLogEvents.Gauntlet_DeviceEvent, "File to copy {File} not found", FileToCopy);
 						}
 					}
-				}
 
-				// copy mapped files in a single pass
-				string CopyCommand = String.Format("--bundle_id {0} --upload={1} --to {2}", Build.PackageName, AppInstallPath, DeviceArtifactPath);
-				ExecuteIOSDeployCommand(CopyCommand, 120);
+					// copy mapped files in a single pass
+					string CopyCommand = String.Format("--bundle_id {0} --upload={1} --to {2}", Build.PackageName, AppInstallPath, DeviceArtifactPath);
+					ExecuteIOSDeployCommand(CopyCommand, 120);
 
-				// store the IPA hash to avoid redundant deployments
-				CopyCommand = String.Format("--bundle_id {0} --upload={1} --to {2}", Build.PackageName, IPAHashFilename, "/Documents/IPAHash.txt");
-				ExecuteIOSDeployCommand(CopyCommand, 120);
+					// store the App hash to avoid redundant deployments
+					if (File.Exists(AppHashFilename))
+					{
+						CopyCommand = String.Format("--bundle_id {0} --upload={1} --to {2}", Build.PackageName, AppHashFilename, "/Documents/AppHash.txt");
+						ExecuteIOSDeployCommand(CopyCommand, 120);
+					}
+				}
+			}
+			else
+			{
+				// Skipping installation
+				Log.Info("Skipping installation as per config.");
+				if (Build.IsIPAFile)
+				{
+					string PayloadDir = Path.Combine(GauntletAppCache, "Payload");
+					LocalAppBundle = Directory.GetDirectories(PayloadDir).Where(D => Path.GetExtension(D) == ".app").FirstOrDefault();
+					if (string.IsNullOrEmpty(LocalAppBundle))
+					{
+						throw new AutomationException($"Could not find local IPA payload at {PayloadDir}");
+					}
+				}
+				else
+				{
+					LocalAppBundle = Build.SourcePath;
+				}
 			}
 
 			IOSAppInstall IOSApp = new IOSAppInstall(AppConfig.Name, this, Build.PackageName, AppConfig.CommandLine);
@@ -843,8 +874,8 @@ namespace Gauntlet
 		// Note: ios-deploy works with app bundles, which requires the IPA be unzipped for deployment (this will allow us to resign in the future as well)
 		string LocalAppBundle = null;
 
-		// the current IPA MD5 hash, which is tracked to avoid unneccessary deployments and unzip operations
-		string IPAHashFilename { get { return Path.Combine(GauntletAppCache, "IPAHash.txt"); } }
+		// the current App MD5 hash, which is tracked to avoid unneccessary deployments and unzip operations
+		string AppHashFilename { get { return Path.Combine(GauntletAppCache, "AppHash.txt"); } }
 
 		// file whose presence signals that cache was resigned
 		string CacheResignedFilename { get { return Path.Combine(GauntletAppCache, "Resigned.txt"); } }
@@ -855,7 +886,7 @@ namespace Gauntlet
 		/// </summary>
 		private bool PrepareIPA(IOSBuild Build)
 		{	
-			Log.Info("Preparing IPA {0}", Build.SourceIPAPath);
+			Log.Info("Preparing IPA {0}", Build.SourcePath);
 
 			try
 			{	
@@ -864,7 +895,7 @@ namespace Gauntlet
 				string StoredHash = null;
 				using (var MD5Hash = MD5.Create())
 				{
-					using (var Stream = File.OpenRead(Build.SourceIPAPath))
+					using (var Stream = File.OpenRead(Build.SourcePath))
 					{
 						Hash = BitConverter.ToString(MD5Hash.ComputeHash(Stream)).Replace("-", "").ToLowerInvariant();
 					}
@@ -872,12 +903,12 @@ namespace Gauntlet
 				string PayloadDir = Path.Combine(GauntletAppCache, "Payload");
 				string SymbolsDir = Path.Combine(GauntletAppCache, "Symbols");
 				
-				if (File.Exists(IPAHashFilename) && Directory.Exists(PayloadDir))
+				if (File.Exists(AppHashFilename) && Directory.Exists(PayloadDir))
 				{
-					StoredHash = File.ReadAllText(IPAHashFilename).Trim();
+					StoredHash = File.ReadAllText(AppHashFilename).Trim();
 					if (Hash != StoredHash)
 					{
-						Log.Verbose("IPA hash out of date, clearing cache");
+						Log.Verbose("App hash out of date, clearing cache");
 						StoredHash = null;
 					}
 				}
@@ -899,16 +930,16 @@ namespace Gauntlet
 						File.Delete(CacheResignedFilename);
 					}
 
-					Log.Verbose("Unzipping IPA {0} to cache at: {1}", Build.SourceIPAPath, GauntletAppCache);
+					Log.Verbose("Unzipping IPA {0} to cache at: {1}", Build.SourcePath, GauntletAppCache);
 
 					string Output;
-					if (!IOSBuild.ExecuteIPADittoCommand(String.Format("-x -k {0} {1}", Build.SourceIPAPath, GauntletAppCache), out Output, PayloadDir))
+					if (!IOSBuild.ExecuteIPADittoCommand(String.Format("-x -k {0} {1}", Build.SourcePath, GauntletAppCache), out Output, PayloadDir))
 					{
-						throw new Exception(String.Format("Unable to extract IPA {0}", Build.SourceIPAPath));
+						throw new Exception(String.Format("Unable to extract IPA {0}", Build.SourcePath));
 					}
 
 					// Cache symbols for symbolicated callstacks
-					string SymbolsZipFile = string.Format("{0}/../../Symbols/{1}.dSYM.zip", Path.GetDirectoryName(Build.SourceIPAPath), Path.GetFileNameWithoutExtension(Build.SourceIPAPath));
+					string SymbolsZipFile = string.Format("{0}/../../Symbols/{1}.dSYM.zip", Path.GetDirectoryName(Build.SourcePath), Path.GetFileNameWithoutExtension(Build.SourcePath));
 
 					Log.Verbose("Checking Symbols at {0}", SymbolsZipFile);
 
@@ -923,13 +954,13 @@ namespace Gauntlet
 					}					
 
 					// store hash
-					File.WriteAllText(IPAHashFilename, Hash);
+					File.WriteAllText(AppHashFilename, Hash);
 
-					Log.Verbose("IPA cached");
+					Log.Verbose("App cached");
 				}
 				else
 				{
-					Log.Verbose("Using cached IPA");
+					Log.Verbose("Using cached App");
 				}
 
 				LocalAppBundle = Directory.GetDirectories(PayloadDir).Where(D => Path.GetExtension(D) == ".app").FirstOrDefault();
@@ -942,7 +973,7 @@ namespace Gauntlet
 			}
 			catch (Exception Ex)
 			{
-				throw new AutomationException("Unable to prepare {0} : {1}", Build.SourceIPAPath, Ex.Message);
+				throw new AutomationException("Unable to prepare {0} : {1}", Build.SourcePath, Ex.Message);
 			}		
 
 			return true;
diff --git a/Engine/Source/Programs/AutomationTool/Gauntlet/Unreal/BuildSource/Gauntlet.UnrealBuildSource.cs b/Engine/Source/Programs/AutomationTool/Gauntlet/Unreal/BuildSource/Gauntlet.UnrealBuildSource.cs
--- a/Engine/Source/Programs/AutomationTool/Gauntlet/Unreal/BuildSource/Gauntlet.UnrealBuildSource.cs
+++ b/Engine/Source/Programs/AutomationTool/Gauntlet/Unreal/BuildSource/Gauntlet.UnrealBuildSource.cs
@@ -514,15 +514,24 @@ namespace Gauntlet
 			// new system of retrieving and encapsulating the info needed to install/launch. Android & Mac
 			Config.Build = GetMatchingBuilds(Role.RoleType, Role.Platform, Role.Configuration, Role.RequiredBuildFlags, Role.RequiredFlavor).OrderBy(B => B.PreferenceOrder).FirstOrDefault();
 
-			if (Config.Build == null && Role.IsNullRole() == false)
+			if (Config.Build == null)
 			{
-				var SupportedBuilds = String.Join("\n", DiscoveredBuilds.Select(B => B.ToString()));
+				if (Role.IsNullRole())
+				{
+					Log.Warning("No supported build found, however role is Null and not configure to run anything.");
+				}
+				else
+				{
+					var SupportedBuilds = String.Join("\n", DiscoveredBuilds.Select(B => B.ToString()));
 
-				Log.Info("Available builds:\n{0}", SupportedBuilds);
-				throw new AutomationException("No build found that can support a role of {0}.", Role);
+					Log.Info("Available builds:\n{0}", SupportedBuilds);
+					throw new AutomationException("No build found that can support a role of {0}.", Role);
+				}
+			}
+			else
+			{
+				Log.Info("Selected build {Build} for test run.", Config.Build.ToString());
 			}
-
-			Log.Info("Selected build {Build} for test run.", Config.Build.ToString());
 
 			if (Role.Options != null)
 			{
