Add fallback option for when the metadata directory does not exist when performing a DLC cook.
Potential location when an override cook dir is provided.

diff --git a/Engine/Source/Programs/AutomationTool/Scripts/CopyBuildToStagingDirectory.Automation.cs b/Engine/Source/Programs/AutomationTool/Scripts/CopyBuildToStagingDirectory.Automation.cs
--- a/Engine/Source/Programs/AutomationTool/Scripts/CopyBuildToStagingDirectory.Automation.cs
+++ b/Engine/Source/Programs/AutomationTool/Scripts/CopyBuildToStagingDirectory.Automation.cs
@@ -984,6 +984,12 @@ namespace AutomationScripts
 					}
 				}
 				SC.MetadataDir = DirectoryReference.Combine(SC.PlatformCookDir, DLCCookedSubDir, "Metadata");
+				if (!DirectoryReference.Exists(SC.MetadataDir))
+				{
+					DirectoryReference DLCMetaDataDir = DirectoryReference.Combine(new DirectoryReference(Params.CookOutputDir), DLCCookedSubDir, "Metadata");
+					Logger.LogInformation($"{SC.MetadataDir} Does not exist. Using alternative dir {DLCMetaDataDir}");
+					SC.MetadataDir = DLCMetaDataDir;
+				}
 			}
 			else
 			{
