Fill in the ability to export snapshots to a zen host.

diff --git a/Engine/Source/Programs/AutomationTool/BuildGraph/Tasks/ZenExportSnapshotTask.cs b/Engine/Source/Programs/AutomationTool/BuildGraph/Tasks/ZenExportSnapshotTask.cs
--- a/Engine/Source/Programs/AutomationTool/BuildGraph/Tasks/ZenExportSnapshotTask.cs
+++ b/Engine/Source/Programs/AutomationTool/BuildGraph/Tasks/ZenExportSnapshotTask.cs
@@ -118,6 +118,12 @@ namespace AutomationTool.Tasks
 		[TaskParameter(Optional = true)]
 		public string DestinationCloudBucket;
 
+		/// <summary>
+		/// The host name to use when exporting to a zen destination
+		/// </summary>
+		[TaskParameter(Optional = true)]
+		public string DestinationZenHost;
+
 		/// <summary>
 		/// The directory to use when exporting to a file destination
 		/// </summary>
@@ -333,7 +339,7 @@ namespace AutomationTool.Tasks
 					string HttpVersion = Parameters.SnapshotDescriptorCloudHttpVersion;
 					if (string.IsNullOrEmpty(HttpVersion))
 					{
-						HostName = Parameters.DestinationCloudHttpVersion;
+						HttpVersion = Parameters.DestinationCloudHttpVersion;
 					}
 
 					IoHash DestinationKeyHash = IoHash.Compute(Encoding.UTF8.GetBytes(Name));
@@ -349,6 +355,16 @@ namespace AutomationTool.Tasks
 					Writer.WriteValue("bucket", BucketName);
 					Writer.WriteValue("key", DestinationKeyHash.ToString().ToLowerInvariant());
 					break;
+				case SnapshotStorageType.Zen:
+					string ProjectName = Parameters.Project.GetFileNameWithoutAnyExtensions().ToLowerInvariant() + ".oplog";
+
+					Writer.WriteValue("name", Name);
+					Writer.WriteValue("type", "zen");
+					Writer.WriteValue("targetplatform", ExportSource.TargetPlatform);
+					Writer.WriteValue("host", Parameters.DestinationZenHost);
+					Writer.WriteValue("projectid", ProjectName);
+					Writer.WriteValue("oplogid", Name);
+					break;
 				case SnapshotStorageType.File:
 					Writer.WriteValue("name", Name);
 					Writer.WriteValue("type", "file");
@@ -520,6 +536,43 @@ namespace AutomationTool.Tasks
 						ExportIndex = ExportIndex + 1;
 					}
 
+					break;
+				case SnapshotStorageType.Zen:
+					if (string.IsNullOrEmpty(Parameters.DestinationZenHost))
+					{
+						throw new AutomationException("Missing destination zen host");
+					}
+					if (string.IsNullOrEmpty(Parameters.DestinationIdentifier))
+					{
+						throw new AutomationException("Missing destination identifier when exporting to zen");
+					}
+
+					string ProjectName = ProjectFile.GetFileNameWithoutAnyExtensions().ToLowerInvariant() + ".oplog";
+
+					OplogExportCommandline.AppendFormat(" --zen {0}", Parameters.DestinationZenHost);
+
+					ExportIndex = 0;
+					foreach (ExportSourceData ExportSource in ExportSources)
+					{
+						string HostUrlArg = string.Format("--hosturl http://{0}:{1}", ExportSource.IsLocalHost ? "localhost" : ExportSource.HostName, ExportSource.HostPort);
+
+						StringBuilder ExportSingleSourceCommandline = new StringBuilder(OplogExportCommandline.Length);
+						ExportSingleSourceCommandline.Append(OplogExportCommandline);
+
+						StringBuilder DestinationKeyBuilder = new StringBuilder();
+						DestinationKeyBuilder.AppendFormat("{0}.{1}", Parameters.DestinationIdentifier, ExportSource.OplogId);
+						ExportNames[ExportIndex] = DestinationKeyBuilder.ToString().ToLowerInvariant();
+
+						ProcessResult.SpewFilterCallbackType SilentOutputFilter = new ProcessResult.SpewFilterCallbackType(Line =>
+							{
+								return null;
+							});
+						ExportSingleSourceCommandline.AppendFormat(" {0} --target-project {1} --target-oplog {2} {3} {4}", HostUrlArg, ProjectName, ExportNames[ExportIndex], ExportSource.ProjectId, ExportSource.OplogId);
+						CommandUtils.RunAndLog(CommandUtils.CmdEnv, ZenExe.FullName, ExportSingleSourceCommandline.ToString(), MaxSuccessCode: int.MaxValue, Options: CommandUtils.ERunOptions.Default, SpewFilterCallback: SilentOutputFilter);
+
+						ExportIndex = ExportIndex + 1;
+					}
+
 					break;
 				case SnapshotStorageType.File:
 					string DefaultProjectId = ProjectUtils.GetProjectPathId(ProjectFile);
