Fixes for the ability to export snapshots to a zen host.

diff --git a/Engine/Source/Programs/AutomationTool/BuildGraph/Tasks/ZenExportSnapshotTask.cs b/Engine/Source/Programs/AutomationTool/BuildGraph/Tasks/ZenExportSnapshotTask.cs
--- a/Engine/Source/Programs/AutomationTool/BuildGraph/Tasks/ZenExportSnapshotTask.cs
+++ b/Engine/Source/Programs/AutomationTool/BuildGraph/Tasks/ZenExportSnapshotTask.cs
@@ -317,6 +317,12 @@ namespace AutomationTool.Tasks
 			byte[] data = FileReference.ReadAllBytes(file);
 			return JsonSerializer.Deserialize<T>(data, GetDefaultJsonSerializerOptions())!;
 		}
+
+		static string SanitizeOplogName(string Name)
+		{
+			return Name.Replace('/','_').Replace(' ', '_').Replace('+','_').Replace('-','_');
+		}
+
 		private void WriteExportSource(JsonWriter Writer, SnapshotStorageType DestinationStorageType, ExportSourceData ExportSource, string Name)
 		{
 			Writer.WriteObjectStart();
@@ -363,7 +369,7 @@ namespace AutomationTool.Tasks
 					Writer.WriteValue("targetplatform", ExportSource.TargetPlatform);
 					Writer.WriteValue("host", Parameters.DestinationZenHost);
 					Writer.WriteValue("projectid", ProjectName);
-					Writer.WriteValue("oplogid", Name);
+					Writer.WriteValue("oplogid", SanitizeOplogName(Name));
 					break;
 				case SnapshotStorageType.File:
 					Writer.WriteValue("name", Name);
@@ -463,6 +469,11 @@ namespace AutomationTool.Tasks
 				OplogExportCommandline.Append(" --force");
 			}
 
+			ProcessResult.SpewFilterCallbackType SilentOutputFilter = new ProcessResult.SpewFilterCallbackType(Line =>
+				{
+					return null;
+				});
+
 			switch (DestinationStorageType)
 			{
 				case SnapshotStorageType.Cloud:
@@ -526,10 +537,6 @@ namespace AutomationTool.Tasks
 						ExportNames[ExportIndex] = DestinationKeyBuilder.ToString().ToLowerInvariant();
 						IoHash DestinationKeyHash = IoHash.Compute(Encoding.UTF8.GetBytes(ExportNames[ExportIndex]));
 
-						ProcessResult.SpewFilterCallbackType SilentOutputFilter = new ProcessResult.SpewFilterCallbackType(Line =>
-							{
-								return null;
-							});
 						ExportSingleSourceCommandline.AppendFormat(" {0} --embedloosefiles --key {1} {2} {3} {4}", HostUrlArg, DestinationKeyHash.ToString().ToLowerInvariant(), BaseKeyArg, ExportSource.ProjectId, ExportSource.OplogId);
 						CommandUtils.RunAndLog(CommandUtils.CmdEnv, ZenExe.FullName, ExportSingleSourceCommandline.ToString(), MaxSuccessCode: int.MaxValue, Options: CommandUtils.ERunOptions.Default, SpewFilterCallback: SilentOutputFilter);
 
@@ -549,6 +556,10 @@ namespace AutomationTool.Tasks
 
 					string ProjectName = ProjectFile.GetFileNameWithoutAnyExtensions().ToLowerInvariant() + ".oplog";
 
+					StringBuilder CreateProjectCommandline = new StringBuilder();
+					CreateProjectCommandline.AppendFormat("project-create --hosturl {0} {1}", Parameters.DestinationZenHost, ProjectName);
+					CommandUtils.RunAndLog(CommandUtils.CmdEnv, ZenExe.FullName, CreateProjectCommandline.ToString(), MaxSuccessCode: int.MaxValue, Options: CommandUtils.ERunOptions.Default, SpewFilterCallback: SilentOutputFilter);
+
 					OplogExportCommandline.AppendFormat(" --zen {0}", Parameters.DestinationZenHost);
 
 					ExportIndex = 0;
@@ -562,12 +573,9 @@ namespace AutomationTool.Tasks
 						StringBuilder DestinationKeyBuilder = new StringBuilder();
 						DestinationKeyBuilder.AppendFormat("{0}.{1}", Parameters.DestinationIdentifier, ExportSource.OplogId);
 						ExportNames[ExportIndex] = DestinationKeyBuilder.ToString().ToLowerInvariant();
+						string DestinationOplog = SanitizeOplogName(ExportNames[ExportIndex]);
 
-						ProcessResult.SpewFilterCallbackType SilentOutputFilter = new ProcessResult.SpewFilterCallbackType(Line =>
-							{
-								return null;
-							});
-						ExportSingleSourceCommandline.AppendFormat(" {0} --target-project {1} --target-oplog {2} {3} {4}", HostUrlArg, ProjectName, ExportNames[ExportIndex], ExportSource.ProjectId, ExportSource.OplogId);
+						ExportSingleSourceCommandline.AppendFormat(" {0} --target-project {1} --target-oplog {2} {3} {4}", HostUrlArg, ProjectName, DestinationOplog, ExportSource.ProjectId, ExportSource.OplogId);
 						CommandUtils.RunAndLog(CommandUtils.CmdEnv, ZenExe.FullName, ExportSingleSourceCommandline.ToString(), MaxSuccessCode: int.MaxValue, Options: CommandUtils.ERunOptions.Default, SpewFilterCallback: SilentOutputFilter);
 
 						ExportIndex = ExportIndex + 1;
