Configuring the log heartbeat channel properly in UEPref.EditorGauntletTestController

#rnx

diff --git a/Engine/Source/Programs/AutomationTool/Gauntlet/Unreal/Editor/EditorTests.cs b/Engine/Source/Programs/AutomationTool/Gauntlet/Unreal/Editor/EditorTests.cs
--- a/Engine/Source/Programs/AutomationTool/Gauntlet/Unreal/Editor/EditorTests.cs
+++ b/Engine/Source/Programs/AutomationTool/Gauntlet/Unreal/Editor/EditorTests.cs
@@ -285,9 +285,6 @@ namespace UEPerf
 
 			// Enforcing PIE to allow the Gauntlet controller to start
 			AppConfig.CommandLineParams.Add("PIE");
-
-			// Log Channels to listen for activity
-			LogCategoriesForEvents = UnrealLog.EditorBusyChannels.ToList();
 		}
 	}
 
@@ -299,6 +296,15 @@ namespace UEPerf
 		{
 		}
 
+		/// <summary>
+		/// Log Channels to listen for activity
+		/// </summary>
+		/// <returns></returns>
+		public override IEnumerable<string> GetHeartbeatLogCategories()
+		{
+			return UnrealLog.EditorBusyChannels;
+		}
+
 		public override EditorGauntletTestControllerConfig GetConfiguration()
 		{
 			if (CachedConfig != null)
