Ensure the "embedloosefiles" argument is applied regardless of export destination type when exporting zen snapshots.

diff --git a/Engine/Source/Programs/AutomationTool/BuildGraph/Tasks/ZenExportSnapshotTask.cs b/Engine/Source/Programs/AutomationTool/BuildGraph/Tasks/ZenExportSnapshotTask.cs
--- a/Engine/Source/Programs/AutomationTool/BuildGraph/Tasks/ZenExportSnapshotTask.cs
+++ b/Engine/Source/Programs/AutomationTool/BuildGraph/Tasks/ZenExportSnapshotTask.cs
@@ -463,7 +463,7 @@ namespace AutomationTool.Tasks
 
 			// Format the command line
 			StringBuilder OplogExportCommandline = new StringBuilder();
-			OplogExportCommandline.Append("oplog-export");
+			OplogExportCommandline.Append("oplog-export --embedloosefiles");
 			if (Parameters.Force)
 			{
 				OplogExportCommandline.Append(" --force");
@@ -537,7 +537,7 @@ namespace AutomationTool.Tasks
 						ExportNames[ExportIndex] = DestinationKeyBuilder.ToString().ToLowerInvariant();
 						IoHash DestinationKeyHash = IoHash.Compute(Encoding.UTF8.GetBytes(ExportNames[ExportIndex]));
 
-						ExportSingleSourceCommandline.AppendFormat(" {0} --embedloosefiles --key {1} {2} {3} {4}", HostUrlArg, DestinationKeyHash.ToString().ToLowerInvariant(), BaseKeyArg, ExportSource.ProjectId, ExportSource.OplogId);
+						ExportSingleSourceCommandline.AppendFormat(" {0} --key {1} {2} {3} {4}", HostUrlArg, DestinationKeyHash.ToString().ToLowerInvariant(), BaseKeyArg, ExportSource.ProjectId, ExportSource.OplogId);
 						CommandUtils.RunAndLog(CommandUtils.CmdEnv, ZenExe.FullName, ExportSingleSourceCommandline.ToString(), MaxSuccessCode: int.MaxValue, Options: CommandUtils.ERunOptions.Default, SpewFilterCallback: SilentOutputFilter);
 
 						ExportIndex = ExportIndex + 1;
