Fix issue when trying alternative staging dir.
Dir parameter must exist.

diff --git a/Engine/Source/Programs/AutomationTool/Scripts/CopyBuildToStagingDirectory.Automation.cs b/Engine/Source/Programs/AutomationTool/Scripts/CopyBuildToStagingDirectory.Automation.cs
--- a/Engine/Source/Programs/AutomationTool/Scripts/CopyBuildToStagingDirectory.Automation.cs
+++ b/Engine/Source/Programs/AutomationTool/Scripts/CopyBuildToStagingDirectory.Automation.cs
@@ -984,7 +984,7 @@ namespace AutomationScripts
 					}
 				}
 				SC.MetadataDir = DirectoryReference.Combine(SC.PlatformCookDir, DLCCookedSubDir, "Metadata");
-				if (!DirectoryReference.Exists(SC.MetadataDir))
+				if (!DirectoryReference.Exists(SC.MetadataDir) && !String.IsNullOrEmpty(Params.CookOutputDir))
 				{
 					DirectoryReference DLCMetaDataDir = DirectoryReference.Combine(new DirectoryReference(Params.CookOutputDir), DLCCookedSubDir, "Metadata");
 					Logger.LogInformation($"{SC.MetadataDir} Does not exist. Using alternative dir {DLCMetaDataDir}");
