- Fixed an issue with writing a file without creating the directory first

diff --git a/Engine/Source/Programs/AutomationTool/Apple/ApplePlatform.Automation.cs b/Engine/Source/Programs/AutomationTool/Apple/ApplePlatform.Automation.cs
--- a/Engine/Source/Programs/AutomationTool/Apple/ApplePlatform.Automation.cs
+++ b/Engine/Source/Programs/AutomationTool/Apple/ApplePlatform.Automation.cs
@@ -108,6 +108,7 @@ public abstract class ApplePlatform : Platform
 		// so plop a file down in it's place
 		if (SC.StageTargetPlatform.PlatformType == UnrealTargetPlatform.Mac)
 		{
+		    Directory.CreateDirectory(BinaryPath.Directory.FullName);
 			File.WriteAllText(Path.Combine(BinaryPath.Directory.FullName, ".binariesdir"), "");
 		}
 
