LLT Android support

diff --git a/Engine/Build/LowLevelTests.xml b/Engine/Build/LowLevelTests.xml
--- a/Engine/Build/LowLevelTests.xml
+++ b/Engine/Build/LowLevelTests.xml
@@ -13,6 +13,12 @@
 
 	<Property Name="AlternateDevicePoolPlatformsProperty" Value=""/>
 
+	<Property Name="AdbKeysPath" Value=""/>
+
+	<Macro Name="SelectAdbKeysPath" Arguments="BuildPlatform">
+		<!-- Extend this macro and set AdbKeysPath -->
+	</Macro>	
+
 	<!-- Platform specific macros -->
 	<Macro Name="Catch2BuildAppendExtraCMakeArgsPlatform" Arguments="Platform"/>
 	<Macro Name="Catch2BuildAppendExtraCMakeArgs" Arguments="ExtraArgValue">
@@ -131,7 +137,7 @@
 			<Compile Target="$($(TestName)Target)" Configuration="$(Configuration)" Platform="$(Platform)" AllowXGE="false" Arguments="-Mode=Test $($(Platform)$(TestName)ExtraCompilationArgs)"/>
 		</Do>
 	</Macro>
-	<Macro Name="DeployAndTest" Arguments="TestName;Platform">
+	<Macro Name="DeployAndTest" Arguments="TestName;Platform;BuildPlatform">
 		<Property Name="PlatformStageDirName" Value="$(Platform)"/>
 		<Property Name="PlatformStageDirName" Value="Windows" If="'$(Platform)' == 'Win64'"/>
 		<Property Name="DevicePool" Value="$(GauntletDevicePool)" />
@@ -142,11 +148,16 @@
 		<Property Name="TestTags" Value=""/>
 		<Property Name="TestTags" Value="$($(Platform)$(TestName)Tags)" If="'$($(Platform)$(TestName)Tags)' != ''"/>
 		<Property Name="TestTags" Value="$(OverrideTags)" If="'$(OverrideTags)' != ''"/>
+		<Do If="'$(Platform)' == 'Android'">
+			<Property Name="TestArgsExt" Value="$(TestArgsExt) -cleandevice"/>
+		</Do>
 		<Property Name="TestArgsExt" Value="$(TestArgsExt) -tags=&quot;$(TestTags)&quot;"/>
 		<Property Name="TestArgsExt" Value="$(TestArgsExt) --buildmachine" If="$(IsBuildMachine)"/>
 		<Property Name="TestArgsExt" Value="$(TestArgsExt) -timeout=$(PerTestTimeout)"/>
 		<Property Name="TestArgsExt" Value="$(TestArgsExt) $($(TestName)GauntletArgs)"/>
-		<Property Name="TestArgsExt" Value="$(TestArgsExt) -extra-args=&quot;$($(TestName)ExtraArgs)&quot;"/>
+		<Do If="'$($(TestName)ExtraArgs)' != ''">
+			<Property Name="TestArgsExt" Value="$(TestArgsExt) -extra-args=&quot;$($(TestName)ExtraArgs)&quot;"/>
+		</Do>
 		<Property Name="TestArgsExt" Value="$(TestArgsExt) $($(Platform)ExtraArgs)"/>
 		<Property Name="TestArgsExt" Value="$(TestArgsExt) -VerifyLogin" If="$(EnableDeviceLoginVerification)"/>
 		<Do If="'$(Device)' != ''">
@@ -179,6 +190,11 @@
 			<Expand Name="$(Platform)DockerSetup" AppNameLowercase="$(TestTargetLowercase)" AppDir="$(BinariesRelativeForwardSlash)/$(Platform)/$($(TestName)Target)$(TestModeSuffix)"/>
 			<Property Name="TestArgsExt" Value="$(TestArgsExt) -containerized"/>
 		</Do>
+
+		<Do If="'$(Platform)' == 'Android'">
+			<Expand Name="SelectAdbKeysPath" BuildPlatform="$(BuildPlatform)"/>
+			<Property Name="TestArgsExt" Value="$(TestArgsExt) -adbkeys=$(AdbKeysPath) -captureoutput"/>
+		</Do>
 		
 		<Command Name="RunLowLevelTests"
 				Arguments="-test=LowLevelTests -testapp=$($(TestName)Target)$(TestModeSuffix) -platform=$(Platform) -configuration=$(Configuration) $(TestArgsExt) -VeryVerbose"/>
@@ -259,7 +275,7 @@
 						<Node Name="$($(TestName)Short) Tests $(Platform)" Requires="$(TestRequirements)">
 							<Expand Name="BuildAndStageTest" TestName="$(TestName)" Platform="$(Platform)"/>
 							<Do If="'$($(Platform)$(TestName)RunSupported)' == 'True'">
-								<Expand Name="DeployAndTest" TestName="$(TestName)" Platform="$(Platform)"/>
+								<Expand Name="DeployAndTest" TestName="$(TestName)" Platform="$(Platform)" BuildPlatform="$(BuildPlatform)"/>
 							</Do>
 						</Node>
 						<Label Category="Test" Name="$($(TestName)Short) $(Platform)" Requires="$($(TestName)Short) Tests $(Platform)"/>
diff --git a/Engine/Build/LowLevelTests/Foundation.xml b/Engine/Build/LowLevelTests/Foundation.xml
--- a/Engine/Build/LowLevelTests/Foundation.xml
+++ b/Engine/Build/LowLevelTests/Foundation.xml
@@ -18,7 +18,7 @@
   <Property Name="Win64FoundationTags" Value="" />
   <Property Name="MacFoundationTags" Value="" />
   <Property Name="IOSFoundationTags" Value="" />
-  <Property Name="AndroidFoundationTags" Value="" />
+  <Property Name="AndroidFoundationTags" Value="~[Perf]~[Slow]~[AndroidSkip]" />
   <Property Name="LinuxFoundationTags" Value="~[.]~[Slow]" />
   <Property Name="LinuxArm64FoundationTags" Value="" />
   <Property Name="TVOSFoundationTags" Value="" />
@@ -32,7 +32,7 @@
   <Property Name="Win64FoundationRunSupported" Value="True" />
   <Property Name="MacFoundationRunSupported" Value="True" />
   <Property Name="IOSFoundationRunSupported" Value="False" />
-  <Property Name="AndroidFoundationRunSupported" Value="False" />
+  <Property Name="AndroidFoundationRunSupported" Value="True" />
   <Property Name="LinuxFoundationRunSupported" Value="True" />
   <Property Name="LinuxArm64FoundationRunSupported" Value="True" />
   <Property Name="TVOSFoundationRunSupported" Value="False" />
diff --git a/Engine/Build/LowLevelTests/MathCore.xml b/Engine/Build/LowLevelTests/MathCore.xml
--- a/Engine/Build/LowLevelTests/MathCore.xml
+++ b/Engine/Build/LowLevelTests/MathCore.xml
@@ -27,7 +27,7 @@
   <Property Name="IOSMathCoreRunContainerized" Value="False" />
   <Property Name="AndroidMathCoreTags" Value="" />
   <Property Name="AndroidMathCoreExtraCompilationArgs" Value="-allmodules -architectures=arm64" />
-  <Property Name="AndroidMathCoreRunSupported" Value="False" />
+  <Property Name="AndroidMathCoreRunSupported" Value="True" />
   <Property Name="AndroidMathCoreRunContainerized" Value="False" />
   <Property Name="LinuxMathCoreTags" Value="" />
   <Property Name="LinuxMathCoreExtraCompilationArgs" Value="-allmodules" />
diff --git a/Engine/Source/Developer/LowLevelTestsRunner/Android/Java/src/com/epicgames/unreal/tests/TestActivity.java b/Engine/Source/Developer/LowLevelTestsRunner/Android/Java/src/com/epicgames/unreal/tests/TestActivity.java
new file mode 100644
--- /dev/null
+++ b/Engine/Source/Developer/LowLevelTestsRunner/Android/Java/src/com/epicgames/unreal/tests/TestActivity.java
@@ -0,0 +1,93 @@
+// Copyright Epic Games, Inc. All Rights Reserved.
+
+package com.epicgames.unreal.tests;
+
+import android.app.Activity;
+import android.os.Bundle;
+import android.util.Log;
+
+import java.io.File;
+import java.io.FileInputStream;
+import java.io.IOException;
+import java.util.Arrays;
+import java.util.ArrayList;
+
+import android.content.pm.PackageManager;
+
+public class TestActivity extends Activity {
+	private int testsReturnCode = -1;
+	
+	public TestActivity() {
+	}
+
+	public String readCmdLineFrom(String path) {
+		try {
+			File cmdLineFile = new File(path);
+			FileInputStream cmdLineInputStream = new FileInputStream(cmdLineFile);
+			byte[] stringBytes = new byte[(int) cmdLineFile.length()];
+			if (cmdLineInputStream.read(stringBytes) != stringBytes.length) {
+				return "";
+			}
+			cmdLineInputStream.close();
+			return new String(stringBytes);
+		} catch (IOException e) {
+			Log.e("UE", "Could not process command line, exception: " + e.getMessage());
+			return "";
+		}
+	}
+
+	@Override
+	public void onCreate(Bundle savedInstanceState) {
+		super.onCreate(savedInstanceState);
+
+		System.loadLibrary("Unreal");
+
+		String PackageName = null;
+		try
+		{
+			PackageName = String.valueOf(getPackageManager().getPackageInfo(getPackageName(), 0).packageName);
+		}
+		catch (PackageManager.NameNotFoundException e)
+		{
+			Log.d("UE", "Failed to read package name: " + e.getMessage());
+		}
+
+		if (PackageName != null)
+		{
+			String[] PackageNameParts = PackageName.split("\\.");
+			String PackageNameLastPart = PackageNameParts[PackageNameParts.length - 1];
+			String cmdLinePath = "/sdcard/UnrealGame/" + PackageNameLastPart + "/UECommandLine.txt";
+
+			String TestCommandLine = readCmdLineFrom(cmdLinePath);
+			String AbsolutePath = getFilesDir().getAbsolutePath() + "/";
+			String[] AllArgs = TestCommandLine.trim().split("\\s+");
+
+			Log.d("UE", "Command line is " + TestCommandLine);
+
+			testsReturnCode = runTests(AbsolutePath, AllArgs);
+		}
+
+		// Required for test status until we can generate reports on Android
+		Log.d("UE", "Tests finished with exit code " + testsReturnCode);
+
+		try {
+			// Waiting here for any test cleanup seems to be required to prevent potential hang after System.exit()...
+			Thread.sleep(3000);
+		} catch (InterruptedException ie)
+		{}
+
+		finish();
+
+		return;
+	}
+
+	@Override
+	public void onDestroy()
+	{
+		super.onDestroy();
+		System.exit(testsReturnCode);
+
+	}
+
+	public static native int runTests(String path, String[] args);
+}
diff --git a/Engine/Source/Developer/LowLevelTestsRunner/Android/Java/src/com/epicgames/unreal/tests/TestApplication.java b/Engine/Source/Developer/LowLevelTestsRunner/Android/Java/src/com/epicgames/unreal/tests/TestApplication.java
new file mode 100644
--- /dev/null
+++ b/Engine/Source/Developer/LowLevelTestsRunner/Android/Java/src/com/epicgames/unreal/tests/TestApplication.java
@@ -0,0 +1,67 @@
+// Copyright Epic Games, Inc. All Rights Reserved.
+
+package com.epicgames.unreal.tests;
+
+import android.app.Application;
+import android.content.Context;
+import android.content.res.Configuration;
+import androidx.lifecycle.Lifecycle;
+import androidx.lifecycle.LifecycleObserver;
+import androidx.lifecycle.OnLifecycleEvent;
+import androidx.lifecycle.ProcessLifecycleOwner;
+
+public class TestApplication extends Application implements LifecycleObserver {
+	private static boolean isForeground = false;
+
+	private static Context context;
+
+	@Override
+	public void onCreate() {
+		super.onCreate();
+		TestApplication.context = getApplicationContext();
+		ProcessLifecycleOwner.get().getLifecycle().addObserver(this);
+	}
+
+	public static Context getAppContext() {
+		return TestApplication.context;
+	}
+
+	@Override
+	public void attachBaseContext(Context base) {
+		super.attachBaseContext(base);
+	}
+
+	@Override
+	public void onLowMemory() {
+		super.onLowMemory();
+	}
+
+	@Override
+	public void onTrimMemory(int level) {
+		super.onTrimMemory(level);
+	}
+
+	@Override
+	public void onConfigurationChanged(Configuration newConfig) {
+		super.onConfigurationChanged(newConfig);
+	}
+
+	@OnLifecycleEvent(Lifecycle.Event.ON_START)
+	void onEnterForeground() {
+		isForeground = true;
+	}
+
+	@OnLifecycleEvent(Lifecycle.Event.ON_STOP)
+	void onEnterBackground() {
+		isForeground = false;
+	}
+
+	@SuppressWarnings("unused")
+	public static boolean isAppInForeground() {
+		return isForeground;
+	}
+
+	public static boolean isAppInBackground() {
+		return !isForeground;
+	}
+}
diff --git a/Engine/Source/Developer/LowLevelTestsRunner/LowLevelTestsAndroid_UPL.xml b/Engine/Source/Developer/LowLevelTestsRunner/LowLevelTestsAndroid_UPL.xml
new file mode 100644
--- /dev/null
+++ b/Engine/Source/Developer/LowLevelTestsRunner/LowLevelTestsAndroid_UPL.xml
@@ -0,0 +1,56 @@
+<?xml version="1.0" encoding="utf-8"?>
+<!-- Low level tests UPL script -->
+<!-- Replaces application and activity elements, removes unnecessary elements, replaces Java /src with low level tests Java source -->
+<root xmlns:android="http://schemas.android.com/apk/res/android">
+	<!-- init section is always evaluated once per architecture -->
+	<init>
+		<log text="Low level tests UPL script"/>
+	</init>
+	<androidManifestUpdates>
+		<removeElement tag="application"/>
+		<setElement result="TestApplication" value="application" />
+		<addAttribute tag="$TestApplication" name="android:name" value="com.epicgames.unreal.tests.TestApplication" />
+		<addAttribute tag="$TestApplication" name="android:label" value="@string/app_name" />
+		<addAttribute tag="$TestApplication" name="android:icon" value="@drawable/icon" />
+		<addAttribute tag="$TestApplication" name="android:hardwareAccelerated" value="true" />
+		<addAttribute tag="$TestApplication" name="android:extractNativeLibs" value="true" />
+		<addAttribute tag="$TestApplication" name="android:requestLegacyExternalStorage" value="true" />
+		<addAttribute tag="$TestApplication" name="android:hasCode" value="true" />
+		<addElement tag="manifest" name="TestApplication"/>
+
+		<setElement result="TestActivity" value="activity" />
+		<addAttribute tag="$TestActivity" name="android:name" value="com.epicgames.unreal.tests.TestActivity" />
+		<addAttribute tag="$TestActivity" name="android:exported" value="true" />
+		<addAttribute tag="$TestActivity" name="android:label" value="@string/app_name" />
+		<addAttribute tag="$TestActivity" name="android:theme" value="@android:style/Theme.Black.NoTitleBar.Fullscreen" />
+		<addAttribute tag="$TestActivity" name="android:launchMode" value="singleTask" />
+		<addAttribute tag="$TestActivity" name="android:debuggable" value="true" />
+		<addElement tag="application" name="TestActivity"/>
+
+		<setElement result="TestActivityMetaData" value="meta-data" />
+		<addAttribute tag="$TestActivityMetaData" name="android:name" value="android.app.lib_name" />
+		<addAttribute tag="$TestActivityMetaData" name="android:value" value="Unreal" />
+		<addElement tag="activity" name="TestActivityMetaData" />
+
+		<setElement result="TestActivityAction" value="action" />
+		<addAttribute tag="$TestActivityAction" name="android:name" value="android.intent.action.MAIN" />
+
+		<setElement result="TestActivityCategory" value="category" />
+		<addAttribute tag="$TestActivityCategory" name="android:name" value="android.intent.category.LAUNCHER" />
+
+		<setElement result="TestActivityIntentFilter" value="intent-filter" />
+		<addElement tag="activity" name="TestActivityIntentFilter" />
+
+		<addElement tag="intent-filter" name="TestActivityAction" />
+		<addElement tag="intent-filter" name="TestActivityCategory" />
+	</androidManifestUpdates>
+	<resourceCopies>
+		<isArch arch="arm64-v8a">
+			<deleteFiles filespec="arm64/src/*.*" recursive="true" />
+		</isArch>
+		<isArch arch="x86_64">
+			<deleteFiles filespec="x64/src/*.*" recursive="true" />
+		</isArch>
+		<copyDir src="$S(PluginDir)/Android/Java/src" dst="$S(BuildDir)/src"/>
+	</resourceCopies>
+</root>
diff --git a/Engine/Source/Developer/LowLevelTestsRunner/LowLevelTestsRunner.Build.cs b/Engine/Source/Developer/LowLevelTestsRunner/LowLevelTestsRunner.Build.cs
--- a/Engine/Source/Developer/LowLevelTestsRunner/LowLevelTestsRunner.Build.cs
+++ b/Engine/Source/Developer/LowLevelTestsRunner/LowLevelTestsRunner.Build.cs
@@ -42,6 +42,12 @@ namespace UnrealBuildTool.Rules
 			{
 				PrivateDependencyModuleNames.Add("CoreUObject");
 			}
+
+			if (Target.Platform == UnrealTargetPlatform.Android)
+			{
+				string ModulePath = Utils.MakePathRelativeTo(ModuleDirectory, Target.RelativeEnginePath);
+				AdditionalPropertiesForReceipt.Add("AndroidPlugin", Path.Combine(ModulePath, "LowLevelTestsAndroid_UPL.xml"));
+			}
 		}
 	}
 }
diff --git a/Engine/Source/Developer/LowLevelTestsRunner/Private/Platform/Android/AndroidTestRunner.cpp b/Engine/Source/Developer/LowLevelTestsRunner/Private/Platform/Android/AndroidTestRunner.cpp
--- a/Engine/Source/Developer/LowLevelTestsRunner/Private/Platform/Android/AndroidTestRunner.cpp
+++ b/Engine/Source/Developer/LowLevelTestsRunner/Private/Platform/Android/AndroidTestRunner.cpp
@@ -3,6 +3,10 @@
 #if defined(PLATFORM_ANDROID)
 #include "TestRunner.h"
 
+#include <jni.h>
+#include <string>
+#include <vector>
+
 const char* AndroidPath;
 
 extern "C" int LowLevelTestsMain(int argc, const char* argv[])
@@ -19,4 +23,47 @@ const char* GetProcessExecutablePath()
 {
 	return AndroidPath;
 }
+
+extern "C"
+{
+	// Maps to runTests in TestActivity.java
+	JNIEXPORT jint JNICALL
+		Java_com_epicgames_unreal_tests_TestActivity_runTests(
+			JNIEnv* env,
+			jobject jobj, jstring jpath, jobjectArray jargs)
+	{
+		jboolean isCopy;
+		const char* Path = env->GetStringUTFChars(jpath, &isCopy);
+
+		jsize jargsCount = env->GetArrayLength(jargs);
+
+		const char* Args[jargsCount + 1];
+		if (Path != nullptr)
+		{
+			Args[0] = AndroidPath = Path;
+		}
+		env->ReleaseStringUTFChars(jpath, Path);
+
+		for (int i = 0; i < jargsCount; i++) {
+			jboolean isArgCopy;
+			jstring jarg = (jstring)env->GetObjectArrayElement(jargs, i);
+			const char* Arg = env->GetStringUTFChars(jarg, &isArgCopy);
+			char* ArgCpy = new char[strlen(Arg)];
+			if (Arg != nullptr)
+			{
+				strcpy(ArgCpy, Arg);
+				Args[i + 1] = ArgCpy;
+			}
+			env->ReleaseStringUTFChars(jarg, Arg);
+		}
+
+		return LowLevelTestsMain(jargsCount + 1, Args);
+	}
+
+	JNIEXPORT jint JNI_OnLoad(JavaVM* InJavaVM, void* InReserved)
+	{
+		return JNI_VERSION_1_6;
+	}
+}
+
 #endif
\ No newline at end of file
diff --git a/Engine/Source/Developer/LowLevelTestsRunner/Private/TestCommon/Initialization.cpp b/Engine/Source/Developer/LowLevelTestsRunner/Private/TestCommon/Initialization.cpp
--- a/Engine/Source/Developer/LowLevelTestsRunner/Private/TestCommon/Initialization.cpp
+++ b/Engine/Source/Developer/LowLevelTestsRunner/Private/TestCommon/Initialization.cpp
@@ -35,14 +35,14 @@ void InitStats()
 
 void UsePlatformFileStubIfRequired()
 {
-#if WITH_ENGINE && UE_LLT_USE_PLATFORM_FILE_STUB
+#if UE_LLT_USE_PLATFORM_FILE_STUB
 	if (IPlatformFile* WrapperFile = FPlatformFileManager::Get().GetPlatformFile(TEXT("PlatformFileStub")))
 	{
 		IPlatformFile* CurrentPlatformFile = &FPlatformFileManager::Get().GetPlatformFile();
 		WrapperFile->Initialize(CurrentPlatformFile, TEXT(""));
 		FPlatformFileManager::Get().SetPlatformFile(*WrapperFile);
 	}
-#endif // WITH_ENGINE && UE_LLT_USE_PLATFORM_FILE_STUB
+#endif // UE_LLT_USE_PLATFORM_FILE_STUB
 }
 
 void SaveDefaultPlatformFile()
diff --git a/Engine/Source/Programs/AutomationTool/Android/AndroidPlatform.Automation.cs b/Engine/Source/Programs/AutomationTool/Android/AndroidPlatform.Automation.cs
--- a/Engine/Source/Programs/AutomationTool/Android/AndroidPlatform.Automation.cs
+++ b/Engine/Source/Programs/AutomationTool/Android/AndroidPlatform.Automation.cs
@@ -33,6 +33,8 @@ public class AndroidPlatform : Platform
     private const string TargetAndroidLocation = "obb/";
 	private const string TargetAndroidTemp = "/data/local/tmp/";
 
+	private const string DefaultLaunchActivity = "com.epicgames.unreal.SplashActivity";
+
 	public class AdbCreatedProcess : AutomationTool.IProcessResult
 	{
 		private readonly object StopSyncObject = new object();
@@ -3727,7 +3729,7 @@ public class AndroidPlatform : Platform
 	/** Returns the launch activity name to launch (must call GetPackageInfo first), returns "com.epicgames.unreal.SplashActivity" default if not found */
 	public static string GetLaunchableActivityName()
 	{
-		string ReturnValue = "com.epicgames.unreal.SplashActivity";
+		string ReturnValue = DefaultLaunchActivity;
 		if (LaunchableActivityLine != null)
 		{
 			// the line should look like: launchable-activity: name='com.epicgames.unreal.SplashActivity'  label='TappyChicken' icon=''
@@ -3740,6 +3742,15 @@ public class AndroidPlatform : Platform
 		return ReturnValue;
 	}
 
+	public static string GetLaunchableActivityName(string ApkName)
+	{
+		if (!string.IsNullOrEmpty(ApkName) && LaunchableActivityLine == null)
+		{
+			GetPackageInfo(ApkName, false);
+		}
+		return GetLaunchableActivityName();
+	}
+
 	/** Returns the app type from the packaged APK metadata, returns "" if not found */
 	public static string GetMetaAppType()
 	{
diff --git a/Engine/Source/Programs/AutomationTool/Gauntlet/Platform/Android/Gauntlet.TargetDeviceAndroid.cs b/Engine/Source/Programs/AutomationTool/Gauntlet/Platform/Android/Gauntlet.TargetDeviceAndroid.cs
--- a/Engine/Source/Programs/AutomationTool/Gauntlet/Platform/Android/Gauntlet.TargetDeviceAndroid.cs
+++ b/Engine/Source/Programs/AutomationTool/Gauntlet/Platform/Android/Gauntlet.TargetDeviceAndroid.cs
@@ -427,7 +427,7 @@ namespace Gauntlet
 			// kill any currently running instance:
 			KillRunningProcess(Install.AndroidPackageName);
 
-			string LaunchActivity = AndroidPlatform.GetLaunchableActivityName();
+			string LaunchActivity = AndroidPlatform.GetLaunchableActivityName(Install.ApkPath);
 
 			Log.Info("Launching {0} on '{1}' ", Install.AndroidPackageName + "/" + LaunchActivity, ToString());
 			Log.Verbose("\t{0}", Install.CommandLine);
@@ -856,6 +856,9 @@ namespace Gauntlet
 			string ExternalFilesPath = StoragePath + "/Android/data/" + Build.AndroidPackageName + "/files/UnrealGame/" + AppConfig.ProjectName;
 			string DeviceBaseDir = Build.UsesExternalFilesDir ? ExternalFilesPath : ExternalStoragePath;
 
+			// path to the APK to install.
+			string ApkPath = Build.SourceApkPath;
+
 			// get the device's external file paths, always clear between runs
 			DeviceExternalStorageSavedPath = string.Format("{0}/{1}/Saved", ExternalStoragePath, AppConfig.ProjectName);
 			DeviceExternalFilesSavedPath = string.Format("{0}/{1}/Saved", ExternalFilesPath, AppConfig.ProjectName);
@@ -908,9 +911,6 @@ namespace Gauntlet
 				RunAdbDeviceCommand(string.Format("shell mkdir -p {0}/", DeviceExternalStorageSavedPath));
 				RunAdbDeviceCommand(string.Format("shell mkdir -p {0}/", DeviceExternalFilesSavedPath));
 
-				// path to the APK to install.
-				string ApkPath = Build.SourceApkPath;
-
 				// check for a local newer executable
 				if (Globals.Params.ParseParam("dev"))
 				{
@@ -1046,7 +1046,7 @@ namespace Gauntlet
 
 			CopyCommandlineFile(AppConfig.CommandLine, Build.AndroidPackageName, AppConfig.ProjectName, Build.UsesExternalFilesDir);
 
-			AndroidAppInstall AppInstall = new AndroidAppInstall(this, AppConfig.ProjectName, Build.AndroidPackageName, AppConfig.CommandLine);
+			AndroidAppInstall AppInstall = new AndroidAppInstall(this, InApkPath: ApkPath, AppConfig.ProjectName, Build.AndroidPackageName, AppConfig.CommandLine);
 
 			return AppInstall;
 		}
@@ -1067,6 +1067,8 @@ namespace Gauntlet
 
 		public string AppTag { get; set; }
 
+		public string ApkPath { get; protected set; }
+
 		public AndroidAppInstall(TargetDeviceAndroid InDevice, string InName, string InAndroidPackageName, string InCommandLine, string InAppTag = "UE")
 		{
 			AndroidDevice = InDevice;
@@ -1076,6 +1078,16 @@ namespace Gauntlet
 			AppTag = InAppTag;
 		}
 
+		public AndroidAppInstall(TargetDeviceAndroid InDevice, string InApkPath, string InName, string InAndroidPackageName, string InCommandLine, string InAppTag = "UE")
+		{
+			AndroidDevice = InDevice;
+			ApkPath = InApkPath;
+			Name = InName;
+			AndroidPackageName = InAndroidPackageName;
+			CommandLine = InCommandLine;
+			AppTag = InAppTag;
+		}
+
 		public IAppInstance Run()
 		{
 			return AndroidDevice.Run(this);
diff --git a/Engine/Source/Programs/AutomationTool/LowLevelTests/RunLowLevelTests.cs b/Engine/Source/Programs/AutomationTool/LowLevelTests/RunLowLevelTests.cs
--- a/Engine/Source/Programs/AutomationTool/LowLevelTests/RunLowLevelTests.cs
+++ b/Engine/Source/Programs/AutomationTool/LowLevelTests/RunLowLevelTests.cs
@@ -656,7 +656,7 @@ namespace LowLevelTests
 	{
 		bool CanSupportPlatform(UnrealTargetPlatform InPlatform);
 
-		StagedBuild CreateBuild(UnrealTargetPlatform InPlatform, UnrealTargetConfiguration InConfiguration, string InTestApp, string InBuildPath, bool bSkipStage);
+		IBuild CreateBuild(UnrealTargetPlatform InPlatform, UnrealTargetConfiguration InConfiguration, string InTestApp, string InBuildPath, bool bSkipStage);
 		protected static string GetExecutable(UnrealTargetPlatform InPlatform, UnrealTargetConfiguration InConfiguration, string InTestApp, string InBuildPath, string FileRegEx)
 		{
 			IEnumerable<string> Executables = DirectoryUtils.FindMatchingFiles(InBuildPath, FileRegEx, -1).Select(FileInfo => FileInfo.FullName);
@@ -691,6 +691,12 @@ namespace LowLevelTests
 					continue;
 				}
 
+				if (InPlatform == UnrealTargetPlatform.Android && BuildExecutableName.StartsWith(InTestApp))
+				{
+					Log.VeryVerbose("Output Executable for Android: {0}", Path.GetRelativePath(InBuildPath, Executable));
+					return Path.GetRelativePath(InBuildPath, Executable);
+				}
+
 				// Development executable does not contain configuration or platform name
 				Log.Verbose("Config type: {0}", InConfiguration);
 				if (InConfiguration == UnrealTargetConfiguration.Development)
@@ -753,7 +759,7 @@ namespace LowLevelTests
 			return InPlatform.IsInGroup(UnrealPlatformGroup.Desktop);
 		}
 
-		public StagedBuild CreateBuild(UnrealTargetPlatform InPlatform, UnrealTargetConfiguration InConfiguration, string InTestApp, string InBuildPath, bool bSkipStage)
+		public IBuild CreateBuild(UnrealTargetPlatform InPlatform, UnrealTargetConfiguration InConfiguration, string InTestApp, string InBuildPath, bool bSkipStage)
 		{
 			string ExecutablePath = ILowLevelTestsBuildFactory.GetExecutable(InPlatform, InConfiguration, InTestApp, InBuildPath, GetExecutableRegex(InPlatform));
 			return new LowLevelTestsBuild(InPlatform, InConfiguration, InBuildPath, ExecutablePath);
@@ -778,6 +784,43 @@ namespace LowLevelTests
 		}
 	}
 
+	public class AndroidLowLevelTestsBuildFactory : ILowLevelTestsBuildFactory
+	{
+		public bool CanSupportPlatform(UnrealTargetPlatform InPlatform)
+		{
+			return InPlatform.IsInGroup(UnrealPlatformGroup.Android);
+		}
+
+		public IBuild CreateBuild(UnrealTargetPlatform InPlatform, UnrealTargetConfiguration InConfiguration, string InTestApp, string InBuildPath, bool bSkipStage)
+		{
+			string ExecutablePath = ILowLevelTestsBuildFactory.GetExecutable(InPlatform, InConfiguration, InTestApp, InBuildPath, GetExecutableRegex(InPlatform));
+			return new AndroidBuild(InConfiguration, $"com.epicgames.{InTestApp}", Path.Combine(InBuildPath, ExecutablePath), new Dictionary<string, string>(), BuildFlags.Packaged | BuildFlags.CanReplaceCommandLine, false, false, false);
+		}
+
+		public string GetExecutableRegex(UnrealTargetPlatform InPlatform)
+		{
+			return @"[A-Za-z0-9_]+(Tests)?(?:-[A-Za-z0-9_]+)?(?:-[A-Za-z0-9_]+)?.apk";
+		}
+	}
+
+	public class AndroidLowLevelTestsReporting : ILowLevelTestsReporting
+	{
+		public bool CanSupportPlatform(UnrealTargetPlatform InPlatform)
+		{
+			return InPlatform.IsInGroup(UnrealPlatformGroup.Android);
+		}
+
+		public string GetTargetReportPath(UnrealTargetPlatform InPlatform, string InTestApp, string InBuildPath)
+		{
+			return string.Format("{0}LLTResults.out", InPlatform.ToString());
+		}
+
+		public string CopyDeviceReportTo(IAppInstall InAppInstall, UnrealTargetPlatform InPlatform, string InTestApp, string InBuildPath, string InTargetDirectory)
+		{
+			throw new NotSupportedException("Reports currently not supported for Android."); 
+		}
+	}
+
 	public class DesktopLowLevelTestsReporting : ILowLevelTestsReporting
 	{
 		public bool CanSupportPlatform(UnrealTargetPlatform InPlatform)
@@ -824,7 +867,7 @@ namespace LowLevelTests
 
 		public UnrealTargetPlatform Platform { get; protected set; }
 		public UnrealTargetConfiguration Configuration { get; protected set; }
-		public StagedBuild DiscoveredBuild { get; protected set; }
+		public IBuild DiscoveredBuild { get; protected set; }
 
 		public LowLevelTestsBuildSource(string InTestApp, string InBuildPath, UnrealTargetPlatform InTargetPlatform, UnrealTargetConfiguration InConfiguration, bool InSkipStage)
 		{
@@ -891,7 +934,11 @@ namespace LowLevelTests
 
 				// Set reporting options, filters etc
 				CachedConfig.CommandLineParams.AddRawCommandline("--durations=no");
-				if (!string.IsNullOrEmpty(InReportType))
+				if (CachedConfig.Platform == UnrealTargetPlatform.Android)
+				{
+					CachedConfig.CommandLineParams.AddRawCommandline("--reporter=console");
+				}
+				else if (!string.IsNullOrEmpty(InReportType))
 				{
 					CachedConfig.CommandLineParams.AddRawCommandline(string.Format("--reporter={0}", InReportType));
 					string ReportPath = LowLevelTestsReporting.GetTargetReportPath(Platform, TestApp, BuildPath);
diff --git a/Engine/Source/Programs/AutomationTool/LowLevelTests/Tests/Gauntlet.LowLevelTestNode.cs b/Engine/Source/Programs/AutomationTool/LowLevelTests/Tests/Gauntlet.LowLevelTestNode.cs
--- a/Engine/Source/Programs/AutomationTool/LowLevelTests/Tests/Gauntlet.LowLevelTestNode.cs
+++ b/Engine/Source/Programs/AutomationTool/LowLevelTests/Tests/Gauntlet.LowLevelTestNode.cs
@@ -6,6 +6,7 @@ using Gauntlet;
 using System;
 using System.IO;
 using System.Linq;
+using System.Text.RegularExpressions;
 using UnrealBuildBase;
 using UnrealBuildTool;
 
@@ -132,13 +133,7 @@ namespace LowLevelTests
 					ParseLowLevelTestsLog();
 
 					// Print stdout when -captureoutput, certain platforms don't always redirect stdout
-					if (CurrentProcessedLines != null && Context.Options.CaptureOutput)
-					{
-						foreach (string OutputLine in CurrentProcessedLines)
-						{
-							Console.WriteLine(OutputLine);
-						}
-					}
+					PrintLogIfCaptureOutput();
 
 					if (CheckForTimeout())
 					{
@@ -177,6 +172,9 @@ namespace LowLevelTests
 			{
 				base.StopTest(InReason);
 
+				ParseLowLevelTestsLog();
+				PrintLogIfCaptureOutput();
+
 				if (TestInstance != null && !TestInstance.HasExited)
 				{
 					TestInstance.Kill();
@@ -217,7 +215,29 @@ namespace LowLevelTests
 
 				bool? ReportCopied = null;
 				string ReportPath = null;
-				if (!string.IsNullOrEmpty(Context.Options.ReportType))
+
+				int? ExitCodeOverride = null;
+
+				// No reports from Android tests yet. Since adb shell doesn't forward exit code, we look for it in the log output.
+				if (Context.Options.Platform == UnrealTargetPlatform.Android)
+				{
+					Match AndroidExitCodeLog = Regex.Match(TestInstance.StdOut, @"Tests finished with exit code (\d+)");
+					if (AndroidExitCodeLog.Success)
+					{
+						ExitCodeOverride = int.Parse(Regex.Match(AndroidExitCodeLog.Value, @"\d+").Value);
+					}
+					else if (TestInstance.StdOut.Contains("beginning of crash"))
+					{
+						Log.Info("Crash occurred during test.");
+						ExitCodeOverride = -1;
+					}
+					else
+					{
+						Log.Error("Could not find exit code in Android log, assumming failure.");
+						ExitCodeOverride = -1;
+					}
+				}
+				else if (!string.IsNullOrEmpty(Context.Options.ReportType))
 				{
 					ILowLevelTestsReporting LowLevelTestsReporting = Gauntlet.Utils.InterfaceHelpers.FindImplementations<ILowLevelTestsReporting>(true)
 						.Where(B => B.CanSupportPlatform(Context.Options.Platform))
@@ -235,8 +255,8 @@ namespace LowLevelTests
 					}
 				}
 
-
 				string ExitReason = "";
+				int ExitCode = ExitCodeOverride.HasValue ? ExitCodeOverride.Value : TestInstance.ExitCode;
 				if (TestInstance.WasKilled)
 				{
 					if (InReason == StopReason.MaxDuration || LowLevelTestResult == TestResult.TimedOut)
@@ -250,10 +270,10 @@ namespace LowLevelTests
 						ExitReason = $"Process was killed by Gauntlet with reason {InReason.ToString()}.";
 					}
 				}
-				else if (TestInstance.ExitCode != 0)
+				else if (ExitCode != 0)
 				{
 					LowLevelTestResult = TestResult.Failed;
-					ExitReason = $"Process exited with exit code {TestInstance.ExitCode}";
+					ExitReason = $"Process exited with exit code {ExitCode}";
 				}
 				else if (ReportCopied.HasValue && !ReportCopied.Value)
 				{
@@ -299,7 +319,7 @@ namespace LowLevelTests
 				}
 				else // ReportPath == null
 				{
-					if (TestInstance.ExitCode != 0)
+					if (ExitCode != 0)
 					{
 						LowLevelTestResult = TestResult.Failed;
 						ExitReason = "Tests failed (no report to parse)";
@@ -310,7 +330,7 @@ namespace LowLevelTests
 						ExitReason = "Tests passed (no report to parse)";
 					}
 				}
-				Log.Info($"Low level test exited with code {TestInstance.ExitCode} and reason: {ExitReason}");
+				Log.Info($"Low level test exited with code {ExitCode} and reason: {ExitReason}");
 			}
 			catch
 			{
@@ -342,12 +362,23 @@ namespace LowLevelTests
 				CurrentProcessedLines = TestInstance.StdOut
 					.Substring(LastStdoutSeekPos)
 					.Split("\n")
-					.Where(Line => Line.Contains("LogLowLevelTests"))
+					.Where(Line => !string.IsNullOrWhiteSpace(Line))
 					.ToArray();
 				LastStdoutSeekPos = TestInstance.StdOut.Length - 1;
 			}
 		}
 
+		private void PrintLogIfCaptureOutput()
+		{
+			if (CurrentProcessedLines != null && Context.Options.CaptureOutput)
+			{
+				foreach (string OutputLine in CurrentProcessedLines)
+				{
+					Console.WriteLine(OutputLine.TrimEnd(Environment.NewLine.ToCharArray()));
+				}
+			}
+		}
+
 		private bool CheckForTimeout()
 		{
 			if (CurrentProcessedLines == null || CurrentProcessedLines.Length == 0)
diff --git a/Engine/Source/Programs/LowLevelTests/FoundationTests/FoundationTests.Build.cs b/Engine/Source/Programs/LowLevelTests/FoundationTests/FoundationTests.Build.cs
--- a/Engine/Source/Programs/LowLevelTests/FoundationTests/FoundationTests.Build.cs
+++ b/Engine/Source/Programs/LowLevelTests/FoundationTests/FoundationTests.Build.cs
@@ -55,6 +55,15 @@ public class FoundationTests : TestModuleRules
 			}
 			TestMetadata.PlatformCompilationExtraArgs.Add(Platform, PlatformCompilationArgs);
 		}
+
+		// Platform-specific tags
+		TestMetadata.PlatformTags.Add(UnrealTargetPlatform.Linux, "~[.]~[Slow]");
+		TestMetadata.PlatformTags.Add(UnrealTargetPlatform.Android, "~[Perf]~[Slow]~[AndroidSkip]");
+
+		// Allow Android run for this test
+		// Will remove Android from PlatformsRunUnsupported as more diverse types of tests can run on this platform
+		TestMetadata.PlatformsRunUnsupported.Remove(UnrealTargetPlatform.Android);
+
 		UpdateBuildGraphPropertiesFile(TestMetadata);
 	}
 }
\ No newline at end of file
diff --git a/Engine/Source/Programs/LowLevelTests/MathCoreTests/MathCoreTests.Build.cs b/Engine/Source/Programs/LowLevelTests/MathCoreTests/MathCoreTests.Build.cs
--- a/Engine/Source/Programs/LowLevelTests/MathCoreTests/MathCoreTests.Build.cs
+++ b/Engine/Source/Programs/LowLevelTests/MathCoreTests/MathCoreTests.Build.cs
@@ -43,6 +43,7 @@ public class MathCoreTests : TestModuleRules
 			}
 			TestMetadata.PlatformCompilationExtraArgs.Add(Platform, PlatformCompilationArgs);
 		}
+		TestMetadata.PlatformsRunUnsupported.Remove(UnrealTargetPlatform.Android);
 		UpdateBuildGraphPropertiesFile(TestMetadata);
 	}
 }
\ No newline at end of file
diff --git a/Engine/Source/Programs/LowLevelTests/MathCoreTests/Private/TestGroupEvents.cpp b/Engine/Source/Programs/LowLevelTests/MathCoreTests/Private/TestGroupEvents.cpp
--- a/Engine/Source/Programs/LowLevelTests/MathCoreTests/Private/TestGroupEvents.cpp
+++ b/Engine/Source/Programs/LowLevelTests/MathCoreTests/Private/TestGroupEvents.cpp
@@ -3,6 +3,11 @@
 
 #include <catch2/catch_test_macros.hpp>
 
+GROUP_BEFORE_GLOBAL(Catch::DefaultGroup)
+{
+	UsePlatformFileStubIfRequired();
+}
+
 GROUP_AFTER_GLOBAL(Catch::DefaultGroup)
 {
 	CleanupLogs();
diff --git a/Engine/Source/Programs/UnrealBuildTool/Configuration/TestModuleRules.cs b/Engine/Source/Programs/UnrealBuildTool/Configuration/TestModuleRules.cs
--- a/Engine/Source/Programs/UnrealBuildTool/Configuration/TestModuleRules.cs
+++ b/Engine/Source/Programs/UnrealBuildTool/Configuration/TestModuleRules.cs
@@ -121,12 +121,6 @@ namespace UnrealBuildTool
 				PrivateDependencyModuleNames.Add("LowLevelTestsRunner");
 			}
 
-			// Platforms specific setup
-			if (Target.Platform == UnrealTargetPlatform.Android)
-			{
-				PublicDefinitions.Add("CATCH_CONFIG_NOSTDOUT");
-			}
-
 			if (Target.Platform == UnrealTargetPlatform.IOS || Target.Platform == UnrealTargetPlatform.TVOS)
 			{
 				// Fix missing frameworks from ApplicationCore
diff --git a/Engine/Source/Programs/UnrealBuildTool/Configuration/TestTargetRules.cs b/Engine/Source/Programs/UnrealBuildTool/Configuration/TestTargetRules.cs
--- a/Engine/Source/Programs/UnrealBuildTool/Configuration/TestTargetRules.cs
+++ b/Engine/Source/Programs/UnrealBuildTool/Configuration/TestTargetRules.cs
@@ -275,7 +275,7 @@ namespace UnrealBuildTool
 
 			bBuildInSolutionByDefault = false;
 
-			bDeployAfterCompile = target.Platform != UnrealTargetPlatform.Android;
+			bDeployAfterCompile = true;
 			bIsBuildingConsoleApplication = true;
 
 			// Disabling default true flags that aren't necessary for tests
@@ -311,6 +311,8 @@ namespace UnrealBuildTool
 			bBuildRequiresCookedData = true;
 			bBuildDeveloperTools = false;
 
+			bUsePlatformFileStub = (Platform == UnrealTargetPlatform.Android);
+
 			// Useful for debugging test failures
 			if (target.Configuration == UnrealTargetConfiguration.Debug)
 			{
@@ -326,7 +328,6 @@ namespace UnrealBuildTool
 			GlobalDefinitions.Add("TEST_FOR_VALID_FILE_SYSTEM_MEMORY=0");
 
 			// LLT Globals
-			GlobalDefinitions.Add("UE_LLT_USE_PLATFORM_FILE_STUB=0");
 			GlobalDefinitions.Add("UE_LLT_WITH_MOCK_ENGINE_DEFAULTS=0");
 
 			// Platform specific setup
@@ -339,10 +340,19 @@ namespace UnrealBuildTool
 				GlobalDefinitions.Add("USE_ANDROID_LAUNCH=0");
 				GlobalDefinitions.Add("USE_ANDROID_JNI=0");
 
-				// Workaround for a linker bug when building LowLevelTests for Android
+				// Addresses linker errors when building LowLevelTests for Android
 				// TODO: This should be written to the intermediate directory, somewhere else not when TargetRules are being created
 				FileReference versionScriptFile = new FileReference(Path.GetTempPath() + $"LLTWorkaroundScrip-{Name}.ldscript");
-				FileReference.WriteAllTextIfDifferent(versionScriptFile, "{ local: *; };");
+				string SymbolsLinkScriptContent = @"{
+					global: LowLevelTestsMain;
+					global: JNI_OnLoad;
+					global: Java_com_epicgames_unreal_tests_TestActivity_runTests;
+					*VM*;
+					*AndroidPath*;
+					*Catch*;
+					local: *;
+				};";
+				FileReference.WriteAllTextIfDifferent(versionScriptFile, SymbolsLinkScriptContent);
 				AdditionalLinkerArguments = " -Wl,--version-script=\"" + versionScriptFile.FullName + "\"";
 			}
 			else if (target.Platform == UnrealTargetPlatform.IOS)
diff --git a/Engine/Source/Programs/UnrealBuildTool/Platform/Android/UEDeployAndroid.cs b/Engine/Source/Programs/UnrealBuildTool/Platform/Android/UEDeployAndroid.cs
--- a/Engine/Source/Programs/UnrealBuildTool/Platform/Android/UEDeployAndroid.cs
+++ b/Engine/Source/Programs/UnrealBuildTool/Platform/Android/UEDeployAndroid.cs
@@ -4987,6 +4987,12 @@ namespace UnrealBuildTool
 				// we don't actually need the SO for the bSkipGradleBuild case
 				string? FinalSOName = null;
 				string DestApkDirectory = Path.Combine(ProjectDirectory, "Binaries", "Android");
+				// If OutputPath is not at the same level with DestApkDirectory, then the target was built with binary subfolder
+				string? OutputDirectory = Path.GetDirectoryName(OutputPath);
+				if (OutputDirectory != null && Path.GetRelativePath(DestApkDirectory, OutputDirectory) != OutputDirectory)
+				{
+					DestApkDirectory = OutputDirectory;
+				}
 				string? DestApkName = null;
 				if (bSkipGradleBuild)
 				{
diff --git a/Engine/Source/Programs/UnrealBuildTool/System/UnrealPluginLanguage.cs b/Engine/Source/Programs/UnrealBuildTool/System/UnrealPluginLanguage.cs
--- a/Engine/Source/Programs/UnrealBuildTool/System/UnrealPluginLanguage.cs
+++ b/Engine/Source/Programs/UnrealBuildTool/System/UnrealPluginLanguage.cs
@@ -332,10 +332,10 @@ namespace UnrealBuildTool
 	 * If you must remove files (like development-only files from distribution builds) you can
 	 * use this node:
 	 * 
-	 *	<deleteFiles filespec=""/>
+	 *	<deleteFiles filespec="" recursive=""/>
 	 *	
-	 * It is restricted to only removing files from the BuildDir.  Here is example usage to remove
-	 * the Oculus Signature Files (osig) from the assets directory:
+	 * It is restricted to only removing files from the BuildDir. If recursive is true it will search all subfolders.
+	 * Here is example usage to remove the Oculus Signature Files (osig) from the assets directory:
 	 * 
 	 *	<deleteFiles filespec="assets/oculussig_*"/>
 	 *	
@@ -991,7 +991,7 @@ namespace UnrealBuildTool
 			}
 		}
 
-		private static void DeleteFiles(string Filespec, ILogger Logger)
+		private static void DeleteFiles(string Filespec, bool Recursive, ILogger Logger)
 		{
 			string BaseDir = Path.GetDirectoryName(Filespec)!;
 			string Mask = Path.GetFileName(Filespec);
@@ -1001,7 +1001,7 @@ namespace UnrealBuildTool
 				return;
 			}
 
-			string[] Files = Directory.GetFiles(BaseDir, Mask, SearchOption.TopDirectoryOnly);
+			string[] Files = Directory.GetFiles(BaseDir, Mask, Recursive ? SearchOption.AllDirectories : SearchOption.TopDirectoryOnly);
 			foreach (string Filename in Files)
 			{
 				File.SetAttributes(Filename, FileAttributes.Normal);
@@ -1805,6 +1805,7 @@ namespace UnrealBuildTool
 					case "deleteFiles":
 						{
 							string? Filespec = GetAttribute(CurrentContext, Node, "filespec");
+							bool bRecursive = StringToBool(GetAttribute(CurrentContext, Node, "recursive", true, false, "false"));
 							if (Filespec != null)
 							{
 								if (Filespec.Contains(':') || Filespec.Contains(".."))
@@ -1814,7 +1815,7 @@ namespace UnrealBuildTool
 								else
 								{
 									// force relative to BuildDir (and only from global context so someone doesn't try to be clever)
-									DeleteFiles(Path.Combine(GlobalContext.StringVariables["BuildDir"], Filespec), Logger);
+									DeleteFiles(Path.Combine(GlobalContext.StringVariables["BuildDir"], Filespec), bRecursive, Logger);
 								}
 							}
 						}
diff --git a/Engine/Source/Runtime/Core/Tests/Misc/CharTest.cpp b/Engine/Source/Runtime/Core/Tests/Misc/CharTest.cpp
--- a/Engine/Source/Runtime/Core/Tests/Misc/CharTest.cpp
+++ b/Engine/Source/Runtime/Core/Tests/Misc/CharTest.cpp
@@ -31,7 +31,7 @@ void RunCharTests(uint32 MaxChar)
 	}
 }
 
-TEST_CASE_NAMED(TCharTest, "System::Core::Misc::Char", "[ApplicationContextMask][SmokeFilter]")
+TEST_CASE_NAMED(TCharTest, "System::Core::Misc::Char", "[ApplicationContextMask][SmokeFilter][AndroidSkip]")
 {
 	const char* CurrentLocale = setlocale(LC_CTYPE, nullptr);
 	if (CurrentLocale == nullptr)
