Add a idle timeout check to the UEPerf.EditorGauntletTestController tick

#rnx

diff --git a/Engine/Source/Programs/AutomationTool/Gauntlet/Unreal/Editor/EditorTests.cs b/Engine/Source/Programs/AutomationTool/Gauntlet/Unreal/Editor/EditorTests.cs
--- a/Engine/Source/Programs/AutomationTool/Gauntlet/Unreal/Editor/EditorTests.cs
+++ b/Engine/Source/Programs/AutomationTool/Gauntlet/Unreal/Editor/EditorTests.cs
@@ -1,11 +1,14 @@
 // Copyright Epic Games, Inc. All Rights Reserved.
 
+using System;
 using System.Collections.Generic;
 using System.Data;
 using System.IO;
 using System.Linq;
 using AutomationTool;
 using Gauntlet;
+using EpicGames.Core;
+using Log = Gauntlet.Log;
 
 namespace UEPerf
 {
@@ -259,6 +262,12 @@ namespace UEPerf
 		[AutoParam]
 		public string RHI { get; set; } = string.Empty;
 
+		/// <summary>
+		/// Log Idle timeout in second
+		/// </summary>
+		[AutoParam]
+		public int LogIdleTimeout { get; set; } = 30 * 60;
+
 		public override void ApplyToConfig(UnrealAppConfig AppConfig, UnrealSessionRole ConfigRole, IEnumerable<UnrealSessionRole> OtherRoles)
 		{
 			if (string.IsNullOrEmpty(Controller))
@@ -291,20 +300,12 @@ namespace UEPerf
 	public class EditorGauntletTestController : UnrealTestNode<EditorGauntletTestControllerConfig>
 	{
 		private int LastLogCount = 0;
+		private DateTime LastAutomationEntryTime = DateTime.MinValue;
 		private bool ValidateResolveMap = false;
 		public EditorGauntletTestController(UnrealTestContext InContext) : base(InContext)
 		{
 		}
 
-		/// <summary>
-		/// Log Channels to listen for activity
-		/// </summary>
-		/// <returns></returns>
-		public override IEnumerable<string> GetHeartbeatLogCategories()
-		{
-			return UnrealLog.EditorBusyChannels;
-		}
-
 		public override EditorGauntletTestControllerConfig GetConfiguration()
 		{
 			if (CachedConfig != null)
@@ -356,17 +357,50 @@ namespace UEPerf
 
 		public override void TickTest(UnrealSessionInstance InInstance)
 		{
-			base.TickTest(InInstance);
+			float IdleTimeout = 30 * 60;
+			if (GetConfiguration() is EditorGauntletTestControllerConfig Config && Config.LogIdleTimeout > 0)
+			{
+				IdleTimeout = Config.LogIdleTimeout;
+			}
 
-			if (ValidateResolveMap)
+			// We are only interested in what the editor is doing
+			var App = InInstance.EditorApp;
+			if (App != null)
 			{
-				var App = InInstance.EditorApp;
-				if (App != null)
+				UnrealLogStreamParser Parser = new UnrealLogStreamParser();
+				LastLogCount += Parser.ReadStream(App.StdOut, LastLogCount);
+
+				IEnumerable<string> ChannelEntries = Parser.GetLogFromShortNameChannels(UnrealLog.EditorBusyChannels.Append("Gauntlet"));
+
+				// Any new entries?
+				if (ChannelEntries.Any())
 				{
-					string Map = GetConfiguration().Map;
-					UnrealLogStreamParser Parser = new UnrealLogStreamParser();
-					LastLogCount += Parser.ReadStream(App.StdOut, LastLogCount);
+					// log new entries so people have something to look at
+					ChannelEntries.ToList().ForEach(S => Log.Info(S));
+					LastAutomationEntryTime = DateTime.Now;
+				}
+				else
+				{
+					// Check for timeouts
+					if (LastAutomationEntryTime == DateTime.MinValue)
+					{
+						LastAutomationEntryTime = DateTime.Now;
+					}
+
+					double ElapsedTime = (DateTime.Now - LastAutomationEntryTime).TotalSeconds;
 
+					// Check for timeout
+					if (ElapsedTime > IdleTimeout)
+					{
+						Log.Warning(KnownLogEvents.Gauntlet_TestEvent, "No activity observed in last {Time:0.00} minutes. Aborting test", IdleTimeout / 60);
+						MarkTestComplete();
+						SetUnrealTestResult(TestResult.TimedOut);
+					}
+				}
+
+				if (ValidateResolveMap)
+				{
+					string Map = GetConfiguration().Map;
 					string ResolvedMap = Parser.GetLogLinesContaining($"to resolve {Map}.").FirstOrDefault();
 
 					if (!string.IsNullOrEmpty(ResolvedMap))
