Improvements to the UE.CookOnTheFly test

1. The Lyra is launched and opens the map "L_Expanse" instead of just opening the start screen like before.
2. Fixed the problem of the test fail on timeout.

diff --git a/Engine/Source/Programs/AutomationTool/Gauntlet/Unreal/Automation/UE.CookOnTheFly.cs b/Engine/Source/Programs/AutomationTool/Gauntlet/Unreal/Automation/UE.CookOnTheFly.cs
--- a/Engine/Source/Programs/AutomationTool/Gauntlet/Unreal/Automation/UE.CookOnTheFly.cs
+++ b/Engine/Source/Programs/AutomationTool/Gauntlet/Unreal/Automation/UE.CookOnTheFly.cs
@@ -80,6 +80,7 @@ namespace UE
 
 		public override UnrealTestConfiguration GetConfiguration()
 		{
+			const string Map = "L_Expanse";
 			UnrealTestConfiguration Config = base.GetConfiguration();
 			UnrealTestRole EditorRole = Config.RequireRole(UnrealTargetRole.Editor);
 			EditorRole.CommandLine += "-run=cook -cookonthefly -zenstore -log";
@@ -104,7 +105,7 @@ namespace UE
 			}
 
 			ClientRole.DeferredLaunch = true;
-			ClientRole.CommandLine += string.Format("-cookonthefly -filehostip=\"{0}\" -log -LogCmds=\"LogCookOnTheFly Verbose\"", HostIP);
+			ClientRole.CommandLine += string.Format("-cookonthefly -filehostip=\"{0}\" -log -LogCmds=\"LogCookOnTheFly Verbose\" -map=\"{1}\"", HostIP, Map);
 
 			return Config;
 		}
@@ -268,7 +269,7 @@ namespace UE
 						&& EditorLogParser.GetLogLinesContaining(CookingProcessString).Any()
 						&& ClientCookEntries.Any()
 						&& ClientLogParser.GetLogLinesContaining(ReceivedPackagesCookedString).Any()
-						&& ClientLogParser.GetLogLinesContaining(TransportCreatedString).Any())
+						&& (ClientLogParser.GetLogLinesContaining(TransportCreatedString).Any() || !CommandUtils.IsBuildMachine))
 					{
 						Log.Info("Found '{0}', '{1}', '{2}', '{3}'. The CookOnTheFly log channel is active. The cooking process is taking place.", GameStartedString, CookingProcessString, ReceivedPackagesCookedString, TransportCreatedString);
 						MarkTestComplete();
