Fixing handling of idletimeout variable and its default value

#rnx

diff --git a/Engine/Source/Programs/AutomationTool/Gauntlet/Unreal/Editor/EditorTests.cs b/Engine/Source/Programs/AutomationTool/Gauntlet/Unreal/Editor/EditorTests.cs
--- a/Engine/Source/Programs/AutomationTool/Gauntlet/Unreal/Editor/EditorTests.cs
+++ b/Engine/Source/Programs/AutomationTool/Gauntlet/Unreal/Editor/EditorTests.cs
@@ -266,7 +266,7 @@ namespace UEPerf
 		/// Log Idle timeout in second
 		/// </summary>
 		[AutoParam]
-		public int LogIdleTimeout { get; set; } = 30 * 60;
+		public int LogIdleTimeoutSec { get; set; } = 0;
 
 		public override void ApplyToConfig(UnrealAppConfig AppConfig, UnrealSessionRole ConfigRole, IEnumerable<UnrealSessionRole> OtherRoles)
 		{
@@ -302,6 +302,7 @@ namespace UEPerf
 		private int LastLogCount = 0;
 		private DateTime LastAutomationEntryTime = DateTime.MinValue;
 		private bool ValidateResolveMap = false;
+		private float IdleTimeoutSec = 5 * 60;
 		public EditorGauntletTestController(UnrealTestContext InContext) : base(InContext)
 		{
 		}
@@ -352,17 +353,17 @@ namespace UEPerf
 				ValidateResolveMap = true;
 			}
 
+			if (Config.LogIdleTimeoutSec > 0)
+			{
+				IdleTimeoutSec = Config.LogIdleTimeoutSec;
+			}
+
+
 			return base.StartTest(Pass, InNumPasses);
 		}
 
 		public override void TickTest(UnrealSessionInstance InInstance)
 		{
-			float IdleTimeout = 30 * 60;
-			if (GetConfiguration() is EditorGauntletTestControllerConfig Config && Config.LogIdleTimeout > 0)
-			{
-				IdleTimeout = Config.LogIdleTimeout;
-			}
-
 			// We are only interested in what the editor is doing
 			var App = InInstance.EditorApp;
 			if (App != null)
@@ -390,9 +391,9 @@ namespace UEPerf
 					double ElapsedTime = (DateTime.Now - LastAutomationEntryTime).TotalSeconds;
 
 					// Check for timeout
-					if (ElapsedTime > IdleTimeout)
+					if (ElapsedTime > IdleTimeoutSec)
 					{
-						Log.Warning(KnownLogEvents.Gauntlet_TestEvent, "No activity observed in last {Time:0.00} minutes. Aborting test", IdleTimeout / 60);
+						Log.Warning(KnownLogEvents.Gauntlet_TestEvent, "No activity observed in last {Time:0.00} minutes. Aborting test", IdleTimeoutSec / 60);
 						MarkTestComplete();
 						SetUnrealTestResult(TestResult.TimedOut);
 					}
