Make IdleTimeout variable a member of the test node properties

#rnx

diff --git a/Engine/Source/Programs/AutomationTool/Gauntlet/Unreal/Automation/UE.Automation.cs b/Engine/Source/Programs/AutomationTool/Gauntlet/Unreal/Automation/UE.Automation.cs
--- a/Engine/Source/Programs/AutomationTool/Gauntlet/Unreal/Automation/UE.Automation.cs
+++ b/Engine/Source/Programs/AutomationTool/Gauntlet/Unreal/Automation/UE.Automation.cs
@@ -427,6 +427,8 @@ namespace UE
 
 		private DateTime LastAutomationEntryTime = DateTime.MinValue;
 
+		private float IdleTimeoutSec = 30 * 60;
+
 		/// Maximum of events display per test
 		protected virtual int MaxEventsDisplayPerTest { get; set; } = 10;
 
@@ -542,6 +544,10 @@ namespace UE
 						ReportDirInfo.Delete(true);
 					}
 				}
+				if (Config.LogIdleTimeout > 0)
+				{
+					IdleTimeoutSec = Config.LogIdleTimeout;
+				}
 			}
 
 			return base.StartTest(Pass, InNumPasses);
@@ -566,12 +572,6 @@ namespace UE
 		/// </summary>
 		public override void TickTest()
 		{
-			float IdleTimeout = 30 * 60;
-			if (GetConfiguration() is AutomationTestConfig Config && Config.LogIdleTimeout > 0)
-			{
-				IdleTimeout = Config.LogIdleTimeout;
-			}
-
 			// We are primarily interested in what the editor is doing
 			var AppInstance = TestInstance.EditorApp;
 
@@ -598,9 +598,9 @@ namespace UE
 				double ElapsedTime = (DateTime.Now - LastAutomationEntryTime).TotalSeconds;
 
 				// Check for timeout
-				if (ElapsedTime > IdleTimeout)
+				if (ElapsedTime > IdleTimeoutSec)
 				{
-					Log.Warning(KnownLogEvents.Gauntlet_TestEvent, "No activity observed in last {Time:0.00} minutes. Aborting test", IdleTimeout / 60);
+					Log.Warning(KnownLogEvents.Gauntlet_TestEvent, "No activity observed in last {Time:0.00} minutes. Aborting test", IdleTimeoutSec / 60);
 					MarkTestComplete();
 					SetUnrealTestResult(TestResult.TimedOut);
 				}
