Emit a warning if the ZenExportSnapshotTask experiences a non-zero return code when calling zen.exe to export an oplog.

diff --git a/Engine/Source/Programs/AutomationTool/BuildGraph/Tasks/ZenExportSnapshotTask.cs b/Engine/Source/Programs/AutomationTool/BuildGraph/Tasks/ZenExportSnapshotTask.cs
--- a/Engine/Source/Programs/AutomationTool/BuildGraph/Tasks/ZenExportSnapshotTask.cs
+++ b/Engine/Source/Programs/AutomationTool/BuildGraph/Tasks/ZenExportSnapshotTask.cs
@@ -382,6 +382,22 @@ namespace AutomationTool.Tasks
 			Writer.WriteObjectEnd();
 		}
 
+		private void TryRunAndLogWithoutSpew(string App, string CommandLine)
+		{
+			ProcessResult.SpewFilterCallbackType SilentOutputFilter = new ProcessResult.SpewFilterCallbackType(Line =>
+				{
+					return null;
+				});
+			try
+			{
+				CommandUtils.RunAndLog(CommandUtils.CmdEnv, App, CommandLine, MaxSuccessCode: 0, Options: CommandUtils.ERunOptions.Default, SpewFilterCallback: SilentOutputFilter);
+			}
+			catch (CommandUtils.CommandFailedException e)
+			{
+				Logger.LogWarning(e.ToString());
+			}
+		}
+
 		/// <summary>
 		/// Execute the task.
 		/// </summary>
@@ -469,11 +485,6 @@ namespace AutomationTool.Tasks
 				OplogExportCommandline.Append(" --force");
 			}
 
-			ProcessResult.SpewFilterCallbackType SilentOutputFilter = new ProcessResult.SpewFilterCallbackType(Line =>
-				{
-					return null;
-				});
-
 			switch (DestinationStorageType)
 			{
 				case SnapshotStorageType.Cloud:
@@ -538,7 +549,7 @@ namespace AutomationTool.Tasks
 						IoHash DestinationKeyHash = IoHash.Compute(Encoding.UTF8.GetBytes(ExportNames[ExportIndex]));
 
 						ExportSingleSourceCommandline.AppendFormat(" {0} --key {1} {2} {3} {4}", HostUrlArg, DestinationKeyHash.ToString().ToLowerInvariant(), BaseKeyArg, ExportSource.ProjectId, ExportSource.OplogId);
-						CommandUtils.RunAndLog(CommandUtils.CmdEnv, ZenExe.FullName, ExportSingleSourceCommandline.ToString(), MaxSuccessCode: int.MaxValue, Options: CommandUtils.ERunOptions.Default, SpewFilterCallback: SilentOutputFilter);
+						TryRunAndLogWithoutSpew(ZenExe.FullName, ExportSingleSourceCommandline.ToString());
 
 						ExportIndex = ExportIndex + 1;
 					}
@@ -558,7 +569,7 @@ namespace AutomationTool.Tasks
 
 					StringBuilder CreateProjectCommandline = new StringBuilder();
 					CreateProjectCommandline.AppendFormat("project-create --hosturl {0} {1}", Parameters.DestinationZenHost, ProjectName);
-					CommandUtils.RunAndLog(CommandUtils.CmdEnv, ZenExe.FullName, CreateProjectCommandline.ToString(), MaxSuccessCode: int.MaxValue, Options: CommandUtils.ERunOptions.Default, SpewFilterCallback: SilentOutputFilter);
+					TryRunAndLogWithoutSpew(ZenExe.FullName, CreateProjectCommandline.ToString());
 
 					OplogExportCommandline.AppendFormat(" --zen {0}", Parameters.DestinationZenHost);
 
@@ -576,7 +587,7 @@ namespace AutomationTool.Tasks
 						string DestinationOplog = SanitizeOplogName(ExportNames[ExportIndex]);
 
 						ExportSingleSourceCommandline.AppendFormat(" {0} --target-project {1} --target-oplog {2} {3} {4}", HostUrlArg, ProjectName, DestinationOplog, ExportSource.ProjectId, ExportSource.OplogId);
-						CommandUtils.RunAndLog(CommandUtils.CmdEnv, ZenExe.FullName, ExportSingleSourceCommandline.ToString(), MaxSuccessCode: int.MaxValue, Options: CommandUtils.ERunOptions.Default, SpewFilterCallback: SilentOutputFilter);
+						TryRunAndLogWithoutSpew(ZenExe.FullName, ExportSingleSourceCommandline.ToString());
 
 						ExportIndex = ExportIndex + 1;
 					}
@@ -624,7 +635,7 @@ namespace AutomationTool.Tasks
 						}
 						ExportSingleSourceCommandline.AppendFormat(" --file {0} --name {1} {2} {3} {4}", CommandUtils.MakePathSafeToUseWithCommandLine(PlatformDestinationFileDir.FullName), DestinationFileName, BaseNameArg, ProjectId, ExportSource.OplogId);
 
-						CommandUtils.RunAndLog(CommandUtils.CmdEnv, ZenExe.FullName, ExportSingleSourceCommandline.ToString(), Options: CommandUtils.ERunOptions.Default);
+						TryRunAndLogWithoutSpew(ZenExe.FullName, ExportSingleSourceCommandline.ToString());
 
 						ExportIndex = ExportIndex + 1;
 					}
