Treat "nopak" staging with use of zenstore as a equivalent to enabling the "fileserver" commandline argument during stage time alone.

diff --git a/Engine/Source/Programs/AutomationTool/Scripts/CopyBuildToStagingDirectory.Automation.cs b/Engine/Source/Programs/AutomationTool/Scripts/CopyBuildToStagingDirectory.Automation.cs
--- a/Engine/Source/Programs/AutomationTool/Scripts/CopyBuildToStagingDirectory.Automation.cs
+++ b/Engine/Source/Programs/AutomationTool/Scripts/CopyBuildToStagingDirectory.Automation.cs
@@ -474,6 +474,17 @@ namespace AutomationScripts
 			return InternationalizationPreset;
 		}
 
+		private static bool ShouldTreatAsFileServer(ProjectParams Params, DeploymentContext SC)
+		{
+			if (Params.FileServer)
+			{
+				return true;
+			}
+			FileReference ZenStoreMarkerFile = FindZenProjectStoreMarker(Params, SC);
+			bool UseZenServerHost = (ZenStoreMarkerFile != null) && !Params.UsePak(SC.StageTargetPlatform);
+			return UseZenServerHost;
+		}
+
 		public static List<string> GetCulturesToStage(ProjectParams Params, ConfigHierarchy PlatformGameConfig, bool bMustExist = true)
 		{
 			// Initialize cultures to stage.
@@ -1405,7 +1416,7 @@ namespace AutomationScripts
 						}
 
 						// UFS is required when using a file server
-						StagedFileType MovieFileType = Params.FileServer ? StagedFileType.UFS : StagedFileType.NonUFS;
+						StagedFileType MovieFileType = ShouldTreatAsFileServer(Params, SC) ? StagedFileType.UFS : StagedFileType.NonUFS;
 
 						DirectoryReference EngineMoviesDir = DirectoryReference.Combine(SC.EngineRoot, "Content", "Movies");
 						if (DirectoryReference.Exists(EngineMoviesDir))
@@ -2251,7 +2262,7 @@ namespace AutomationScripts
 			CopyManifestFilesToStageDir(SC, SC.FilesToStage.NonUFSFiles, "NonUFSFiles");
 
 			Dictionary<StagedFileReference, FileReference> UFSFiles = new Dictionary<StagedFileReference, FileReference>(SC.FilesToStage.UFSFiles);
-			bool bStageUnrealFileSystemFiles = !Params.CookOnTheFly && !Params.UsePak(SC.StageTargetPlatform) && !Params.FileServer;
+			bool bStageUnrealFileSystemFiles = !Params.CookOnTheFly && !Params.UsePak(SC.StageTargetPlatform) && !ShouldTreatAsFileServer(Params, SC);
 			if (bStageUnrealFileSystemFiles)
 			{
 				foreach (KeyValuePair<StagedFileReference, FileReference> Pair in SC.CrashReporterUFSFiles)
