Adding retry mechanism when exporting zen oplogs.  Also ensuring that we don't write out a descriptor if we failed to export the oplog.

diff --git a/Engine/Source/Programs/AutomationTool/BuildGraph/Tasks/ZenExportSnapshotTask.cs b/Engine/Source/Programs/AutomationTool/BuildGraph/Tasks/ZenExportSnapshotTask.cs
--- a/Engine/Source/Programs/AutomationTool/BuildGraph/Tasks/ZenExportSnapshotTask.cs
+++ b/Engine/Source/Programs/AutomationTool/BuildGraph/Tasks/ZenExportSnapshotTask.cs
@@ -382,7 +382,7 @@ namespace AutomationTool.Tasks
 			Writer.WriteObjectEnd();
 		}
 
-		private void TryRunAndLogWithoutSpew(string App, string CommandLine)
+		private bool TryRunAndLogWithoutSpew(string App, string CommandLine)
 		{
 			ProcessResult.SpewFilterCallbackType SilentOutputFilter = new ProcessResult.SpewFilterCallbackType(Line =>
 				{
@@ -395,7 +395,26 @@ namespace AutomationTool.Tasks
 			catch (CommandUtils.CommandFailedException e)
 			{
 				Logger.LogWarning(e.ToString());
+				return false;
+			}
+			return true;
+		}
+
+		private bool TryExportOplogCommand(string App, string CommandLine)
+		{
+			int AttemptLimit = 2;
+			int Attempt = 0;
+			while (Attempt < AttemptLimit)
+			{
+				if (TryRunAndLogWithoutSpew(App, CommandLine))
+				{
+					return true;
+				}
+				Logger.LogWarning("Attempt {0} of exporting the oplog failed, {1}...", Attempt+1, Attempt < (AttemptLimit-1) ? "retrying" : "abandoning");
+
+				Attempt = Attempt + 1;
 			}
+			return false;
 		}
 
 		/// <summary>
@@ -473,6 +492,7 @@ namespace AutomationTool.Tasks
 			}
 			int ExportIndex = 0;
 			string[] ExportNames = new string[ExportSources.Count];
+			List<ExportSourceData> SuccessfullyExportedSources = new List<ExportSourceData>();
 
 			// Get the Zen executable path
 			FileReference ZenExe = ZenExeFileReference();
@@ -549,7 +569,10 @@ namespace AutomationTool.Tasks
 						IoHash DestinationKeyHash = IoHash.Compute(Encoding.UTF8.GetBytes(ExportNames[ExportIndex]));
 
 						ExportSingleSourceCommandline.AppendFormat(" {0} --key {1} {2} {3} {4}", HostUrlArg, DestinationKeyHash.ToString().ToLowerInvariant(), BaseKeyArg, ExportSource.ProjectId, ExportSource.OplogId);
-						TryRunAndLogWithoutSpew(ZenExe.FullName, ExportSingleSourceCommandline.ToString());
+						if (TryExportOplogCommand(ZenExe.FullName, ExportSingleSourceCommandline.ToString()))
+						{
+							SuccessfullyExportedSources.Add(ExportSource);
+						}
 
 						ExportIndex = ExportIndex + 1;
 					}
@@ -587,7 +610,10 @@ namespace AutomationTool.Tasks
 						string DestinationOplog = SanitizeOplogName(ExportNames[ExportIndex]);
 
 						ExportSingleSourceCommandline.AppendFormat(" {0} --target-project {1} --target-oplog {2} {3} {4}", HostUrlArg, ProjectName, DestinationOplog, ExportSource.ProjectId, ExportSource.OplogId);
-						TryRunAndLogWithoutSpew(ZenExe.FullName, ExportSingleSourceCommandline.ToString());
+						if (TryExportOplogCommand(ZenExe.FullName, ExportSingleSourceCommandline.ToString()))
+						{
+							SuccessfullyExportedSources.Add(ExportSource);
+						}
 
 						ExportIndex = ExportIndex + 1;
 					}
@@ -635,7 +661,10 @@ namespace AutomationTool.Tasks
 						}
 						ExportSingleSourceCommandline.AppendFormat(" --file {0} --name {1} {2} {3} {4}", CommandUtils.MakePathSafeToUseWithCommandLine(PlatformDestinationFileDir.FullName), DestinationFileName, BaseNameArg, ProjectId, ExportSource.OplogId);
 
-						TryRunAndLogWithoutSpew(ZenExe.FullName, ExportSingleSourceCommandline.ToString());
+						if (TryExportOplogCommand(ZenExe.FullName, ExportSingleSourceCommandline.ToString()))
+						{
+							SuccessfullyExportedSources.Add(ExportSource);
+						}
 
 						ExportIndex = ExportIndex + 1;
 					}
@@ -645,13 +674,13 @@ namespace AutomationTool.Tasks
 			}
 			
 
-			if ((Parameters.SnapshotDescriptorFile != null) && ExportSources.Any())
+			if ((Parameters.SnapshotDescriptorFile != null) && SuccessfullyExportedSources.Any())
 			{
 				if (Parameters.SnapshotDescriptorFile.FullName.Contains("{Platform}"))
 				{
 					// Separate descriptor file per platform
 					ExportIndex = 0;
-					foreach (ExportSourceData ExportSource in ExportSources)
+					foreach (ExportSourceData ExportSource in SuccessfullyExportedSources)
 					{
 						FileReference PlatformSnapshotDescriptorFile = new FileReference(Parameters.SnapshotDescriptorFile.FullName.Replace("{Platform}", ExportSource.TargetPlatform, StringComparison.InvariantCultureIgnoreCase));
 						DirectoryReference.CreateDirectory(PlatformSnapshotDescriptorFile.Directory);
@@ -676,7 +705,7 @@ namespace AutomationTool.Tasks
 						Writer.WriteArrayStart("snapshots");
 								
 						ExportIndex = 0;
-						foreach (ExportSourceData ExportSource in ExportSources)
+						foreach (ExportSourceData ExportSource in SuccessfullyExportedSources)
 						{
 							WriteExportSource(Writer, DestinationStorageType, ExportSource, ExportNames[ExportIndex]);
 							ExportIndex = ExportIndex + 1;
