Detect invalid slashes in markdown links.

diff --git a/Engine/Source/Programs/AutomationTool/BuildGraph/Tasks/CheckMarkdownTask.cs b/Engine/Source/Programs/AutomationTool/BuildGraph/Tasks/CheckMarkdownTask.cs
--- a/Engine/Source/Programs/AutomationTool/BuildGraph/Tasks/CheckMarkdownTask.cs
+++ b/Engine/Source/Programs/AutomationTool/BuildGraph/Tasks/CheckMarkdownTask.cs
@@ -96,18 +96,21 @@ namespace AutomationTool.Tasks
 						string link = match.Groups[2].Value;
 						if (!Regex.IsMatch(link, "^(?:[a-z]+:/)?/"))
 						{
-							int hashIdx = link.IndexOf('#');
-							if (hashIdx == -1)
+							if (!link.Contains('\\'))
 							{
-								link = FileReference.Combine(file.Directory, link).FullName;
-							}
-							else if (hashIdx == 0)
-							{
-								link = $"{file}{link[hashIdx..]}";
-							}
-							else
-							{
-								link = $"{FileReference.Combine(file.Directory, link[0..hashIdx])}{link[hashIdx..]}";
+								int hashIdx = link.IndexOf('#');
+								if (hashIdx == -1)
+								{
+									link = FileReference.Combine(file.Directory, link).FullName;
+								}
+								else if (hashIdx == 0)
+								{
+									link = $"{file}{link[hashIdx..]}";
+								}
+								else
+								{
+									link = $"{FileReference.Combine(file.Directory, link[0..hashIdx])}{link[hashIdx..]}";
+								}
 							}
 
 							bool validLink;
