[Gauntlet] Makes UnrealHelpers.UnrealConfigFromFileName public

diff --git a/Engine/Source/Programs/AutomationTool/Gauntlet/Unreal/Utils/Gauntlet.UnrealHelpers.cs b/Engine/Source/Programs/AutomationTool/Gauntlet/Unreal/Utils/Gauntlet.UnrealHelpers.cs
--- a/Engine/Source/Programs/AutomationTool/Gauntlet/Unreal/Utils/Gauntlet.UnrealHelpers.cs
+++ b/Engine/Source/Programs/AutomationTool/Gauntlet/Unreal/Utils/Gauntlet.UnrealHelpers.cs
@@ -285,7 +285,7 @@ namespace Gauntlet
 			return ExeName;
 		}
 
-		internal class ConfigInfo
+		public class ConfigInfo
 		{
 			public UnrealTargetRole 			RoleType;
 			public UnrealTargetPlatform? 		Platform;
@@ -314,8 +314,8 @@ namespace Gauntlet
 		{
 			CustomModuleToRoles.Add(InModuleName, InRole);
 		}
-		
-		static ConfigInfo GetUnrealConfigFromFileName(string InProjectName, string InName)
+
+		static public ConfigInfo GetUnrealConfigFromFileName(string InProjectName, string InName)
 		{
 			ConfigInfo Config = new ConfigInfo();
 
