- Updated packaging with -distribution, to take use -archivedirectory param (if it ends with .xcarchive) as an override location to write the .xcarchive, instead of ~/Library/Developer/Xcode (when -archive isn't set)
- Similarly if -archive and -archivedirectory ends with .xcarchive, then that is the name of the directory to archive to, instead of using Xcode's name
- Improved logging of xcarchive info

diff --git a/Engine/Source/Programs/AutomationTool/Apple/ApplePlatform.Automation.cs b/Engine/Source/Programs/AutomationTool/Apple/ApplePlatform.Automation.cs
--- a/Engine/Source/Programs/AutomationTool/Apple/ApplePlatform.Automation.cs
+++ b/Engine/Source/Programs/AutomationTool/Apple/ApplePlatform.Automation.cs
@@ -10,6 +10,7 @@ using UnrealBuildBase;
 using System.Diagnostics;
 using System.Collections.Generic;
 using Microsoft.Extensions.Logging;
+using System.Xml.Linq;
 
 public abstract class ApplePlatform : Platform
 {
@@ -157,7 +158,7 @@ public abstract class ApplePlatform : Platform
 		// find the most recent archive for this target (based on name of target, this ignores Development vs Shipping, but 
 		// since Distribution is meant only for Shipping it's ok
 		string ArchiveName = AppleExports.MakeBinaryFileName(TargetName, Target.Receipt.Platform, UnrealTargetConfiguration.Development,
-		 	Target.Receipt.Architectures, UnrealTargetConfiguration.Development, null) + " *.xcarchive";
+		 	Target.Receipt.Architectures, UnrealTargetConfiguration.Development, null);
 
 		if (ArchiveName != TargetName)
 		{
@@ -228,6 +229,15 @@ public abstract class ApplePlatform : Platform
 					}
 				}
 
+				// if we aren't going to archive right after this to another location, and then archivedirectory specifies a .xcarchive extension
+				// then we want to archive directly to that location instead of the standard ~/Library/Developer/Xcode/Archives. This is useful for
+				// build machines that want to make a .xcarchive and export/upload it, without needing to hunt down latest, or dirty up the Library dir
+				string XcArchivePath = "";
+				if (Params.Distribution && Params.Archive == false && Params.ArchiveDirectoryParam.EndsWith(".xcarchive"))
+				{
+					XcArchivePath = Params.ArchiveDirectoryParam;
+					ExtraOptions += $" -archivePath \"{XcArchivePath}\"";
+				}
 
 				// if we we packaging for distrbution, we will create a .xcarchive which can be used to submit to app stores, or exported for other distribution methods
 				// the archive will be created in the standard Archives location accessible via Xcode. Using -archive will copy it out into
@@ -235,17 +245,34 @@ public abstract class ApplePlatform : Platform
 				AppleExports.XcodeBuildMode BuildMode = Params.Distribution ? AppleExports.XcodeBuildMode.Distribute : AppleExports.XcodeBuildMode.Package;
 				if (AppleExports.BuildWithStubXcodeProject(Params.RawProjectPath, Receipt.Platform, Receipt.Architectures, Receipt.Configuration, TargetName, BuildMode, Logger, ExtraOptions) == 0)
 				{
-					Logger.LogInformation("=====================================================================================");
+					Logger.LogInformation("==================================================================================================================================================");
 					if (Params.Distribution)
 					{
-						Logger.LogInformation("Created .xcarchive in Xcode's Library, which can be seen in Xcode's Organizer window");
-						Logger.LogInformation("You may use this to validate and prepare for various distribution methods");
+						if (XcArchivePath == "")
+						{
+							XcArchivePath = AppleExports.FindLatestXcArchive(TargetName)?.FullName;
+							if (XcArchivePath == null)
+							{
+								Logger.LogError("Unable to discover a xcarchive for {TargetName} in ~/Library/Developer/Xcode/Archives.", TargetName);
+								Logger.LogError("This is unexpected since the command succeeded");
+							}
+							else
+							{
+								Logger.LogInformation("Created '{XcArchivePath}', which can be seen in Xcode's Organizer window.", XcArchivePath);
+								Logger.LogInformation("You may use this to validate and prepare for various distribution methods");
+							}
+						}
+						else
+						{
+							Logger.LogInformation("Created '{XcArchivePath}' using an overriden location - this will _not_ appear in Xcode's Organizer window.", XcArchivePath);
+							Logger.LogInformation("You may use this to validate and prepare for various distribution methods");
+						}
 					}
 					else
 					{
 						Logger.LogInformation("Finalized {App} for running fully self-contained", GetFinalAppPath(Target, SC));
 					}
-					Logger.LogInformation("=====================================================================================");
+					Logger.LogInformation("==================================================================================================================================================");
 				}
 			}
 		}
@@ -272,6 +299,7 @@ public abstract class ApplePlatform : Platform
 
 			// distribution mode we want to archive the .xcarchive that was created during Package
 			DirectoryReference ArchiveSource;
+			DirectoryReference ArchiveDestination = null;
 			if (Params.Distribution)
 			{
 				// find the most recent .xcarchive in the Xcode archives library
@@ -282,6 +310,15 @@ public abstract class ApplePlatform : Platform
 					Logger.LogError("Unable to find a .xcarchive in Xcode's Library to archive to {ArchiveDir}", SC.ArchiveDirectory);
 					return;
 				}
+
+				if (SC.ArchiveDirectory.HasExtension(".xcarchive"))
+				{
+					ArchiveDestination = SC.ArchiveDirectory;
+					if (DirectoryReference.Exists(ArchiveDestination))
+					{
+						DirectoryReference.Delete(ArchiveDestination);
+					}
+				}
 			}
 			else
 			{
@@ -293,11 +330,16 @@ public abstract class ApplePlatform : Platform
 				}
 			}
 
-			Logger.LogInformation("=====================================================================================");
-			Logger.LogInformation("Copying {Type} package {ArchiveSource} to archive directory {ArchiveDir}", Params.Distribution ? "Distribution" : "Development", ArchiveSource, SC.ArchiveDirectory);
-			Logger.LogInformation("=====================================================================================");
+			if (ArchiveDestination == null)
+			{
+				ArchiveDestination = DirectoryReference.Combine(SC.ArchiveDirectory, ArchiveSource.GetDirectoryName());
+			}
+			
+			Logger.LogInformation("==================================================================================================================================================");
+			Logger.LogInformation("Copying {Type} package {ArchiveSource} to archive directory {ArchiveDir}", Params.Distribution ? "Distribution" : "Development", ArchiveSource, ArchiveDestination);
+			Logger.LogInformation("==================================================================================================================================================");
 
-			Utils.RunLocalProcessAndReturnStdOut("/usr/bin/env", $"ditto \"{ArchiveSource}\" \"{SC.ArchiveDirectory}/{ArchiveSource.GetDirectoryName()}\"", null);
+			Utils.RunLocalProcessAndReturnStdOut("/usr/bin/env", $"ditto \"{ArchiveSource}\" \"{ArchiveDestination}\"", null);
 		}
 	}
 }
diff --git a/Engine/Source/Programs/AutomationTool/Apple/ExportAppleBuild.Automation.cs b/Engine/Source/Programs/AutomationTool/Apple/ExportAppleBuild.Automation.cs
--- a/Engine/Source/Programs/AutomationTool/Apple/ExportAppleBuild.Automation.cs
+++ b/Engine/Source/Programs/AutomationTool/Apple/ExportAppleBuild.Automation.cs
@@ -25,13 +25,19 @@ class ExportAppleBuild : BuildCommand
 		string PlistParam = ParseParamValue("Plist");
 		string XcArchivePathParam = ParseParamValue("XcArchive");
 		string LatestXcArchiveParam = ParseParamValue("LatestXcArchiveForTarget");
+		string LatestXcArchiveSearch = ParseParamValue("LatestXcArchiveSearchPath");
 		string ExportPath = ParseParamValue("ExportPath");
 
 		DirectoryReference XcArchive;
 		if (!string.IsNullOrEmpty(LatestXcArchiveParam))
 		{
 			// find the latest archive for the given project name
-			XcArchive = AppleExports.FindLatestXcArchive(LatestXcArchiveParam);
+			DirectoryReference SearchPath = null;
+			if (!string.IsNullOrEmpty(LatestXcArchiveSearch))
+			{
+				SearchPath = new DirectoryReference(LatestXcArchiveSearch);
+			}
+			XcArchive = AppleExports.FindLatestXcArchive(LatestXcArchiveParam, SearchPath);
 		}
 		else
 		{
@@ -40,7 +46,7 @@ class ExportAppleBuild : BuildCommand
 
 		if (XcArchive == null || !DirectoryReference.Exists(XcArchive))
 		{
-			Logger.LogError("No XCArchive found, with -XCArchive or LatestXCArchiveForTarget params");
+			Logger.LogError("No XCArchive found, with -XCArchive or -LatestXCArchiveForTarget params");
 			return ExitCode.Error_Arguments;
 		}
 
@@ -83,11 +89,6 @@ class ExportAppleBuild : BuildCommand
 		int Return;
 		Utils.RunLocalProcessAndReturnStdOut("/usr/bin/xcodebuild", CommandLine, Logger, out Return);
 
-// xcodebuild -exportArchive -archivePath "[REDACTED_USER_HOME]/Library/Developer/Xcode/Archives/2024-04-06/LyraGame 4-6-24, 2.35 PM.xcarchive" 
-// -exportPath ./private_keys -exportOptionsPlist Engine/Build/Xcode/ExportOptions_UploadTestFlight.plist 
-// -authenticationKeyIssuerID 69a6de71-5e09-47e3-e053-5b8c7c11a4d1 -authenticationKeyPath $PWD/private_keys/AuthKey_HCTZ2FPZ2Y.p8 
-// -authenticationKeyID HCTZ2FPZ2Y
-
 		return (ExitCode)Return;
 	}
 }
diff --git a/Engine/Source/Programs/UnrealBuildTool/Platform/Apple/AppleExports.cs b/Engine/Source/Programs/UnrealBuildTool/Platform/Apple/AppleExports.cs
--- a/Engine/Source/Programs/UnrealBuildTool/Platform/Apple/AppleExports.cs
+++ b/Engine/Source/Programs/UnrealBuildTool/Platform/Apple/AppleExports.cs
@@ -295,12 +295,13 @@ namespace UnrealBuildTool
 		/// <summary>
 		/// FInds the latest .xcarchive for a given Target name (the .xcarchive will start with this name then have a data appended)
 		/// </summary>
-		/// <param name="TargetName"></param>
+		/// <param name="TargetName">Name of the target to look for, this will be the prefix for the .xcarchive to search for</param>
+		/// <param name="SearchPath">If null, look in Xcode Archives library, otherwiwse, look here</param>
 		/// <returns></returns>
-		public static DirectoryReference? FindLatestXcArchive(string TargetName)
+		public static DirectoryReference? FindLatestXcArchive(string TargetName, DirectoryReference? SearchPath=null)
 		{
 			DirectoryReference UserDir = new DirectoryReference(Environment.GetFolderPath(Environment.SpecialFolder.UserProfile));
-			DirectoryReference Library = DirectoryReference.Combine(UserDir, "Library/Developer/Xcode/Archives");
+			DirectoryReference Library = SearchPath ?? DirectoryReference.Combine(UserDir, "Library/Developer/Xcode/Archives");
 
 			// order date named folders (use creating data, not name, but same thing)
 			List<DirectoryReference> DateDirs = DirectoryReference.EnumerateDirectories(Library).ToList();
@@ -308,10 +309,9 @@ namespace UnrealBuildTool
 			DateDirs.Reverse();
 
 			// go through each folder, starting at most recent, looking for an archive for the target
+			string Wildcard = $"{TargetName} *.xcarchive";
 			foreach (DirectoryReference DateDir in DateDirs)
 			{
-				string Wildcard = $"{TargetName} *.xcarchive";
-
 				List<DirectoryReference> XcArchives = DirectoryReference.EnumerateDirectories(DateDir, Wildcard).ToList();
 				if (XcArchives.Count > 0)
 				{
