TempStorage Clean up - also include directory timestamps

#rnx

diff --git a/Engine/Source/Programs/AutomationTool/BuildGraph/TempStorage.cs b/Engine/Source/Programs/AutomationTool/BuildGraph/TempStorage.cs
--- a/Engine/Source/Programs/AutomationTool/BuildGraph/TempStorage.cs
+++ b/Engine/Source/Programs/AutomationTool/BuildGraph/TempStorage.cs
@@ -1608,14 +1608,15 @@ namespace AutomationTool
 			List<DirectoryInfo> BuildsToDelete = new List<DirectoryInfo>();
 			foreach (DirectoryInfo StreamDirectory in new DirectoryInfo(TempStorageDir).EnumerateDirectories().OrderBy(x => x.Name))
 			{
-				Logger.LogInformation("Scanning {Arg0}...", StreamDirectory.FullName);
+				Logger.LogInformation("Scanning {Path}...", StreamDirectory.FullName);
 				try
 				{
 					foreach (DirectoryInfo BuildDirectory in StreamDirectory.EnumerateDirectories())
 					{
 						try
 						{
-							if (!BuildDirectory.EnumerateFiles("*", SearchOption.AllDirectories).Any(x => x.LastWriteTimeUtc > RetainTime))
+							if (!BuildDirectory.EnumerateFiles("*", SearchOption.AllDirectories).Any(x => x.LastWriteTimeUtc > RetainTime) &&
+								!BuildDirectory.EnumerateDirectories("*", SearchOption.AllDirectories).Any(x => x.LastWriteTimeUtc > RetainTime))
 							{
 								BuildsToDelete.Add(BuildDirectory);
 							}
@@ -1632,7 +1633,7 @@ namespace AutomationTool
 					Logger.LogError(Ex, "Exception while trying to scan {StreamDirectory}: {Ex}", StreamDirectory, Ex);
 				}
 			}
-			Logger.LogInformation("Found {NumBuilds} builds; {Arg1} to delete.", NumBuilds, BuildsToDelete.Count);
+			Logger.LogInformation("Found {NumBuilds} builds; {Count} to delete.", NumBuilds, BuildsToDelete.Count);
 
 			// Loop through them all, checking for files older than the delete time
 			int Idx = BuildsToDelete.Count;
@@ -1650,12 +1651,12 @@ namespace AutomationTool
 					FileInfo DeleteInProgressFile = BuildsToDelete[Idx].GetFiles("DeleteInProgress.tmp").FirstOrDefault();
 					if (DeleteInProgressFile != null && DeleteInProgressFile.LastWriteTimeUtc < (DateTimeOffset.UtcNow - TimeSpan.FromMinutes(20)))
 					{
-						Logger.LogInformation("[{Arg0}/{Arg1}] {Arg2} flagged as delete in progress, skipping...", BuildsToDelete.Count - Idx, BuildsToDelete.Count, BuildsToDelete[Idx].FullName);
+						Logger.LogInformation("[{Index}/{Total}] {Path} flagged as delete in progress, skipping...", BuildsToDelete.Count - Idx, BuildsToDelete.Count, BuildsToDelete[Idx].FullName);
 						continue;
 					}
 
 					File.WriteAllBytes(Path.Combine(BuildsToDelete[Idx].FullName, "DeleteInProgress.tmp"), Array.Empty<byte>());
-					Logger.LogInformation("[{Arg0}/{Arg1}] Deleting {Arg2}...", BuildsToDelete.Count - Idx, BuildsToDelete.Count, BuildsToDelete[Idx].FullName);
+					Logger.LogInformation("[{Index}/{Total}] Deleting {Path}...", BuildsToDelete.Count - Idx, BuildsToDelete.Count, BuildsToDelete[Idx].FullName);
 					BuildsToDelete[Idx].Delete(true);
 				}
 				catch (Exception Ex)
@@ -1682,7 +1683,7 @@ namespace AutomationTool
 						}
 						catch (Exception Ex)
 						{
-							Logger.LogWarning("Unexpected failure trying to delete (potentially empty) stream directory {Arg0}: {Ex}", StreamDirectory.FullName, Ex);
+							Logger.LogWarning("Unexpected failure trying to delete (potentially empty) stream directory {Path}: {Ex}", StreamDirectory.FullName, Ex);
 						}
 					}
 				}
