Localization:
Refactoring to the help text for LocalizationTargetEditor commands.
- Also fixed a small bug with parameter parsing for NewPluginLocalizationTarget command

#rnx

diff --git a/Engine/Source/Programs/AutomationTool/Localization/LocalizationTargetEditor/Commands/LocalizationTargetCommand.cs b/Engine/Source/Programs/AutomationTool/Localization/LocalizationTargetEditor/Commands/LocalizationTargetCommand.cs
--- a/Engine/Source/Programs/AutomationTool/Localization/LocalizationTargetEditor/Commands/LocalizationTargetCommand.cs
+++ b/Engine/Source/Programs/AutomationTool/Localization/LocalizationTargetEditor/Commands/LocalizationTargetCommand.cs
@@ -3,10 +3,9 @@ using AutomationTool;
 using Microsoft.Extensions.Logging;
 using System;
 using System.Collections.Generic;
-using System.IO;
 using System.Linq;
 using System.Reflection;
-using UnrealBuildTool;
+using System.Text;
 
 using static AutomationTool.CommandUtils;
 
@@ -51,7 +50,6 @@ namespace EpicGames.Localization
 			return command;
 		}
 
-
 		protected virtual bool ParseCommandLine()
 		{
 			bIsExecutingInPreview = _commandLineHelper.ParseParam("Preview");
@@ -68,7 +66,40 @@ namespace EpicGames.Localization
 
 		public virtual string GetHelpText()
 		{
-			return "";
+			StringBuilder helpTextStringBuilder = new StringBuilder();
+			BuildHelpDescription(helpTextStringBuilder);
+			BuildHelpArguments(helpTextStringBuilder);
+			BuildHelpUsage(helpTextStringBuilder);
+			return helpTextStringBuilder.ToString();
+		}
+
+		/// <summary>
+		/// Buils the description portion of the help text. Concrete child classes of LocalizationTargetCommand should override this function 
+		/// </summary> and call the base implementation to build any descriptions from parent classes.
+		/// <param name="builder"></param>
+		protected virtual void BuildHelpDescription(StringBuilder builder)
+		{
+			builder.AppendLine("Description:");
+		}
+
+		/// <summary>
+		/// Builds the arguments portion of the help text explaining what each parameter of the command does. Child classes should override this function
+		/// </summary> to call the base implementation as well as provide descriptions for the child arguments.
+		/// <param name="builder"></param>
+		protected virtual void BuildHelpArguments(StringBuilder builder)
+		{
+			builder.AppendLine("Arguments:");
+			builder.AppendLine("Preview - An optional flag that will execute this command in preview mode. No files will be created or edited. No folders will be created. No files will be added or checked out of perforce.");
+		}
+
+		/// <summary>
+		/// Builds the usage portion of the help text. Useres should override this function and call teh base implementation 
+		/// and provide usage scenarios and the associated usage syntax for each scenario.
+		/// </summary>
+		/// <param name="builder"></param>
+		protected virtual void BuildHelpUsage(StringBuilder builder)
+		{
+			builder.AppendLine("Usage:");
 		}
 
 		public static List<LocalizationTargetCommand> GetAllCommands()
diff --git a/Engine/Source/Programs/AutomationTool/Localization/LocalizationTargetEditor/Commands/NewPluginLocalizationTargetCommand.cs b/Engine/Source/Programs/AutomationTool/Localization/LocalizationTargetEditor/Commands/NewPluginLocalizationTargetCommand.cs
--- a/Engine/Source/Programs/AutomationTool/Localization/LocalizationTargetEditor/Commands/NewPluginLocalizationTargetCommand.cs
+++ b/Engine/Source/Programs/AutomationTool/Localization/LocalizationTargetEditor/Commands/NewPluginLocalizationTargetCommand.cs
@@ -27,29 +27,30 @@ namespace EpicGames.Localization
 			DisplayName = "New Plugin Localization Target";
 		}
 
-		public override string GetHelpText()
+		protected override void BuildHelpDescription(StringBuilder builder)
 		{
-			StringBuilder helpTextBuilder = new StringBuilder();
-			helpTextBuilder.AppendLine("Creates new localization targets for plugins. This command generates the necessary localization config files for a plugin and updates the plugin descriptor with new localization target data.");
-			helpTextBuilder.AppendLine("Passing the -p4 flag to UAT will create a change list with all of the added and edited files.");
-			helpTextBuilder.AppendLine("If the localization config files already exist for the plugin or if the plugin descriptor already contains the localization target, this command will not edit any of the existing files.");
+			base.BuildHelpDescription(builder);
+			builder.AppendLine("Creates new localization targets for plugins. This command generates the necessary localization config files for a plugin and updates the plugin descriptor with new localization target data.");
+			builder.AppendLine("Passing the -p4 flag to UAT will create a change list with all of the added and edited files.");
+			builder.AppendLine("If the localization config files already exist for the plugin or if the plugin descriptor already contains the localization target, this command will not edit any of the existing files.");
+		}
 
-			helpTextBuilder.AppendLine("Arguments:");
-			helpTextBuilder.AppendLine("UERootDirectory - Optional root directory for the Unreal Engine. This is usually one level above your project directory. Defaults to CmdEnv.LocalRoot");
-			helpTextBuilder.AppendLine("UEProjectDirectory - Required relative path from UERootDirectory to the project. This should be  your project directory that contains the Plugins directory.");
-			helpTextBuilder.AppendLine("UEProjectName - An optional name of the project the plugin is under. If blank, this implies the plugin is under Engine.");
-			helpTextBuilder.AppendLine("LocalizationTargetLoadingPolicy - A required string that determines when the plugin should load the localization data for the localization target. Valid values are Game, Always, Editor, PropertyNames, ToolTips. See LocalizationTargetDescriptorLoadingPolicy");
-			helpTextBuilder.AppendLine("LocalizationConfigGenerationPolicy- An optional string that specifies how localization config files should be generated during the localization gather process for the plugin. Acceptable values are Never, Auto, User. If not specified, defaults to Auto.");
-			helpTextBuilder.AppendLine("Never means no localization config files will be generated or used during the localization process. No localization data will be generated for the plugin druing a gather. Auto means temporary, default localization config files will be generated druing localization gather and used to generate localization data. User means there are user provided localization config files in the plugins's Config/Localization folder that will be used during localization gathers to generate localization data.");
-			helpTextBuilder.AppendLine("LocalizationTargetNameSuffix - An optional suffix to give to the plugin localization target. By default, the name of the plugin localization target would be the name of the plugin. This allows multiple localization targets to be created for plugins.");
-			// Include plugins 
-			helpTextBuilder.AppendLine("IncludePlugins - An optional comma separated list of plugins to create localization targets for. E.g PluginA,PluginB,PluginC");
-			helpTextBuilder.AppendLine("IncludePluginsDirectory - An optional relative directory to UEProjectDirectory. All plugins under this directory will have localization targets created if they are not excluded. E.g Plugins/PluginFolderA.");
-			helpTextBuilder.AppendLine("ExcludePlugins - A comma separated list of plugins to exclude from having localization targets created. E.g PluginA,BpluginB,PluginC");
-			helpTextBuilder.AppendLine("ExcludePluginsDirectory - An optional relative directory from UEProjectDirectory. All plugins under this directory will be excluded from localization target creation. E.g Plugin/DirectoryToExclude");
+		protected override void BuildHelpArguments(StringBuilder builder)
+		{
+			base.BuildHelpArguments(builder);
+			builder.AppendLine("LocalizationTargetLoadingPolicy - A required string that determines when the plugin should load the localization data for the localization target. Valid values are Game, Always, Editor, PropertyNames, ToolTips. See LocalizationTargetDescriptorLoadingPolicy");
+			builder.AppendLine("LocalizationConfigGenerationPolicy- An optional string that specifies how localization config files should be generated during the localization gather process for the plugin. Acceptable values are Never, Auto, User. If not specified, defaults to Auto.");
+			builder.AppendLine("Never means no localization config files will be generated or used during the localization process. No localization data will be generated for the plugin druing a gather. Auto means temporary, default localization config files will be generated druing localization gather and used to generate localization data. User means there are user provided localization config files in the plugins's Config/Localization folder that will be used during localization gathers to generate localization data.");
+			builder.AppendLine("LocalizationTargetNameSuffix - An optional suffix to give to the plugin localization target. By default, the name of the plugin localization target would be the name of the plugin. This allows multiple localization targets to be created for plugins.");
+		}
 
-			helpTextBuilder.AppendLine("Preview - An optional flag that will execute this command in preview mode. No files will be created or edited. No folders will be created. No files will be added or checked out of perforce.");
-			return helpTextBuilder.ToString();
+		protected override void BuildHelpUsage(StringBuilder builder)
+		{
+			base.BuildHelpUsage(builder);
+			builder.AppendLine("Creates a localization target for MyPlugin with a loading policy of Game and localization config generation policy of Auto and and checks out the necessary files in perforce.");
+			builder.AppendLine("RunUAT.bat LocalizationTargetEditor -Command=NewPluginLocalizationTarget -UEProjectDirectory=MyGame -IncludePlugins=MyPlugin -LocalizationTargetLoadingPolicy=Game -p4");
+			builder.AppendLine("Creates a new localization target for all plugins under MyDirectory with loading policy Game and config generaiton policy of User so that localization config files will be generated in each plugin's Config/Localization folder. The relevant files will be checked out of perforce.");
+			builder.AppendLine("RunUAT.bat LocalizationTargetEditor -Command=NewPluginLocalizationTarget -UEProjectDirectory=MyGame -IncludePluginsDirectory=Plugins/MyDirectory -LocalizationTargetLoadingPolicy=Game -LocalizationConfigGenerationPolicy=User -p4");
 		}
 
 		protected override bool ParseCommandLine()
@@ -58,7 +59,7 @@ namespace EpicGames.Localization
 			{
 				return false;
 			}
-			_localizationTargetNameSuffix = _commandLineHelper.ParseParamValue("_localizationTargetNameSuffix");
+			_localizationTargetNameSuffix = _commandLineHelper.ParseParamValue("LocalizationTargetNameSuffix");
 			// @TODOLocalization: Support an override for the localization config file format 
 			_fileFormat = LocalizationConfigFileFormat.Latest;
 			_loadingPolicy = _commandLineHelper.ParseRequiredEnumParamEnum<LocalizationTargetDescriptorLoadingPolicy>("LocalizationTargetLoadingPolicy");
diff --git a/Engine/Source/Programs/AutomationTool/Localization/LocalizationTargetEditor/Commands/PluginLocalizationTargetCommand.cs b/Engine/Source/Programs/AutomationTool/Localization/LocalizationTargetEditor/Commands/PluginLocalizationTargetCommand.cs
--- a/Engine/Source/Programs/AutomationTool/Localization/LocalizationTargetEditor/Commands/PluginLocalizationTargetCommand.cs
+++ b/Engine/Source/Programs/AutomationTool/Localization/LocalizationTargetEditor/Commands/PluginLocalizationTargetCommand.cs
@@ -4,8 +4,7 @@ using Microsoft.Extensions.Logging;
 using System;
 using System.Collections.Generic;
 using System.IO;
-using System.Linq;
-using System.Reflection;
+using System.Text;
 using UnrealBuildTool;
 
 using static AutomationTool.CommandUtils;
@@ -22,6 +21,19 @@ namespace EpicGames.Localization
 		public HashSet<string> IncludePlugins { get; protected set; } = new HashSet<string>();
 		public HashSet<string> ExcludePlugins { get; protected set; } = new HashSet<string>();
 
+		protected override void BuildHelpArguments(StringBuilder builder)
+		{
+			base.BuildHelpArguments(builder);
+			builder.AppendLine("UERootDirectory - Optional root directory for the Unreal Engine. This is usually one level above your project directory. Defaults to CmdEnv.LocalRoot");
+			builder.AppendLine("UEProjectDirectory - Required relative path from UERootDirectory to the project. This should be  your project directory that contains the Plugins directory.");
+			builder.AppendLine("UEProjectName - An optional name of the project the plugin is under. If blank, this implies the plugin is under Engine.");
+
+			// Include plugins 
+			builder.AppendLine("IncludePlugins - An optional comma separated list of plugins to operate on. E.g PluginA,PluginB,PluginC");
+			builder.AppendLine("IncludePluginsDirectory - An optional relative directory to UEProjectDirectory. All plugins under this directory will be operated on if they are not excluded. E.g Plugins/PluginFolderA.");
+			builder.AppendLine("ExcludePlugins - A comma separated list of plugins to exclude from being operated on. E.g PluginA,BpluginB,PluginC");
+			builder.AppendLine("ExcludePluginsDirectory - An optional relative directory from UEProjectDirectory. All plugins under this directory will be excluded from being operated on. E.g Plugin/DirectoryToExclude");
+		}
 
 		protected override bool ParseCommandLine()
 		{
