- Routed output diffrently because xcodebuild like to write to stderr, which will make horde return failures, and log it as red

diff --git a/Engine/Source/Programs/AutomationTool/Apple/ExportAppleBuild.Automation.cs b/Engine/Source/Programs/AutomationTool/Apple/ExportAppleBuild.Automation.cs
--- a/Engine/Source/Programs/AutomationTool/Apple/ExportAppleBuild.Automation.cs
+++ b/Engine/Source/Programs/AutomationTool/Apple/ExportAppleBuild.Automation.cs
@@ -87,7 +87,8 @@ class ExportAppleBuild : BuildCommand
 		Logger.LogInformation($"Running 'xcodebuild {CommandLine}'...");
 
 		int Return;
-		Utils.RunLocalProcessAndReturnStdOut("/usr/bin/xcodebuild", CommandLine, Logger, out Return);
+		string Output = Utils.RunLocalProcessAndReturnStdOut("/usr/bin/xcodebuild", CommandLine, null, out Return);
+		Logger.LogInformation(Output);
 
 		return (ExitCode)Return;
 	}
