BuildGraph: Allow specifying an IgnoreModified flag on node attributes to ignore mismatched timestamps/lengths of files matching a particular pattern (accepts p4-style wildcards and multiple rules separated by semicolons).

diff --git a/Engine/Source/Programs/AutomationTool/BuildGraph/BgSchema.cs b/Engine/Source/Programs/AutomationTool/BuildGraph/BgSchema.cs
--- a/Engine/Source/Programs/AutomationTool/BuildGraph/BgSchema.cs
+++ b/Engine/Source/Programs/AutomationTool/BuildGraph/BgSchema.cs
@@ -572,6 +572,7 @@ namespace AutomationTool
 			extension.Attributes.Add(CreateSchemaAttribute("RunEarly", ScriptSchemaStandardType.Boolean, XmlSchemaUse.Optional));
 			extension.Attributes.Add(CreateSchemaAttribute("NotifyOnWarnings", ScriptSchemaStandardType.Boolean, XmlSchemaUse.Optional));
 			extension.Attributes.Add(CreateSchemaAttribute("Annotations", ScriptSchemaStandardType.BalancedString, XmlSchemaUse.Optional));
+			extension.Attributes.Add(CreateSchemaAttribute("IgnoreModified", ScriptSchemaStandardType.BalancedString, XmlSchemaUse.Optional));
 
 			XmlSchemaComplexContent contentModel = new XmlSchemaComplexContent();
 			contentModel.Content = extension;
diff --git a/Engine/Source/Programs/AutomationTool/BuildGraph/BgScriptReader.cs b/Engine/Source/Programs/AutomationTool/BuildGraph/BgScriptReader.cs
--- a/Engine/Source/Programs/AutomationTool/BuildGraph/BgScriptReader.cs
+++ b/Engine/Source/Programs/AutomationTool/BuildGraph/BgScriptReader.cs
@@ -287,8 +287,8 @@ namespace AutomationTool
 		/// <summary>
 		/// Constructor
 		/// </summary>
-		public BgScriptNode(string name, IReadOnlyList<BgNodeOutput> inputs, IReadOnlyList<string> outputNames, IReadOnlyList<BgNodeDef> inputDependencies, IReadOnlyList<BgNodeDef> orderDependencies, IReadOnlyList<FileReference> requiredTokens)
-			: base(name, inputs, outputNames, inputDependencies, orderDependencies, requiredTokens)
+		public BgScriptNode(string name, IReadOnlyList<BgNodeOutput> inputs, IReadOnlyList<string> outputNames, IReadOnlyList<BgNodeDef> inputDependencies, IReadOnlyList<BgNodeDef> orderDependencies, IReadOnlyList<FileReference> requiredTokens, IReadOnlyList<string> ignoreModified)
+			: base(name, inputs, outputNames, inputDependencies, orderDependencies, requiredTokens, ignoreModified)
 		{
 		}
 	}
@@ -1368,6 +1368,7 @@ namespace AutomationTool
 				bool bRunEarly = ReadBooleanAttribute(element, "RunEarly", false);
 				bool bNotifyOnWarnings = ReadBooleanAttribute(element, "NotifyOnWarnings", true);
 				Dictionary<string, string> annotations = ReadAnnotationsAttribute(element, "Annotations");
+				string[] ignoreModified = ReadListAttribute(element, "IgnoreModified");
 
 				// Resolve all the inputs we depend on
 				HashSet<BgNodeOutput> inputs = ResolveInputReferences(element, requiresNames);
@@ -1432,7 +1433,7 @@ namespace AutomationTool
 				if (CheckNameIsUnique(element, name))
 				{
 					// Add it to the node lookup
-					BgScriptNode newNode = new BgScriptNode(name, inputs.ToArray(), validOutputNames.ToArray(), inputDependencies.ToArray(), orderDependencies.ToArray(), requiredTokens.ToArray());
+					BgScriptNode newNode = new BgScriptNode(name, inputs.ToArray(), validOutputNames.ToArray(), inputDependencies.ToArray(), orderDependencies.ToArray(), requiredTokens.ToArray(), ignoreModified.ToArray());
 					newNode.RunEarly = bRunEarly;
 					newNode.NotifyOnWarnings = bNotifyOnWarnings;
 					foreach ((string key, string value) in annotations)
diff --git a/Engine/Source/Programs/AutomationTool/BuildGraph/BgScriptWriter.cs b/Engine/Source/Programs/AutomationTool/BuildGraph/BgScriptWriter.cs
--- a/Engine/Source/Programs/AutomationTool/BuildGraph/BgScriptWriter.cs
+++ b/Engine/Source/Programs/AutomationTool/BuildGraph/BgScriptWriter.cs
@@ -173,6 +173,11 @@ namespace AutomationTool
 				writer.WriteAttributeString("RunEarly", node.RunEarly.ToString());
 			}
 
+			if (node.IgnoreModified.Count > 0)
+			{
+				writer.WriteAttributeString("IgnoreModified", String.Join(";", node.IgnoreModified));
+			}
+
 			BgScriptNode scriptNode = (BgScriptNode)node;
 			foreach (BgTask task in scriptNode.Tasks)
 			{
diff --git a/Engine/Source/Programs/AutomationTool/BuildGraph/BuildGraph.cs b/Engine/Source/Programs/AutomationTool/BuildGraph/BuildGraph.cs
--- a/Engine/Source/Programs/AutomationTool/BuildGraph/BuildGraph.cs
+++ b/Engine/Source/Programs/AutomationTool/BuildGraph/BuildGraph.cs
@@ -1015,6 +1015,9 @@ namespace AutomationTool
 			// Register something to execute cleanup commands
 			using CleanupScriptRunner CleanupRunner = new CleanupScriptRunner(Logger);
 
+			// Create a filter for modified files that should be ignored
+			FileFilter IgnoreModifiedFilter = new FileFilter(Node.IgnoreModified);
+
 			// Create the mapping of tag names to file sets
 			Dictionary<string, HashSet<FileReference>> TagNameToFileSet = new Dictionary<string, HashSet<FileReference>>();
 
@@ -1034,7 +1037,7 @@ namespace AutomationTool
 				Scope.Span.SetTag("blocks", InputStorageBlocks.Count);
 				foreach (TempStorageBlock InputStorageBlock in InputStorageBlocks)
 				{
-					TempStorageManifest Manifest = Storage.Retrieve(InputStorageBlock.NodeName, InputStorageBlock.OutputName);
+					TempStorageManifest Manifest = Storage.Retrieve(InputStorageBlock.NodeName, InputStorageBlock.OutputName, IgnoreModifiedFilter);
 					InputManifests[InputStorageBlock] = Manifest;
 				}
 				Scope.Span.SetTag("size", InputManifests.Sum(x => x.Value.GetTotalSize()));
@@ -1083,7 +1086,7 @@ namespace AutomationTool
 			foreach (TempStorageFile File in InputManifests.Values.SelectMany(x => x.Files))
 			{
 				string Message;
-				if (!ModifiedFiles.ContainsKey(File.RelativePath) && !File.Compare(Unreal.RootDirectory, out Message))
+				if (!ModifiedFiles.ContainsKey(File.RelativePath) && !File.Compare(Unreal.RootDirectory, out Message) && !IgnoreModifiedFilter.Matches(File.RelativePath))
 				{
 					// look up the previous nodes to help with error diagnosis
 					List<string> PreviousNodeNames = InputManifests.Where(x => x.Value.Files.Contains(File))
@@ -1102,14 +1105,14 @@ namespace AutomationTool
 				string modifiedFileList = "";
 				if (ModifiedFiles.Count < 100)
 				{
-					modifiedFileList = String.Join("\n", ModifiedFiles.Select(x => x.Value));
+					modifiedFileList = String.Join("\n", ModifiedFiles.Select(x => $"  {x.Value}"));
 				}
 				else
 				{
-					modifiedFileList = String.Join("\n", ModifiedFiles.Take(100).Select(x => x.Value));
-					modifiedFileList += $"{Environment.NewLine}And {ModifiedFiles.Count - 100} more.";
+					modifiedFileList = String.Join("\n", ModifiedFiles.Take(100).Select(x => $"  {x.Value}"));
+					modifiedFileList += $"\n  ...and {ModifiedFiles.Count - 100} more.";
 				}
-				throw new AutomationException("Build {0} from a previous step have been modified:\n{1}", (ModifiedFiles.Count == 1) ? "product" : "products", modifiedFileList);
+				throw new AutomationException("Build {0} from a previous step have been modified:\n{1}\nOutput overlapping artifacts to a different location, or ignore them using the Node's IgnoreModified attribute.", (ModifiedFiles.Count == 1) ? "product" : "products", modifiedFileList);
 			}
 
 			// Determine all the output files which are required to be copied to temp storage (because they're referenced by nodes in another agent)
diff --git a/Engine/Source/Programs/AutomationTool/BuildGraph/TempStorage.cs b/Engine/Source/Programs/AutomationTool/BuildGraph/TempStorage.cs
--- a/Engine/Source/Programs/AutomationTool/BuildGraph/TempStorage.cs
+++ b/Engine/Source/Programs/AutomationTool/BuildGraph/TempStorage.cs
@@ -17,6 +17,7 @@ using System.Xml;
 using System.Xml.Serialization;
 using UnrealBuildTool;
 using Microsoft.Extensions.Logging;
+using System.Text;
 
 using static AutomationTool.CommandUtils;
 
@@ -893,8 +894,9 @@ namespace AutomationTool
 		/// </summary>
 		/// <param name="NodeName">The node which created the storage block</param>
 		/// <param name="OutputName">Name of the block to retrieve. May be null or empty.</param>
+		/// <param name="IgnoreModified">Filter for files to ignore</param>
 		/// <returns>Manifest of the files retrieved</returns>
-		public TempStorageManifest Retrieve(string NodeName, string OutputName)
+		public TempStorageManifest Retrieve(string NodeName, string OutputName, FileFilter IgnoreModified)
 		{
 			using (IScope Scope = GlobalTracer.Instance.BuildSpan("RetrieveFromTempStorage").StartActive())
 			{
@@ -954,14 +956,28 @@ namespace AutomationTool
 				}
 
 				// Check all the local files are as expected
-				bool bAllMatch = true;
+				List<string> ModifiedFileMessages = new List<string>();
 				foreach(TempStorageFile File in Manifest.Files)
 				{
-					bAllMatch &= File.Compare(RootDir);
+					string Message;
+					if (!File.Compare(RootDir, out Message) && !IgnoreModified.Matches(File.RelativePath))
+					{
+						ModifiedFileMessages.Add(Message);
+					}
 				}
-				if(!bAllMatch)
+				if(ModifiedFileMessages.Count > 0)
 				{
-					throw new AutomationException("Files have been modified");
+					string modifiedFileList = "";
+					if (ModifiedFileMessages.Count < 100)
+					{
+						modifiedFileList = String.Join("\n", ModifiedFileMessages.Select(x => $"  {x}"));
+					}
+					else
+					{
+						modifiedFileList = String.Join("\n", ModifiedFileMessages.Take(100).Select(x => $"  {x}"));
+						modifiedFileList += $"\n  ...and {ModifiedFileMessages.Count - 100} more.";
+					}
+					throw new AutomationException($"Files have been modified:\n{modifiedFileList}");
 				}
 
 				// Update the stats and return
@@ -1475,10 +1491,10 @@ namespace AutomationTool
 			TempStore.Archive("TestNode", "NamedOutput", NamedOutput.Keys.ToArray(), true);
 			
 			// Check both outputs are still ok
-			TempStorageManifest DefaultManifest = TempStore.Retrieve("TestNode", null);
+			TempStorageManifest DefaultManifest = TempStore.Retrieve("TestNode", null, new FileFilter());
 			CheckManifest(WorkingDir, DefaultManifest, DefaultOutput);
 
-			TempStorageManifest NamedManifest = TempStore.Retrieve("TestNode", "NamedOutput");
+			TempStorageManifest NamedManifest = TempStore.Retrieve("TestNode", "NamedOutput", new FileFilter());
 			CheckManifest(WorkingDir, NamedManifest, NamedOutput);
 
 			// Delete local temp storage and the working directory and try again
@@ -1491,7 +1507,7 @@ namespace AutomationTool
 			bool bGotManifest = false;
 			try
 			{
-				TempStore.Retrieve("TestNode", null);
+				TempStore.Retrieve("TestNode", null, new FileFilter());
 				bGotManifest = true;
 			}
 			catch
@@ -1504,7 +1520,7 @@ namespace AutomationTool
 			}
 
 			// Second one should be fine
-			TempStorageManifest NamedManifestFromShared = TempStore.Retrieve("TestNode", "NamedOutput");
+			TempStorageManifest NamedManifestFromShared = TempStore.Retrieve("TestNode", "NamedOutput", new FileFilter());
 			CheckManifest(WorkingDir, NamedManifestFromShared, NamedOutput);
 		}
 
diff --git a/Engine/Source/Programs/Shared/EpicGames.BuildGraph/BgNodeDef.cs b/Engine/Source/Programs/Shared/EpicGames.BuildGraph/BgNodeDef.cs
--- a/Engine/Source/Programs/Shared/EpicGames.BuildGraph/BgNodeDef.cs
+++ b/Engine/Source/Programs/Shared/EpicGames.BuildGraph/BgNodeDef.cs
@@ -162,6 +162,11 @@ namespace EpicGames.BuildGraph
 		/// </summary>
 		public Dictionary<string, string> Annotations { get; } = new Dictionary<string, string>(StringComparer.OrdinalIgnoreCase);
 
+		/// <summary>
+		/// Ignore modified files matching the patterns provided
+		/// </summary>
+		public List<string> IgnoreModified { get; set; } = new List<string>();
+
 		/// <summary>
 		/// Diagnostics to output if executing this node
 		/// </summary>
@@ -190,7 +195,8 @@ namespace EpicGames.BuildGraph
 		/// <param name="inputDependencies">Nodes which this node is dependent on for its inputs</param>
 		/// <param name="orderDependencies">Nodes which this node needs to run after. Should include all input dependencies.</param>
 		/// <param name="requiredTokens">Optional tokens which must be required for this node to run</param>
-		public BgNodeDef(string name, IReadOnlyList<BgNodeOutput> inputs, IReadOnlyList<string> outputNames, IReadOnlyList<BgNodeDef> inputDependencies, IReadOnlyList<BgNodeDef> orderDependencies, IReadOnlyList<FileReference> requiredTokens)
+		/// <param name="ignoreModified">File patterns to ignore when checking for modified timestamps</param>
+		public BgNodeDef(string name, IReadOnlyList<BgNodeOutput> inputs, IReadOnlyList<string> outputNames, IReadOnlyList<BgNodeDef> inputDependencies, IReadOnlyList<BgNodeDef> orderDependencies, IReadOnlyList<FileReference> requiredTokens, IReadOnlyList<string> ignoreModified)
 			: this(name, null, outputNames)
 		{
 			Name = name;
@@ -198,6 +204,7 @@ namespace EpicGames.BuildGraph
 			InputDependencies.AddRange(inputDependencies);
 			OrderDependencies.AddRange(orderDependencies);
 			RequiredTokens.AddRange(requiredTokens);
+			IgnoreModified.AddRange(ignoreModified);
 		}
 
 		/// <summary>
