- Adds a ThreadSanitizer Structured Logging Matcher so we may contextualize TSAN output in Horde logs. It specifically will attempt to annotate files, symbols, and symbol addresses in callstacks, warning on TSAN "WARNING" messages and "SUMMARY" message.
-- The "SUMMARY" message is the only log line treated as an error since it is information (description of problem, and file+line to address) but also not overly specific such that we should be able to make JIRA issues from these log lines without generating duplicate issues for the same race from different callsites.
- Changes EditorPerf TSAN job to use -VeryVerbose in UAT so that stdout and stderr are redirected for logging. TSAN emits logs to stderr/stdout (depends on environement variable, stderr by default). The logs can be redirected to a log file via envvar TSAN_OPTIONS and specifying a path, however, Horde will only use whatever is written to stdout/stderr for generating issues, so unless a post process is added to forward the TSAN log output after the application has terminated, having stdout redirected stderr/stdout by default for TSAN jobs is simplest approach.
- Added two new KnownLogEvents `Sanitizer_Thread` and `Sanitizer`. `Sanitizer_Thread` is used by the added matcher, but `Sanitizer` was also added so that we may follow the pattern of letting the first id represent the broader category of events

diff --git a/Engine/Build/EditorPerf.xml b/Engine/Build/EditorPerf.xml
--- a/Engine/Build/EditorPerf.xml
+++ b/Engine/Build/EditorPerf.xml
@@ -182,7 +182,10 @@
 	<Property Name="LazyLoadZenTestArgs" Value="$(LazyLoadTestArgs) -zenloader" />
 
 	<!-- Gauntlet arguments for all runs -->
-	<Property Name="GauntletArgs" Value="-build=editor -verbose -tempdir=&quot;$(RootDir)/Tests&quot;"/>
+	<Property Name="GauntletArgs" Value="-build=editor -tempdir=&quot;$(RootDir)/Tests&quot;"/>
+	<!-- TSAN prints to stdout/stderr. Using -VeryVerbose allows those streams to be redirected and captured by Gauntlet -->
+	<Property Name="GauntletArgs" Value="$(GauntletArgs) -VeryVerbose" If="'$(WithTSAN)'=='true'"/>
+	<Property Name="GauntletArgs" Value="$(GauntletArgs) -verbose" If="'$(WithTSAN)'=='false'"/>
 	<Property Name="GauntletArgs" Value="$(GauntletArgs) -branch=$(Branch) -ECBranch=&quot;$(Branch)&quot; -ECChangelist=&quot;$(Change)&quot;"/>
 	<Property Name="GauntletArgs" Value="$(GauntletArgs) -maxduration=28800"/>
 
diff --git a/Engine/Source/Programs/AutomationTool/AutomationUtils/Matchers/ThreadSanitizerMatcher.cs b/Engine/Source/Programs/AutomationTool/AutomationUtils/Matchers/ThreadSanitizerMatcher.cs
new file mode 100644
--- /dev/null
+++ b/Engine/Source/Programs/AutomationTool/AutomationUtils/Matchers/ThreadSanitizerMatcher.cs
@@ -0,0 +1,119 @@
+// Copyright Epic Games, Inc. All Rights Reserved.
+
+using EpicGames.Core;
+using EpicGames.Horde.Logs;
+using Microsoft.CodeAnalysis;
+using Microsoft.CodeAnalysis.Differencing;
+using Microsoft.Extensions.FileSystemGlobbing.Internal;
+using Microsoft.Extensions.Logging;
+using System;
+using System.Collections.Generic;
+using System.Diagnostics;
+using System.Linq;
+using System.Text.RegularExpressions;
+
+#nullable enable
+
+namespace AutomationUtils.Matchers
+{
+	/// <summary>
+	/// Matcher for LLVM Thread Sanitizer reports
+	/// 
+	/// Note, TSAN reports emit to stderr/stdout which can cause issues for LogEventMatchers since they assume multiple lines will not be interleaved
+	/// with irrelevant lines which should be skipped / contextualized by other matchers. As such this matcher doesn't try to keep the whole TSAN report as a single series of events
+	/// and instead tries to annotate lines that look like report lines when possible. 
+	/// 
+	/// If you don't see TSAN output at all in your logs, it's likely you need -VeryVerbose to be passed to gauntlet to ensure stdout is redirected
+	/// </summary>
+	class ThreadSanitizerEventMatcher : ILogEventMatcher
+	{
+		// All TSAN reports begin with a WARNING and then SUMMARY line.
+		static readonly Regex s_WarningPattern = new Regex(@"^WARNING: ThreadSanitizer:\s*(.+)$");
+
+		// e.g. Note line number, column number and BuildId are optional
+		// #1 MyType<Type, OtherType>::Function(unsigned long*&) const /mnt/somepath/file.h:90:10 (BinaryName+0x2bc102e0) (BuildId: 8b17597a9f5444a0)
+		static readonly Regex s_StackTracePattern = new Regex(
+			@"^\s+\#[\d]+?\s" +
+			@"(?<symbol>.+?)\s" +
+			@"(?<sourcefile>(\/|\w:).+?)" +
+			@"(:(?<linenumber>\d+?))?" +
+			@"(:(?<colnumber>\d+?))?" +
+			@"\s+\(.+?\+(?<address>0x[0-9a-fA-F]+?)\)" +
+			@"(\s+?\(BuildId:\s+?(?<buildid>.+?)\))?$");
+
+		// e.g.
+		// SUMMARY: ThreadSanitizer: data race /some/path.cpp:169:33 in MyType::Foo()
+		static readonly Regex s_SummaryPattern = new Regex(
+			@"^SUMMARY: ThreadSanitizer:\s*" +
+			@"(?<error>.+)\s" +
+			@"(?<sourcefile>(\/|\w:).+?)" +
+			@"(:(?<linenumber>\d+?))?" +
+			@"(:(?<colnumber>\d+?))?\s+?in\s+?" + 
+			@"(?<symbol>.+)$");
+
+		// When files/symbols can't resolve, addr2line and llvm-symbolizer can replace symbol names with special names like "<null>"
+		// so instead of overcomplicating the above regexes to handle those conditional names, we use a catch all fallback patterns
+		static readonly Regex s_StackTraceFallbackPattern = new Regex(@"^\s+\#[\d]+?\s*(.+)$");
+		static readonly Regex s_SummaryFallbackPattern = new Regex(@"^SUMMARY: ThreadSanitizer:\s*(.+)$");
+
+		public LogEventMatch? Match(ILogCursor cursor)
+		{
+			Match? match;
+			if (cursor.TryMatch(s_WarningPattern, out match))
+			{
+				LogEventBuilder builder = new LogEventBuilder(cursor);
+				return builder.ToMatch(LogEventPriority.Normal, LogLevel.Warning, KnownLogEvents.Sanitizer_Thread);
+			}
+
+			if (cursor.TryMatch(s_StackTracePattern, out match))
+			{
+				LogEventBuilder builder = new LogEventBuilder(cursor);
+
+				do
+				{
+					builder.AnnotateSymbol(match!.Groups["symbol"]);
+					builder.AnnotateSourceFile(match.Groups["sourcefile"], "");
+					builder.Annotate(match.Groups["linenumber"], LogEventMarkup.LineNumber);
+					builder.Annotate(match.Groups["colnumber"], LogEventMarkup.ColumnNumber);
+					builder.Annotate(match.Groups["address"], "");
+					builder.Annotate(match.Groups["buildid"], "");
+
+					builder.MoveNext();
+				} while (builder.Current.TryMatch(s_StackTracePattern, out match));
+
+				return builder.ToMatch(LogEventPriority.Normal, LogLevel.Information, KnownLogEvents.Sanitizer_Thread);
+			}
+			else if (cursor.TryMatch(s_StackTraceFallbackPattern, out match))
+			{
+				LogEventBuilder builder = new LogEventBuilder(cursor);
+
+				do
+				{
+					builder.MoveNext();
+				} while (builder.Current.TryMatch(s_StackTraceFallbackPattern, out match));
+
+				return builder.ToMatch(LogEventPriority.Low, LogLevel.Information, KnownLogEvents.Sanitizer_Thread);
+			}
+
+			if (cursor.TryMatch(s_SummaryPattern, out match))
+			{
+				LogEventBuilder builder = new LogEventBuilder(cursor);
+				builder.Annotate(match.Groups["error"], LogEventMarkup.ErrorCode);
+				builder.AnnotateSourceFile(match.Groups["sourcefile"], "");
+				builder.Annotate(match.Groups["linenumber"], LogEventMarkup.LineNumber);
+				builder.Annotate(match.Groups["colnumber"], LogEventMarkup.ColumnNumber);
+				builder.AnnotateSymbol(match.Groups["symbol"]);
+
+				return builder.ToMatch(LogEventPriority.Normal, LogLevel.Error, KnownLogEvents.Sanitizer_Thread);
+			}
+			else if (cursor.TryMatch(s_SummaryFallbackPattern, out match))
+			{
+				LogEventBuilder builder = new LogEventBuilder(cursor);
+
+				return builder.ToMatch(LogEventPriority.Low, LogLevel.Error, KnownLogEvents.Sanitizer_Thread);
+			}
+
+			return null;
+		}
+	}
+}
diff --git a/Engine/Source/Programs/Shared/EpicGames.Core/KnownLogEvents.cs b/Engine/Source/Programs/Shared/EpicGames.Core/KnownLogEvents.cs
--- a/Engine/Source/Programs/Shared/EpicGames.Core/KnownLogEvents.cs
+++ b/Engine/Source/Programs/Shared/EpicGames.Core/KnownLogEvents.cs
@@ -331,6 +331,16 @@ namespace EpicGames.Core
 		/// </summary>
 		public static EventId Systemic_Max { get; } = new EventId(799);
 
+		/// <summary>
+		/// Generic Sanitizer Error
+		/// </summary>
+		public static EventId Sanitizer { get; } = new EventId(800);
+
+		/// <summary>
+		/// Thread Sanitizer Error
+		/// </summary>
+		public static EventId Sanitizer_Thread{ get; } = new EventId(801);
+
 		/// <summary>
 		/// Horde error codes
 		/// </summary>
