PR #11757: Fix tool deployment from Horde Agents. Specifying a JSON options file is now optional; if empty, the settings will be read from the environment.

diff --git a/Engine/Source/Programs/AutomationTool/BuildGraph/Tasks/DeployToolTask.cs b/Engine/Source/Programs/AutomationTool/BuildGraph/Tasks/DeployToolTask.cs
--- a/Engine/Source/Programs/AutomationTool/BuildGraph/Tasks/DeployToolTask.cs
+++ b/Engine/Source/Programs/AutomationTool/BuildGraph/Tasks/DeployToolTask.cs
@@ -43,7 +43,7 @@ namespace AutomationTool.Tasks
 		/// <summary>
 		/// Settings file to use for the deployment. Should be a JSON file containing server name and access token.
 		/// </summary>
-		[TaskParameter]
+		[TaskParameter(Optional = true)]
 		public string Settings = String.Empty;
 
 		/// <summary>
@@ -122,23 +122,27 @@ namespace AutomationTool.Tasks
 		/// <param name="TagNameToFileSet">Mapping from tag names to the set of files they include</param>
 		public override async Task ExecuteAsync(JobContext Job, HashSet<FileReference> BuildProducts, Dictionary<string, HashSet<FileReference>> TagNameToFileSet)
 		{
-			FileReference settingsFile = ResolveFile(Parameters.Settings);
-			if (!FileReference.Exists(settingsFile))
+			DeploySettings? settings = null;
+			if (!String.IsNullOrEmpty(Parameters.Settings))
 			{
-				throw new AutomationException($"Settings file '{settingsFile}' does not exist");
-			}
+				FileReference settingsFile = ResolveFile(Parameters.Settings);
+				if (!FileReference.Exists(settingsFile))
+				{
+					throw new AutomationException($"Settings file '{settingsFile}' does not exist");
+				}
 
-			byte[] settingsData = await FileReference.ReadAllBytesAsync(settingsFile);
-			JsonSerializerOptions jsonOptions = new JsonSerializerOptions { AllowTrailingCommas = true, ReadCommentHandling = JsonCommentHandling.Skip, PropertyNameCaseInsensitive = true };
+				byte[] settingsData = await FileReference.ReadAllBytesAsync(settingsFile);
+				JsonSerializerOptions jsonOptions = new JsonSerializerOptions { AllowTrailingCommas = true, ReadCommentHandling = JsonCommentHandling.Skip, PropertyNameCaseInsensitive = true };
 
-			DeploySettings? settings = JsonSerializer.Deserialize<DeploySettings>(settingsData, jsonOptions);
-			if (settings == null)
-			{
-				throw new AutomationException($"Unable to read settings file {settingsFile}");
-			}
-			else if (settings.Server == null)
-			{
-				throw new AutomationException($"Missing 'server' key from {settingsFile}");
+				settings = JsonSerializer.Deserialize<DeploySettings>(settingsData, jsonOptions);
+				if (settings == null)
+				{
+					throw new AutomationException($"Unable to read settings file {settingsFile}");
+				}
+				else if (settings.Server == null)
+				{
+					throw new AutomationException($"Missing 'server' key from {settingsFile}");
+				}
 			}
 
 			ToolId toolId = new ToolId(Parameters.Id);
@@ -146,8 +150,14 @@ namespace AutomationTool.Tasks
 			ServiceCollection serviceCollection = new ServiceCollection();
 			serviceCollection.Configure<HordeOptions>(options =>
 			{
-				options.ServerUrl = new Uri(settings.Server);
-				options.AccessToken = settings.Token;
+				if (!String.IsNullOrEmpty(settings?.Server))
+				{
+					options.ServerUrl = new Uri(settings.Server);
+				}
+				if (!String.IsNullOrEmpty(settings?.Token))
+				{
+					options.AccessToken = settings.Token;
+				}
 			});
 			serviceCollection.AddHttpClient();
 			serviceCollection.AddHorde();
@@ -158,7 +168,7 @@ namespace AutomationTool.Tasks
 			using HordeHttpClient hordeHttpClient = hordeClient.CreateHttpClient();
 
 			GetServerInfoResponse infoResponse = await hordeHttpClient.GetServerInfoAsync();
-			Logger.LogInformation("Uploading {ToolId} to {ServerUrl} (Version: {Version}, API v{ApiVersion})...", toolId, settings.Server, infoResponse.ServerVersion, (int)infoResponse.ApiVersion);
+			Logger.LogInformation("Uploading {ToolId} to {ServerUrl} (Version: {Version}, API v{ApiVersion})...", toolId, hordeClient.ServerUrl, infoResponse.ServerVersion, (int)infoResponse.ApiVersion);
 
 			BlobSerializerOptions serializerOptions = BlobSerializerOptions.Create(infoResponse.ApiVersion);
 
