PR #11760: Expand environment file check to also skip when empty string is passed to SpawnTask. Fixes issue in BuildHorde.xml when building with ImageEnvFile variable. (#11760)

diff --git a/Engine/Source/Programs/AutomationTool/BuildGraph/Tasks/SpawnTask.cs b/Engine/Source/Programs/AutomationTool/BuildGraph/Tasks/SpawnTask.cs
--- a/Engine/Source/Programs/AutomationTool/BuildGraph/Tasks/SpawnTask.cs
+++ b/Engine/Source/Programs/AutomationTool/BuildGraph/Tasks/SpawnTask.cs
@@ -106,7 +106,7 @@ namespace AutomationTool.Tasks
 			{
 				ParseEnvironment(Environment, ';', EnvVars);
 			}
-			if (EnvironmentFile != null)
+			if (!String.IsNullOrEmpty(EnvironmentFile))
 			{
 				ParseEnvironment(FileUtils.ReadAllText(ResolveFile(EnvironmentFile)), '\n', EnvVars);
 			}
