Fix stage/deploy functionality so that if you're using fileserver/projectstore functionality, it doesn't write a layout file that expects loose files (INIs, etc) to be on the filesystem.  This required moving some pieces of functionality onto ProjectParams so that they are accessible by a platform automation layer.

diff --git a/Engine/Source/Programs/AutomationTool/AutomationUtils/ProjectParams.cs b/Engine/Source/Programs/AutomationTool/AutomationUtils/ProjectParams.cs
--- a/Engine/Source/Programs/AutomationTool/AutomationUtils/ProjectParams.cs
+++ b/Engine/Source/Programs/AutomationTool/AutomationUtils/ProjectParams.cs
@@ -3206,6 +3206,50 @@ namespace AutomationTool
 			}
 		}
 
+		public FileReference FindZenProjectStoreMarker(DeploymentContext SC)
+		{
+			DirectoryReference ProjectStoreDir = null;
+			if (Stage)
+			{
+				if (SC.PlatformCookDir == null)
+				{
+					throw new AutomationException("FindZenProjectStoreMarker called before PlatformCookDir was populated");
+				}
+				ProjectStoreDir = SC.PlatformCookDir;
+			}
+			else if (Deploy)
+			{
+				ProjectStoreDir = SC.StageDirectory;
+			}
+			if (ProjectStoreDir == null)
+			{
+				return null;
+			}
+			// Check for stage with zenstore without PAK?
+			FileReference PackageStoreManifestFile = FileReference.Combine(ProjectStoreDir, "ue.projectstore");
+			System.IO.FileInfo PackageStoreManifestFileInfo = PackageStoreManifestFile.ToFileInfo();
+			if (PackageStoreManifestFileInfo.Exists)
+			{
+				return PackageStoreManifestFile;
+			}
+			return null;
+		}
+
+		public bool ShouldTreatAsFileServer(DeploymentContext SC)
+		{
+			if (FileServer)
+			{
+				return true;
+			}
+			if (ZenStore)
+			{
+				return true;
+			}
+			FileReference ZenStoreMarkerFile = FindZenProjectStoreMarker(SC);
+			bool UseZenServerHost = (ZenStoreMarkerFile != null) && !UsePak(SC.StageTargetPlatform);
+			return UseZenServerHost;
+		}
+
 		protected bool bLogged = false;
 		public virtual void ValidateAndLog()
 		{
diff --git a/Engine/Source/Programs/AutomationTool/Scripts/CopyBuildToStagingDirectory.Automation.cs b/Engine/Source/Programs/AutomationTool/Scripts/CopyBuildToStagingDirectory.Automation.cs
--- a/Engine/Source/Programs/AutomationTool/Scripts/CopyBuildToStagingDirectory.Automation.cs
+++ b/Engine/Source/Programs/AutomationTool/Scripts/CopyBuildToStagingDirectory.Automation.cs
@@ -474,17 +474,6 @@ namespace AutomationScripts
 			return InternationalizationPreset;
 		}
 
-		private static bool ShouldTreatAsFileServer(ProjectParams Params, DeploymentContext SC)
-		{
-			if (Params.FileServer)
-			{
-				return true;
-			}
-			FileReference ZenStoreMarkerFile = FindZenProjectStoreMarker(Params, SC);
-			bool UseZenServerHost = (ZenStoreMarkerFile != null) && !Params.UsePak(SC.StageTargetPlatform);
-			return UseZenServerHost;
-		}
-
 		public static List<string> GetCulturesToStage(ProjectParams Params, ConfigHierarchy PlatformGameConfig, bool bMustExist = true)
 		{
 			// Initialize cultures to stage.
@@ -1416,7 +1405,7 @@ namespace AutomationScripts
 						}
 
 						// UFS is required when using a file server
-						StagedFileType MovieFileType = ShouldTreatAsFileServer(Params, SC) ? StagedFileType.UFS : StagedFileType.NonUFS;
+						StagedFileType MovieFileType = Params.ShouldTreatAsFileServer(SC) ? StagedFileType.UFS : StagedFileType.NonUFS;
 
 						DirectoryReference EngineMoviesDir = DirectoryReference.Combine(SC.EngineRoot, "Content", "Movies");
 						if (DirectoryReference.Exists(EngineMoviesDir))
@@ -2262,7 +2251,7 @@ namespace AutomationScripts
 			CopyManifestFilesToStageDir(SC, SC.FilesToStage.NonUFSFiles, "NonUFSFiles");
 
 			Dictionary<StagedFileReference, FileReference> UFSFiles = new Dictionary<StagedFileReference, FileReference>(SC.FilesToStage.UFSFiles);
-			bool bStageUnrealFileSystemFiles = !Params.CookOnTheFly && !Params.UsePak(SC.StageTargetPlatform) && !ShouldTreatAsFileServer(Params, SC);
+			bool bStageUnrealFileSystemFiles = !Params.CookOnTheFly && !Params.UsePak(SC.StageTargetPlatform) && !Params.ShouldTreatAsFileServer(SC);
 			if (bStageUnrealFileSystemFiles)
 			{
 				foreach (KeyValuePair<StagedFileReference, FileReference> Pair in SC.CrashReporterUFSFiles)
@@ -5406,6 +5395,7 @@ namespace AutomationScripts
 					DeploymentContexts.Insert(0, SC);
 				}
 
+				SetUpStagingSourceDirectories(Params, SC);
 				SetupCustomStageCopyHandler(Params, SC);
 			}
 
diff --git a/Engine/Source/Programs/AutomationTool/Scripts/RunProjectCommand.Automation.cs b/Engine/Source/Programs/AutomationTool/Scripts/RunProjectCommand.Automation.cs
--- a/Engine/Source/Programs/AutomationTool/Scripts/RunProjectCommand.Automation.cs
+++ b/Engine/Source/Programs/AutomationTool/Scripts/RunProjectCommand.Automation.cs
@@ -796,39 +796,10 @@ namespace AutomationScripts
 			return HostAddresses.ToList();
 		}
 
-		private static FileReference FindZenProjectStoreMarker(ProjectParams Params, DeploymentContext SC)
-		{
-			DirectoryReference ProjectStoreDir = null;
-			if (Params.Stage)
-			{
-				if (!SetUpStagingSourceDirectories(Params, SC))
-				{
-					return null;
-				}
-				ProjectStoreDir = SC.PlatformCookDir;
-			}
-			else if (Params.Deploy)
-			{
-				ProjectStoreDir = SC.StageDirectory;
-			}
-			if (ProjectStoreDir == null)
-			{
-				return null;
-			}
-			// Check for stage with zenstore without PAK?
-			FileReference PackageStoreManifestFile = FileReference.Combine(ProjectStoreDir, "ue.projectstore");
-			System.IO.FileInfo PackageStoreManifestFileInfo = PackageStoreManifestFile.ToFileInfo();
-			if (PackageStoreManifestFileInfo.Exists)
-			{
-				return PackageStoreManifestFile;
-			}
-			return null;
-		}
-
 		private static string GetFileHostCommandline(ProjectParams Params, DeploymentContext SC)
 		{
 			string FileHostParams = "";
-			FileReference ZenStoreMarkerFile = FindZenProjectStoreMarker(Params, SC);
+			FileReference ZenStoreMarkerFile = Params.FindZenProjectStoreMarker(SC);
 			bool UseZenServerHost = (ZenStoreMarkerFile != null) && !Params.UsePak(SC.StageTargetPlatform);
 			if (!Params.CookOnTheFly && !Params.FileServer && !UseZenServerHost)
 			{
