Use AFS instead of ADB when available.

diff --git a/Engine/Plugins/Runtime/AndroidFileServer/Source/AndroidFileServer/Android/java/com/epicgames/unreal/RemoteFileManager.java b/Engine/Plugins/Runtime/AndroidFileServer/Source/AndroidFileServer/Android/java/com/epicgames/unreal/RemoteFileManager.java
--- a/Engine/Plugins/Runtime/AndroidFileServer/Source/AndroidFileServer/Android/java/com/epicgames/unreal/RemoteFileManager.java
+++ b/Engine/Plugins/Runtime/AndroidFileServer/Source/AndroidFileServer/Android/java/com/epicgames/unreal/RemoteFileManager.java
@@ -168,7 +168,7 @@ public class RemoteFileManager implements Runnable
 		replacements.put("^mainobb", mainobb);
 		replacements.put("^patchobb", patchobb);
 
-		for (int overflowIndex=1; overflowIndex < MAX_OVERFLOW_FILES; overflowIndex++)
+		for (int overflowIndex=1; overflowIndex <= MAX_OVERFLOW_FILES; overflowIndex++)
 		{
 			String overflowobb = obbPackageDir + "overflow" + overflowIndex + "." + versionString + "." + packageName + ".obb";
 			replacements.put("^overflow" + overflowIndex + "obb", overflowobb);
@@ -963,6 +963,7 @@ public class RemoteFileManager implements Runnable
 				if (!DirCreateRecursive(path.subpath(0, lastIndex - 1).toString()))
 				{
 					IgnoreData(dataRemaining);
+					//Log.d(TAG, "FileWrite failed (NoSuchFileException): " + path.toString(), fe);
 //					SendResponse("false");
 					return;
 				}
@@ -970,6 +971,7 @@ public class RemoteFileManager implements Runnable
 			catch (Exception e)
 			{
 				IgnoreData(dataRemaining);
+				Log.d(TAG, "FileWrite failed: " + path.toString(), e);
 //				SendResponse("false");
 				return;
 			}
diff --git a/Engine/Source/Programs/AutomationTool/Android/AndroidPlatform.Automation.cs b/Engine/Source/Programs/AutomationTool/Android/AndroidPlatform.Automation.cs
--- a/Engine/Source/Programs/AutomationTool/Android/AndroidPlatform.Automation.cs
+++ b/Engine/Source/Programs/AutomationTool/Android/AndroidPlatform.Automation.cs
@@ -1609,8 +1609,8 @@ public class AndroidPlatform : Platform
 				{
 					foreach (OverflowBatchInstallInfo Overflow in OverflowInfo)
 					{
-						string AfsOverflowName = string.Format("^overflow{0}obb", Overflow.OverflowIndex);
-						OverflowInstallCommands.Add(Overflow.bNoOverflowInstall ? AFSCommand + " deletefile '" + AfsOverflowName + "'" : AFSCommand + " push " + Path.GetFileName(Overflow.OverflowName) + " '" + AfsOverflowName + "'");
+						string AFSOverflowName = string.Format("^overflow{0}obb", Overflow.OverflowIndex);
+						OverflowInstallCommands.Add(Overflow.bNoOverflowInstall ? AFSCommand + " deletefile '" + AFSOverflowName + "'" : AFSCommand + " push " + Path.GetFileName(Overflow.OverflowName) + " '" + AFSOverflowName + "'");
 						OverflowInstallCommands.Add("if \"%ERRORLEVEL%\" NEQ \"0\" goto Error");
 					}
 				}
@@ -1718,8 +1718,8 @@ public class AndroidPlatform : Platform
 				{
 					foreach (OverflowBatchInstallInfo Overflow in OverflowInfo)
 					{
-						string AfsOverflowName = string.Format("^overflow{0}obb", Overflow.OverflowIndex);
-						OverflowInstallCommands.Add(Overflow.bNoOverflowInstall ? AFSCommand + " deletefile \"" + AfsOverflowName + "\"" : AFSCommand + " push " + Path.GetFileName(Overflow.OverflowName) + " \"" + AfsOverflowName + "\"");
+						string AFSOverflowName = string.Format("^overflow{0}obb", Overflow.OverflowIndex);
+						OverflowInstallCommands.Add(Overflow.bNoOverflowInstall ? AFSCommand + " deletefile \"" + AFSOverflowName + "\"" : AFSCommand + " push " + Path.GetFileName(Overflow.OverflowName) + " \"" + AFSOverflowName + "\"");
 						OverflowInstallCommands.Add("if \"%ERRORLEVEL%\" NEQ \"0\" goto Error");
 					}
 				}
@@ -2178,6 +2178,15 @@ public class AndroidPlatform : Platform
 
 		return string.Format("{0} {1}", SerialNumber, Args);
 	}
+	private static string GetAFSCommandLine(string SerialNumber, string Args)
+	{
+	    if (string.IsNullOrEmpty(SerialNumber) == false)
+		{
+			SerialNumber = "-s " + SerialNumber;
+		}
+
+		return string.Format("{0} {1}", SerialNumber, Args);
+	}
 
 	static string LastSpewFilename = "";
 
@@ -2225,6 +2234,31 @@ public class AndroidPlatform : Platform
 		return RunAndLog(CmdEnv, AdbCommand, GetAdbCommandLine(SerialNumber, Args), out SuccessCode, SpewFilterCallback: new ProcessResult.SpewFilterCallbackType(ADBSpewFilter));
 	}
 
+	public static IProcessResult RunAFSCommand(ProjectParams Params, string SerialNumber, string Args, string Input = null, ERunOptions Options = ERunOptions.Default, bool bShouldLogCommand = false)
+	{
+		return RunAFSCommand(SerialNumber, Args, Input, Options, bShouldLogCommand);
+	}
+
+	private static IProcessResult RunAFSCommand(string SerialNumber, string Args, string Input = null, ERunOptions Options = ERunOptions.Default, bool bShouldLogCommand = false)
+	{
+		string AFSExecutable = AndroidExports.GetAFSExecutable(UnrealTargetPlatform.Win64, EpicGames.Core.Log.Logger);
+		AFSExecutable = Path.Combine(Unreal.EngineDirectory.FullName, "Binaries", "DotNET", "Android", "UnrealAndroidFileTool", AFSExecutable);
+		if (Options.HasFlag(ERunOptions.AllowSpew) || Options.HasFlag(ERunOptions.SpewIsVerbose))
+		{
+			LastSpewFilename = "";
+			return Run(AFSExecutable, GetAFSCommandLine(SerialNumber, Args), Input, Options, SpewFilterCallback: new ProcessResult.SpewFilterCallbackType(ADBSpewFilter));
+		}
+		return Run(AFSExecutable, GetAFSCommandLine(SerialNumber, Args), Input, Options);
+	}
+
+	private string RunAndLogAFSCommand(ProjectParams Params, string SerialNumber, string Args, out int SuccessCode)
+	{
+		string AFSExecutable = AndroidExports.GetAFSExecutable(UnrealTargetPlatform.Win64, EpicGames.Core.Log.Logger);
+		AFSExecutable = Path.Combine(Unreal.EngineDirectory.FullName, "Binaries", "DotNET", "Android", "UnrealAndroidFileTool", AFSExecutable);
+		LastSpewFilename = "";
+		return RunAndLog(CmdEnv, AFSExecutable, GetAFSCommandLine(SerialNumber, Args), out SuccessCode, SpewFilterCallback: new ProcessResult.SpewFilterCallbackType(ADBSpewFilter));
+	}
+
 	public override void GetConnectedDevices(ProjectParams Params, out List<string> Devices)
 	{
 		Devices = new List<string>();
diff --git a/Engine/Source/Programs/AutomationTool/Gauntlet/Framework/Base/Gauntlet.AppInstall.cs b/Engine/Source/Programs/AutomationTool/Gauntlet/Framework/Base/Gauntlet.AppInstall.cs
--- a/Engine/Source/Programs/AutomationTool/Gauntlet/Framework/Base/Gauntlet.AppInstall.cs
+++ b/Engine/Source/Programs/AutomationTool/Gauntlet/Framework/Base/Gauntlet.AppInstall.cs
@@ -20,7 +20,7 @@ namespace Gauntlet
 		/// Create a running instance of the app
 		/// </summary>
 		/// <returns></returns>
-		IAppInstance Run();
+		IAppInstance Run(UnrealAppConfig AppConfiguration);
 
 		/// <summary>
 		/// The device that we'll be run on
diff --git a/Engine/Source/Programs/AutomationTool/Gauntlet/Framework/Base/Gauntlet.TargetDevice.cs b/Engine/Source/Programs/AutomationTool/Gauntlet/Framework/Base/Gauntlet.TargetDevice.cs
--- a/Engine/Source/Programs/AutomationTool/Gauntlet/Framework/Base/Gauntlet.TargetDevice.cs
+++ b/Engine/Source/Programs/AutomationTool/Gauntlet/Framework/Base/Gauntlet.TargetDevice.cs
@@ -87,7 +87,7 @@ namespace Gauntlet
 
 		IAppInstall InstallApplication(UnrealAppConfig AppConfiguration);
 
-		IAppInstance Run(IAppInstall App);
+		IAppInstance Run(UnrealAppConfig AppConfiguration, IAppInstall App);
 
 		/// Begin new flow ///
 
diff --git a/Engine/Source/Programs/AutomationTool/Gauntlet/Framework/Base/Gauntlet.TargetDeviceDesktopCommon.cs b/Engine/Source/Programs/AutomationTool/Gauntlet/Framework/Base/Gauntlet.TargetDeviceDesktopCommon.cs
--- a/Engine/Source/Programs/AutomationTool/Gauntlet/Framework/Base/Gauntlet.TargetDeviceDesktopCommon.cs
+++ b/Engine/Source/Programs/AutomationTool/Gauntlet/Framework/Base/Gauntlet.TargetDeviceDesktopCommon.cs
@@ -223,7 +223,7 @@ namespace Gauntlet
             LocalDirectoryMappings.Add(EIntendedBaseCopyDirectory.Saved, Path.Combine(BaseDirectory, "Saved"));
 		}
 
-		public abstract IAppInstance Run(IAppInstall Install);
+		public abstract IAppInstance Run(UnrealAppConfig AppConfiguration, IAppInstall Install);
 
 		protected abstract IAppInstall CreateNativeStagedInstall(UnrealAppConfig AppConfig, NativeStagedBuild Build);
 
@@ -323,9 +323,9 @@ namespace Gauntlet
 			CanAlterCommandArgs = true;
 		}
 
-		public virtual IAppInstance Run()
+		public virtual IAppInstance Run(UnrealAppConfig AppConfiguration)
 		{
-			return Device.Run(this);
+			return Device.Run(AppConfiguration, this);
 		}
 
 		/// <summary>
diff --git a/Engine/Source/Programs/AutomationTool/Gauntlet/Platform/Android/Gauntlet.TargetDeviceAndroid.cs b/Engine/Source/Programs/AutomationTool/Gauntlet/Platform/Android/Gauntlet.TargetDeviceAndroid.cs
--- a/Engine/Source/Programs/AutomationTool/Gauntlet/Platform/Android/Gauntlet.TargetDeviceAndroid.cs
+++ b/Engine/Source/Programs/AutomationTool/Gauntlet/Platform/Android/Gauntlet.TargetDeviceAndroid.cs
@@ -107,6 +107,53 @@ namespace Gauntlet
 
 		protected Dictionary<EIntendedBaseCopyDirectory, string> LocalDirectoryMappings;
 
+		public static bool UsingAndroidFileServer(FileReference RawProjectPath, UnrealTargetConfiguration TargetConfiguration, out bool bEnablePlugin, out string AFSToken, out bool bIsShipping, out bool bIncludeInShipping, out bool bAllowExternalStartInShipping)
+		{
+			UnrealTargetPlatform TargetPlatform = UnrealTargetPlatform.Android;
+			//UnrealTargetConfiguration TargetConfiguration = SC != null ? SC.StageTargetConfigurations[0] : Params.ClientConfigsToBuild[0];
+			bIsShipping = TargetConfiguration == UnrealTargetConfiguration.Shipping;
+
+			ConfigHierarchy Ini = ConfigCache.ReadHierarchy(ConfigHierarchyType.Engine, DirectoryReference.FromFile(RawProjectPath), TargetPlatform);
+			if (!Ini.GetBool("/Script/AndroidFileServerEditor.AndroidFileServerRuntimeSettings", "bEnablePlugin", out bEnablePlugin))
+			{
+				bEnablePlugin = true;
+			}
+			if (!Ini.GetString("/Script/AndroidFileServerEditor.AndroidFileServerRuntimeSettings", "SecurityToken", out AFSToken))
+			{
+				AFSToken = "";
+			}
+			if (!Ini.GetBool("/Script/AndroidFileServerEditor.AndroidFileServerRuntimeSettings", "bIncludeInShipping", out bIncludeInShipping))
+			{
+				bIncludeInShipping = false;
+			}
+			if (!Ini.GetBool("/Script/AndroidFileServerEditor.AndroidFileServerRuntimeSettings", "bAllowExternalStartInShipping", out bAllowExternalStartInShipping))
+			{
+				bAllowExternalStartInShipping = false;
+			}
+
+			if (bIsShipping && !(bIncludeInShipping && bAllowExternalStartInShipping))
+			{
+				return false;
+			}
+			return bEnablePlugin;
+		}
+
+		public bool UseAFS(UnrealAppConfig AppConfig)
+		{
+			AndroidBuild Build = AppConfig.Build as AndroidBuild;
+			if (Build == null)
+			{
+				throw new AutomationException("Unsupported build type {0} for Android!", AppConfig.Build.GetType());
+			}
+
+			bool bAFSEnablePlugin;
+			string AFSToken;
+			bool bIsShipping;
+			bool bAFSIncludeInShipping;
+			bool bAFSAllowExternalStartInShipping;
+			return UsingAndroidFileServer(AppConfig.ProjectFile, Build.Configuration, out bAFSEnablePlugin, out AFSToken, out bIsShipping, out bAFSIncludeInShipping, out bAFSAllowExternalStartInShipping);
+		}
+
 		public TargetDeviceAndroid(string InDeviceName = "", AndroidDeviceData DeviceData = null, string InCachePath = null)
 		{
 			AdbCredentialCache.AddInstance(DeviceData);
@@ -362,16 +409,16 @@ namespace Gauntlet
 			EnablePermissions(Package);
 
 			// Install apk
-			CopyFileToDevice(Package, Build.SourceApkPath, string.Empty);
+			CopyFileToDevice(AppConfig, Package, Build.SourceApkPath, string.Empty);
 
 			// Copy obbs from bulk builds
 			if(Build.FilesToInstall != null && Build.FilesToInstall.Any())
 			{
-				CopyOBBFiles(Build.FilesToInstall, Package);
+				CopyOBBFiles(AppConfig, Build.FilesToInstall, Package);
 			}
 
 			// Now, create and push the UECommandline.txt file
-			CopyCommandlineFile(AppConfig.CommandLine, Package, AppConfig.ProjectName, Build.UsesExternalFilesDir);
+			CopyCommandlineFile(AppConfig, AppConfig.CommandLine, Package, AppConfig.ProjectName, Build.UsesExternalFilesDir);
 		}
 
 		public IAppInstall CreateAppInstall(UnrealAppConfig AppConfig)
@@ -394,7 +441,7 @@ namespace Gauntlet
 			return new AndroidAppInstall(this, AppConfig.ProjectName, Build.AndroidPackageName, AppConfig.CommandLine);
 		}
 
-		public void CopyAdditionalFiles(IEnumerable<UnrealFileToCopy> FilesToCopy)
+		public void CopyAdditionalFiles(UnrealAppConfig AppConfig, IEnumerable<UnrealFileToCopy> FilesToCopy)
 		{
 			foreach (UnrealFileToCopy FileToCopy in FilesToCopy)
 			{
@@ -403,7 +450,7 @@ namespace Gauntlet
 				{
 					FileInfo SourceFile = new FileInfo(FileToCopy.SourceFileLocation);
 					SourceFile.IsReadOnly = false;
-					CopyFileToDevice(null, SourceFile.FullName, DestinationFile);
+					CopyFileToDevice(AppConfig, null, SourceFile.FullName, DestinationFile);
 				}
 				else
 				{
@@ -411,9 +458,33 @@ namespace Gauntlet
 				}
 			}
 		}
+		public void CopyAdditionalFiles(IEnumerable<UnrealFileToCopy> FilesToCopy)
+		{
+			CopyAdditionalFiles(null, FilesToCopy);
+		}
 
-		public IAppInstance Run(IAppInstall App)
+		public IAppInstance Run(UnrealAppConfig AppConfig, IAppInstall App)
 		{
+			bool bAFSEnablePlugin;
+			string AFSToken = "";
+			bool bIsShipping;
+			bool bAFSIncludeInShipping;
+			bool bAFSAllowExternalStartInShipping;
+			bool useADB = true;
+			if (AppConfig != null)
+			{
+				AndroidBuild Build = AppConfig.Build as AndroidBuild;
+				if (Build == null)
+				{
+					throw new AutomationException("Unsupported build type {0} for Android!", AppConfig.Build.GetType());
+				}
+
+				if (UsingAndroidFileServer(AppConfig.ProjectFile, Build.Configuration, out bAFSEnablePlugin, out AFSToken, out bIsShipping, out bAFSIncludeInShipping, out bAFSAllowExternalStartInShipping))
+				{
+					useADB = false;
+				}
+			}
+
 			if (App is not AndroidAppInstall Install)
 			{
 				throw new AutomationException("AppInstall is of incorrect type {0}! Must be of type AndroidAppInstall", App.GetType().Name);
@@ -436,8 +507,15 @@ namespace Gauntlet
 			RunAdbDeviceCommand("logcat --clear");
 
 			// Ensure artifact directories exist
-			RunAdbDeviceCommand(string.Format("shell mkdir -p {0}/", DeviceExternalStorageSavedPath));
-			RunAdbDeviceCommand(string.Format("shell mkdir -p {0}/", DeviceExternalFilesSavedPath));
+			if (useADB)
+			{
+				RunAdbDeviceCommand(string.Format("shell mkdir -p {0}/", DeviceExternalStorageSavedPath));
+				RunAdbDeviceCommand(string.Format("shell mkdir -p {0}/", DeviceExternalFilesSavedPath));
+			}
+			else
+			{
+				RunAFSDeviceCommand(string.Format("-p \"{0}\" -k \"{1}\" mkdir \"^saved\"", Install.AndroidPackageName, AFSToken));
+			}
 
 			// start the app on device!
 			string CommandLine = "shell am start -W -S -n " + Install.AndroidPackageName + "/" + LaunchActivity;
@@ -452,7 +530,7 @@ namespace Gauntlet
 			LocalDirectoryMappings.Clear();
 
 			ProjectDir = ProjectDir.Replace('\\', '/');
-			if(!ProjectDir.EndsWith('/'))
+			if (!ProjectDir.EndsWith('/'))
 			{
 				ProjectDir += '/';
 			}
@@ -484,6 +562,16 @@ namespace Gauntlet
 			return RunAdbGlobalCommand(Args, Wait, bShouldLogCommand, bPauseErrorParsing);
 		}
 
+		public IProcessResult RunAFSDeviceCommand(string Args, bool Wait = true, bool bShouldLogCommand = false, bool bPauseErrorParsing = false)
+		{
+			if (string.IsNullOrEmpty(DeviceName) == false)
+			{
+				Args = string.Format("-s {0} {1}", DeviceName, Args);
+			}
+
+			return RunAFSGlobalCommand(Args, Wait, bShouldLogCommand, bPauseErrorParsing);
+		}
+
 		/// <summary>
 		/// Runs an ADB command, automatically adding the name of the current device to
 		/// the arguments sent to adb
@@ -547,25 +635,147 @@ namespace Gauntlet
 
 			return Process;
 		}
+		
+		public static IProcessResult RunAFSGlobalCommand(string Args, bool Wait = true, bool bShouldLogCommand = false, bool bPauseErrorParsing = false)
+		{
+			CommandUtils.ERunOptions RunOptions = CommandUtils.ERunOptions.AppMustExist | CommandUtils.ERunOptions.NoWaitForExit | CommandUtils.ERunOptions.SpewIsVerbose;
+
+			if (Log.IsVeryVerbose)
+			{
+				RunOptions |= CommandUtils.ERunOptions.AllowSpew;
+			}
+			else
+			{
+				RunOptions |= CommandUtils.ERunOptions.NoLoggingOfRunCommand;
+			}
+
+			if (bShouldLogCommand)
+			{
+				Log.Verbose("Running AFS Command: UnrealAndroidFileTool {0}", Args);
+			}
+
+			IProcessResult Process;
+
+			using (bPauseErrorParsing ? new ScopedSuspendECErrorParsing() : null)
+			{
+				Process = AndroidPlatform.RunAFSCommand(null, null, Args, null, RunOptions);
+
+				if (Wait)
+				{
+					Process.WaitForExit();
+				}
+			}
+
+			return Process;
+		}
+
 
 		public static ITargetDevice[] GetDefaultDevices()
 		{
 			return GetAllAvailableDevices().Select(Device => new TargetDeviceAndroid(Device)).ToArray();
 		}
 
-		public void PostRunCleanup()
+		public void PostRunCleanup(UnrealAppConfig AppConfig, string Package)
 		{
 			// Delete the commandline file, if someone installs an APK on top of ours
 			// they will get very confusing behavior...
 			if (string.IsNullOrEmpty(CommandLineFilePath) == false)
 			{
 				Log.Verbose("Removing {0}", CommandLineFilePath);
-				DeleteFileFromDevice(CommandLineFilePath);
+				if (UseAFS(AppConfig))
+				{
+					DeleteFileFromDevice(AppConfig, Package, "^commandfile");
+				}
+				else
+				{
+					DeleteFileFromDevice(AppConfig, Package, CommandLineFilePath);
+				}
+			}
+		}
+
+		public string AFSGetFileInfo(string PackageName, string AFSToken, string FileName)
+		{
+			IProcessResult result;
+			if (FileName.StartsWith("^"))
+			{
+				// Query to get which file this actually is.
+				string ActualFileName = "";
+				result = RunAFSDeviceCommand(string.Format("-p \"{0}\" -k \"{1}\" query", PackageName, AFSToken));
+				string[] mappings = result.Output.Split(new[] { "\r\n" }, StringSplitOptions.RemoveEmptyEntries).Select(S => S.Trim()).ToArray();
+				foreach (string mapping in mappings)
+				{
+					string[] FromTo = mapping.Split(new[] { " " }, StringSplitOptions.RemoveEmptyEntries).Select(S => S.Trim()).ToArray();
+					if (FromTo.Length != 2)
+					{
+						continue;
+					}
+					if (FromTo[0] == FileName)
+					{
+						ActualFileName = Path.GetFileName(FromTo[1]);
+						break;
+					}
+				}
+
+				if (string.IsNullOrEmpty(ActualFileName))
+				{
+					return "";
+					//throw new AutomationException("Failed to find file {0} in AFS mappings", FileName);
+				}
+
+				// Now List all the files in ^obb and find our file
+				result = RunAFSDeviceCommand(string.Format("-p \"{0}\" -k \"{1}\" ls -ls \"^obb\"", PackageName, AFSToken));
+				string[] ObbFiles = result.Output.Split(new[] { "\r\n" }, StringSplitOptions.RemoveEmptyEntries).Select(S => S.Trim()).ToArray();
+				foreach (string file in ObbFiles)
+				{
+					if (file.Contains(ActualFileName))
+					{
+						return file;
+					}
+				}
+
+				// Now List all the files in ^project and find our file
+				result = RunAFSDeviceCommand(string.Format("-p \"{0}\" -k \"{1}\" ls -ls \"^project\"", PackageName, AFSToken));
+				string[] ProjectFiles = result.Output.Split(new[] { "\r\n" }, StringSplitOptions.RemoveEmptyEntries).Select(S => S.Trim()).ToArray();
+				foreach (string file in ProjectFiles)
+				{
+					if (file.Contains(ActualFileName))
+					{
+						return file;
+					}
+				}
+
+				//throw new AutomationException("Failed to find file {0} in AFS ^obb", FileName);
+				return "";
+			}
+			else
+			{
+				//throw new AutomationException("Unsupported file {0} for AFS", FileName);
+				return "";
 			}
 		}
 
-		public bool CopyFileToDevice(string PackageName, string SourcePath, string DestPath, bool IgnoreDependencies = false)
+		public bool CopyFileToDevice(UnrealAppConfig AppConfig, string PackageName, string SourcePath, string DestPath, bool IgnoreDependencies = false)
 		{
+			bool bAFSEnablePlugin;
+			string AFSToken = "";
+			bool bIsShipping;
+			bool bAFSIncludeInShipping;
+			bool bAFSAllowExternalStartInShipping;
+			bool useADB = true;
+			if (AppConfig != null)
+			{
+				AndroidBuild Build = AppConfig.Build as AndroidBuild;
+				if (Build == null)
+				{
+					throw new AutomationException("Unsupported build type {0} for Android!", AppConfig.Build.GetType());
+				}
+
+				if (UsingAndroidFileServer(AppConfig.ProjectFile, Build.Configuration, out bAFSEnablePlugin, out AFSToken, out bIsShipping, out bAFSIncludeInShipping, out bAFSAllowExternalStartInShipping))
+				{
+					useADB = false;
+				}
+			}
+
 			bool IsAPK = string.Equals(Path.GetExtension(SourcePath), ".apk", StringComparison.OrdinalIgnoreCase);
 
 			// for the APK there's no easy/reliable way to get the date of the version installed, so 
@@ -603,12 +813,19 @@ namespace Gauntlet
 			else
 			{
 				// for other files get the file info
-				AdbResult = RunAdbDeviceCommand(string.Format("shell ls -l {0}", DestPath));
-
-				if (AdbResult.ExitCode == 0)
+				if (useADB)
 				{
-					CurrentFileInfo = AdbResult.Output.ToString().Trim();
+					AdbResult = RunAdbDeviceCommand(string.Format("shell ls -l {0}", DestPath));
+					if (AdbResult.ExitCode == 0)
+					{
+						CurrentFileInfo = AdbResult.Output.ToString().Trim();
+					}
+				}
+				else
+				{
+					CurrentFileInfo = AFSGetFileInfo(PackageName, AFSToken, DestPath);
 				}
+
 			}
 
 			bool SkipInstall = false;
@@ -667,11 +884,21 @@ namespace Gauntlet
 				else
 				{
 					string FileDirectory = Path.GetDirectoryName(DestPath).Replace('\\', '/');
-					RunAdbDeviceCommand(string.Format("shell mkdir -p {0}/", FileDirectory));
 
-					Log.Info("Copying {0} to {1} via adb push", QuotedSourcePath, DestPath);
-					string AdbCommand = string.Format("push {0} {1}", QuotedSourcePath, DestPath);
-					AdbResult = RunAdbDeviceCommand(AdbCommand);
+					if (useADB)
+					{
+						RunAdbDeviceCommand(string.Format("shell mkdir -p {0}/", FileDirectory));
+
+						Log.Info("Copying {0} to {1} via adb push", QuotedSourcePath, DestPath);
+						AdbResult = RunAdbDeviceCommand(string.Format("push {0} {1}", QuotedSourcePath, DestPath));
+					}
+					else
+					{
+						Log.Info("Copying {0} to {1} via AFS", QuotedSourcePath, DestPath);
+						AdbResult = RunAFSDeviceCommand(string.Format("-p \"{0}\" -k \"{1}\" push \"{2}\" \"{3}\"", PackageName, AFSToken, SourcePath, DestPath));
+					}
+
+					// Note: Presently, AdbResult NEVER reports failures to push when using AFS. Should be fixed at some point.
 
 					if (AdbResult.ExitCode != 0)
 					{
@@ -682,8 +909,16 @@ namespace Gauntlet
 							// Try to reconnect
 							if (Connect())
 							{
-								Log.Info("Retrying to copy via adb push...");
-								AdbResult = RunAdbDeviceCommand(AdbCommand);
+								if (useADB)
+								{
+									Log.Info("Retrying to copy via adb push...");
+									AdbResult = RunAdbDeviceCommand(string.Format("push {0} {1}", QuotedSourcePath, DestPath));
+								}
+								else
+								{
+									Log.Info("Retrying to copy via AFS...");
+									AdbResult = RunAFSDeviceCommand(string.Format("-p {0} -k {1} push \"{2}\" \"{3}\"", PackageName, AFSToken, SourcePath, DestPath));
+								}
 								if (AdbResult.ExitCode != 0)
 								{
 									throw new AutomationException("Failed to push {0} to device. Error {1}", SourcePath, AdbResult.Output);
@@ -701,8 +936,15 @@ namespace Gauntlet
 					}
 
 					// Now pull info about the file which we'll write as a dep
-					AdbResult = RunAdbDeviceCommand(string.Format("shell ls -l {0}", DestPath));
-					CurrentFileInfo = AdbResult.Output.ToString().Trim();
+					if (useADB)
+					{
+						AdbResult = RunAdbDeviceCommand(string.Format("shell ls -l {0}", DestPath));
+						CurrentFileInfo = AdbResult.Output.ToString().Trim();
+					}
+					else
+					{
+						CurrentFileInfo = AFSGetFileInfo(PackageName, AFSToken, DestPath);
+					}
 				}
 
 				// write the actual dependency info
@@ -710,9 +952,7 @@ namespace Gauntlet
 
 				// save last modified time to remote deps after success
 				RunAdbDeviceCommand(string.Format("shell mkdir -p {0}", DependencyCacheDir));
-
-				string Cmd = string.Format("shell echo \"{0}\" > {1}", DepContents, DepFile);
-				AdbResult = RunAdbDeviceCommand(Cmd);
+				AdbResult = RunAdbDeviceCommand(string.Format("shell echo \"{0}\" > {1}", DepContents, DepFile));
 
 				if (AdbResult.ExitCode != 0)
 				{
@@ -749,13 +989,46 @@ namespace Gauntlet
 			RunAdbDeviceCommand(KillProcessCommand);
 		}
 
-		protected void CopyOBBFiles(Dictionary<string, string> OBBFiles, string Package)
+		protected void CopyOBBFiles(UnrealAppConfig AppConfig, Dictionary<string, string> OBBFiles, string Package)
 		{
-			string OBBRemoteDestination = string.Format("{0}/Android/obb/{1}", StoragePath, Package);
+			bool bAFSEnablePlugin;
+			string AFSToken = "";
+			bool bIsShipping;
+			bool bAFSIncludeInShipping;
+			bool bAFSAllowExternalStartInShipping;
+			bool useADB = true;
+			if (AppConfig != null)
+			{
+				AndroidBuild Build = AppConfig.Build as AndroidBuild;
+				if (Build == null)
+				{
+					throw new AutomationException("Unsupported build type {0} for Android!", AppConfig.Build.GetType());
+				}
 
-			// Remove all existing obbs
-			RunAdbDeviceCommand(string.Format("shell rm {0}", OBBRemoteDestination));
-			RunAdbDeviceCommand(string.Format("shell mkdir -p {0}", OBBRemoteDestination));
+				if (UsingAndroidFileServer(AppConfig.ProjectFile, Build.Configuration, out bAFSEnablePlugin, out AFSToken, out bIsShipping, out bAFSIncludeInShipping, out bAFSAllowExternalStartInShipping))
+				{
+					useADB = false;
+				}
+			}
+
+			if (useADB)
+			{
+				string OBBRemoteDestination = string.Format("{0}/Android/obb/{1}", StoragePath, Package);
+
+				// Remove all existing obbs
+				RunAdbDeviceCommand(string.Format("shell rm {0}", OBBRemoteDestination));
+				RunAdbDeviceCommand(string.Format("shell mkdir -p {0}", OBBRemoteDestination));
+			}
+			else
+			{
+				RunAFSDeviceCommand(string.Format("-p {0} -k {1} rm {2}", Package, AFSToken, "^mainobb"));
+				RunAFSDeviceCommand(string.Format("-p {0} -k {1} rm {2}", Package, AFSToken, "^patchobb"));
+				for(int i = 1; i <= 32; i++)
+				{
+					string overflowobb = string.Format("^overflow{0}obb", i);
+					RunAFSDeviceCommand(string.Format("-p {0} -k {1} rm {2}", Package, AFSToken, overflowobb));
+				}
+			}
 
 			// obb files need to be named based on APK version (grrr), so find that out. This should return something like
 			// versionCode=2 minSdk=21 targetSdk=21
@@ -772,27 +1045,46 @@ namespace Gauntlet
 				string SourceFile = Pair.Key;
 				string DestinationFile = Pair.Value;
 
-				// If we installed a new APK we need to change the package version
-				Match OBBMatch = Regex.Match(SourceFile, @"\.(\d+)\.com.*\.obb");
-				if (OBBMatch.Success)
+				if (useADB)
 				{
-					string StrippedObb = Path.GetFileName(SourceFile.Replace(".Client.obb", ".obb").Replace(OBBMatch.Groups[1].ToString(), PackageVersion));
-					DestinationFile = StoragePath + "/Android/obb/" + Package + "/" + StrippedObb;
-				}
+					// If we installed a new APK we need to change the package version
+					Match OBBMatch = Regex.Match(SourceFile, @"\.(\d+)\.com.*\.obb");
+					if (OBBMatch.Success)
+					{
+						string StrippedObb = Path.GetFileName(SourceFile.Replace(".Client.obb", ".obb").Replace(OBBMatch.Groups[1].ToString(), PackageVersion));
+						DestinationFile = StoragePath + "/Android/obb/" + Package + "/" + StrippedObb;
+					}
 
-				DestinationFile = Regex.Replace(DestinationFile, "%STORAGE%", StoragePath, RegexOptions.IgnoreCase);
-				CopyFileToDevice(null, SourceFile, DestinationFile);
+					DestinationFile = Regex.Replace(DestinationFile, "%STORAGE%", StoragePath, RegexOptions.IgnoreCase);
+					CopyFileToDevice(AppConfig, null, SourceFile, DestinationFile);
+				}
+				else
+				{
+					string Filename = Path.GetFileName(SourceFile);
+					string[] parts = Filename.Split(new[] { "." }, StringSplitOptions.RemoveEmptyEntries).Select(S => S.Trim()).ToArray();
+					// Make sure last part is obb
+					if (parts[parts.Length - 1] == "obb")
+					{
+						string obbname = string.Format("^{0}obb", parts[0]);
+						CopyFileToDevice(AppConfig, Package, SourceFile, obbname);
+					}
+					else
+					{
+						DestinationFile = Regex.Replace(DestinationFile, "%STORAGE%", StoragePath, RegexOptions.IgnoreCase);
+						CopyFileToDevice(AppConfig, Package, SourceFile, DestinationFile);
+					}
+				}
 			}
 		}
 
-		protected void CopyCommandlineFile(string Commandline, string Package, string Project, bool bUseExternalFileDirectory)
+		protected void CopyCommandlineFile(UnrealAppConfig AppConfig, string Commandline, string Package, string Project, bool bUseExternalFileDirectory)
 		{
 			string ExternalStoragePath = StoragePath + "/UnrealGame/" + Project;
 			string ExternalFilesPath = StoragePath + "/Android/data/" + Package + "/files/UnrealGame/" + Project;
 			string DeviceBaseDir = bUseExternalFileDirectory ? ExternalFilesPath : ExternalStoragePath;
 			CommandLineFilePath = string.Format("{0}/UECommandLine.txt", DeviceBaseDir);
 
-			PostRunCleanup(); // function needs a rename in this context, just delete any old UECommandlines
+			PostRunCleanup(AppConfig, Package); // function needs a rename in this context, just delete any old UECommandlines
 
 			// Create a tempfile, insert the command line, and push it over
 			string CommandLineTmpFile = Path.GetTempFileName();
@@ -800,14 +1092,49 @@ namespace Gauntlet
 			// I've seen a weird thing where adb push truncates by a byte, so add some padding...
 			File.WriteAllText(CommandLineTmpFile, Commandline + "    ");
 
-			CopyFileToDevice(null, CommandLineTmpFile, CommandLineFilePath);
+			if (UseAFS(AppConfig))
+			{
+				CopyFileToDevice(AppConfig, Package, CommandLineTmpFile, "^commandfile");
+			}
+			else
+			{
+				CopyFileToDevice(AppConfig, Package, CommandLineTmpFile, CommandLineFilePath);
+			}
 
 			File.Delete(CommandLineTmpFile);
 		}
 
-		protected bool DeleteFileFromDevice(string DestPath)
+		protected bool DeleteFileFromDevice(UnrealAppConfig AppConfig, string Package, string DestPath)
 		{
-			var AdbResult = RunAdbDeviceCommand(string.Format("shell rm -f {0}", DestPath));
+			bool bAFSEnablePlugin;
+			string AFSToken = "";
+			bool bIsShipping;
+			bool bAFSIncludeInShipping;
+			bool bAFSAllowExternalStartInShipping;
+			bool useADB = true;
+			if (AppConfig != null)
+			{
+				AndroidBuild Build = AppConfig.Build as AndroidBuild;
+				if (Build == null)
+				{
+					throw new AutomationException("Unsupported build type {0} for Android!", AppConfig.Build.GetType());
+				}
+
+				if (UsingAndroidFileServer(AppConfig.ProjectFile, Build.Configuration, out bAFSEnablePlugin, out AFSToken, out bIsShipping, out bAFSIncludeInShipping, out bAFSAllowExternalStartInShipping))
+				{
+					useADB = false;
+				}
+			}
+
+			IProcessResult AdbResult;
+			if (useADB)
+			{
+				AdbResult = RunAdbDeviceCommand(string.Format("shell rm -f {0}", DestPath));
+			}
+			else
+			{
+				AdbResult = RunAFSDeviceCommand(string.Format("-p \"{0}\" -k \"{1}\" rm \"{2}\"", Package, AFSToken, DestPath));
+			}
 			return AdbResult.ExitCode == 0;
 		}
 
@@ -847,6 +1174,20 @@ namespace Gauntlet
 				throw new AutomationException("Invalid build for Android!");
 			}
 
+			bool bAFSEnablePlugin;
+			string AFSToken = "";
+			bool bIsShipping;
+			bool bAFSIncludeInShipping;
+			bool bAFSAllowExternalStartInShipping;
+			bool useADB = true;
+			if (AppConfig != null)
+			{
+				if (UsingAndroidFileServer(AppConfig.ProjectFile, Build.Configuration, out bAFSEnablePlugin, out AFSToken, out bIsShipping, out bAFSIncludeInShipping, out bAFSAllowExternalStartInShipping))
+				{
+					useADB = false;
+				}
+			}
+
 			// kill any currently running instance:
 			KillRunningProcess(Build.AndroidPackageName);
 
@@ -882,7 +1223,7 @@ namespace Gauntlet
 			{
 				Log.Info("Fully cleaning console before install...");
 				RunAdbDeviceCommand(string.Format("shell rm -r {0}/UnrealGame/*", StoragePath));
-				RunAdbDeviceCommand(string.Format("shell rm -r {0}/Android/data/{1}/*", StoragePath, Build.AndroidPackageName));
+				RunAdbDeviceCommand(string.Format("shell rm -r {0}/Android/data/{1}/files/*", StoragePath, Build.AndroidPackageName));
 				RunAdbDeviceCommand(string.Format("shell rm -r {0}/obb/{1}/*", StoragePath, Build.AndroidPackageName));
 				RunAdbDeviceCommand(string.Format("shell rm -r {0}/Android/obb/{1}/*", StoragePath, Build.AndroidPackageName));
 				RunAdbDeviceCommand(string.Format("shell rm -r {0}/Download/*", StoragePath));
@@ -907,10 +1248,6 @@ namespace Gauntlet
 					RunAdbDeviceCommand(string.Format("shell rm -r {0}", OBBRemoteDestination));
 				}
 
-				// remote dir on the device, create it if it doesn't exist
-				RunAdbDeviceCommand(string.Format("shell mkdir -p {0}/", DeviceExternalStorageSavedPath));
-				RunAdbDeviceCommand(string.Format("shell mkdir -p {0}/", DeviceExternalFilesSavedPath));
-
 				// check for a local newer executable
 				if (Globals.Params.ParseParam("dev"))
 				{
@@ -933,7 +1270,18 @@ namespace Gauntlet
 				}
 
 				// first install the APK
-				CopyFileToDevice(Build.AndroidPackageName, ApkPath, string.Empty);
+				CopyFileToDevice(AppConfig, Build.AndroidPackageName, ApkPath, string.Empty);
+
+				// remote dir on the device, create it if it doesn't exist
+				if (useADB)
+				{
+					RunAdbDeviceCommand(string.Format("shell mkdir -p {0}/", DeviceExternalStorageSavedPath));
+					RunAdbDeviceCommand(string.Format("shell mkdir -p {0}/", DeviceExternalFilesSavedPath));
+				}
+				else
+				{
+					RunAFSDeviceCommand(string.Format("-p {0} -k {1} mkdir \"^saved\"", Build.AndroidPackageName, AFSToken));
+				}
 			}
 
 			// Convert the files from the source to final destination names
@@ -1036,7 +1384,25 @@ namespace Gauntlet
 					string RemoteFile = KV.Value;
 
 					Console.WriteLine("Copying {0} to {1}", LocalFile, RemoteFile);
-					CopyFileToDevice(Build.AndroidPackageName, LocalFile, RemoteFile);
+					if (useADB)
+					{
+						CopyFileToDevice(AppConfig, Build.AndroidPackageName, LocalFile, RemoteFile);
+					}
+					else
+					{
+						string Filename = Path.GetFileName(LocalFile);
+						string[] parts = Filename.ToString().Split(new[] { "." }, StringSplitOptions.RemoveEmptyEntries).Select(S => S.Trim()).ToArray();
+						// Make sure last part is obb
+						if (parts[parts.Length - 1] == "obb")
+						{
+							string obbname = string.Format("^{0}obb", parts[0]);
+							CopyFileToDevice(AppConfig, Build.AndroidPackageName, LocalFile, obbname);
+						}
+						else
+						{
+							CopyFileToDevice(AppConfig, Build.AndroidPackageName, LocalFile, RemoteFile);
+						}
+					}
 				}
 			}
 			else
@@ -1044,7 +1410,7 @@ namespace Gauntlet
 				Log.Info("Skipping install of {0} (-skipdeploy)", Build.AndroidPackageName);
 			}
 
-			CopyCommandlineFile(AppConfig.CommandLine, Build.AndroidPackageName, AppConfig.ProjectName, Build.UsesExternalFilesDir);
+			CopyCommandlineFile(AppConfig, AppConfig.CommandLine, Build.AndroidPackageName, AppConfig.ProjectName, Build.UsesExternalFilesDir);
 
 			AndroidAppInstall AppInstall = new AndroidAppInstall(this, InApkPath: ApkPath, AppConfig.ProjectName, Build.AndroidPackageName, AppConfig.CommandLine);
 
@@ -1088,9 +1454,9 @@ namespace Gauntlet
 			AppTag = InAppTag;
 		}
 
-		public IAppInstance Run()
+		public IAppInstance Run(UnrealAppConfig AppConfiguration)
 		{
-			return AndroidDevice.Run(this);
+			return AndroidDevice.Run(AppConfiguration, this);
 		}
 	}
 
diff --git a/Engine/Source/Programs/AutomationTool/Gauntlet/Platform/IOS/Gauntlet.TargetDeviceIOS.cs b/Engine/Source/Programs/AutomationTool/Gauntlet/Platform/IOS/Gauntlet.TargetDeviceIOS.cs
--- a/Engine/Source/Programs/AutomationTool/Gauntlet/Platform/IOS/Gauntlet.TargetDeviceIOS.cs
+++ b/Engine/Source/Programs/AutomationTool/Gauntlet/Platform/IOS/Gauntlet.TargetDeviceIOS.cs
@@ -152,9 +152,9 @@ namespace Gauntlet
 			IOSDevice = InDevice;
 		}
 
-		public IAppInstance Run()
+		public IAppInstance Run(UnrealAppConfig AppConfiguration)
 		{
-			return Device.Run(this);
+			return Device.Run(AppConfiguration, this);
 		}
 	}
 
@@ -260,7 +260,7 @@ namespace Gauntlet
 
 		public CommandUtils.ERunOptions RunOptions { get; set; }
 
-		public IAppInstance Run(IAppInstall App)
+		public IAppInstance Run(UnrealAppConfig AppConfiguration, IAppInstall App)
 		{
 			IOSAppInstall IOSApp = App as IOSAppInstall;
 
diff --git a/Engine/Source/Programs/AutomationTool/Gauntlet/Platform/Linux/Gauntlet.TargetDeviceLinux.cs b/Engine/Source/Programs/AutomationTool/Gauntlet/Platform/Linux/Gauntlet.TargetDeviceLinux.cs
--- a/Engine/Source/Programs/AutomationTool/Gauntlet/Platform/Linux/Gauntlet.TargetDeviceLinux.cs
+++ b/Engine/Source/Programs/AutomationTool/Gauntlet/Platform/Linux/Gauntlet.TargetDeviceLinux.cs
@@ -31,7 +31,7 @@ namespace Gauntlet
 			RunOptions = CommandUtils.ERunOptions.NoWaitForExit | CommandUtils.ERunOptions.NoLoggingOfRunCommand;
 		}
 
-		public override IAppInstance Run(IAppInstall App)
+		public override IAppInstance Run(UnrealAppConfig AppConfiguration, IAppInstall App)
 		{
 			LinuxAppInstall LinuxApp = App as LinuxAppInstall;
 
diff --git a/Engine/Source/Programs/AutomationTool/Gauntlet/Platform/Mac/Gauntlet.TargetDeviceMac.cs b/Engine/Source/Programs/AutomationTool/Gauntlet/Platform/Mac/Gauntlet.TargetDeviceMac.cs
--- a/Engine/Source/Programs/AutomationTool/Gauntlet/Platform/Mac/Gauntlet.TargetDeviceMac.cs
+++ b/Engine/Source/Programs/AutomationTool/Gauntlet/Platform/Mac/Gauntlet.TargetDeviceMac.cs
@@ -96,7 +96,7 @@ namespace Gauntlet
 			return Install;
 		}
 
-		public override IAppInstance Run(IAppInstall App)
+		public override IAppInstance Run(UnrealAppConfig AppConfiguration, IAppInstall App)
 		{
 			MacAppInstall MacInstall = App as MacAppInstall;
 
diff --git a/Engine/Source/Programs/AutomationTool/Gauntlet/Platform/Null/Gauntlet.TargetDeviceNull.cs b/Engine/Source/Programs/AutomationTool/Gauntlet/Platform/Null/Gauntlet.TargetDeviceNull.cs
--- a/Engine/Source/Programs/AutomationTool/Gauntlet/Platform/Null/Gauntlet.TargetDeviceNull.cs
+++ b/Engine/Source/Programs/AutomationTool/Gauntlet/Platform/Null/Gauntlet.TargetDeviceNull.cs
@@ -71,9 +71,9 @@ namespace Gauntlet
 			CommandLine = InCommandLine;
 		}
 
-		public IAppInstance Run()
+		public IAppInstance Run(UnrealAppConfig AppConfiguration)
 		{
-			return Device.Run(this);
+			return Device.Run(AppConfiguration, this);
 		}
 	}
 
@@ -136,7 +136,7 @@ namespace Gauntlet
 
 		public CommandUtils.ERunOptions RunOptions { get; set; }
 
-		public IAppInstance Run(IAppInstall App)
+		public IAppInstance Run(UnrealAppConfig AppConfiguration, IAppInstall App)
 		{
 			NullAppInstall NullApp = App as NullAppInstall;
 
diff --git a/Engine/Source/Programs/AutomationTool/Gauntlet/Platform/Windows/Gauntlet.TargetDeviceWindows.cs b/Engine/Source/Programs/AutomationTool/Gauntlet/Platform/Windows/Gauntlet.TargetDeviceWindows.cs
--- a/Engine/Source/Programs/AutomationTool/Gauntlet/Platform/Windows/Gauntlet.TargetDeviceWindows.cs
+++ b/Engine/Source/Programs/AutomationTool/Gauntlet/Platform/Windows/Gauntlet.TargetDeviceWindows.cs
@@ -59,7 +59,7 @@ namespace Gauntlet
 			return Install;
 		}
 
-		public override IAppInstance Run(IAppInstall App)
+		public override IAppInstance Run(UnrealAppConfig AppConfiguration, IAppInstall App)
 		{
 			WindowsAppInstall WinApp = App as WindowsAppInstall;
 
diff --git a/Engine/Source/Programs/AutomationTool/Gauntlet/SelfTest/Unreal/Gauntlet.SelfTest.TestUnrealInstallAndRunBase.cs b/Engine/Source/Programs/AutomationTool/Gauntlet/SelfTest/Unreal/Gauntlet.SelfTest.TestUnrealInstallAndRunBase.cs
--- a/Engine/Source/Programs/AutomationTool/Gauntlet/SelfTest/Unreal/Gauntlet.SelfTest.TestUnrealInstallAndRunBase.cs
+++ b/Engine/Source/Programs/AutomationTool/Gauntlet/SelfTest/Unreal/Gauntlet.SelfTest.TestUnrealInstallAndRunBase.cs
@@ -106,7 +106,7 @@ namespace Gauntlet.SelfTest
 				}
 
 				// Run the app and wait for either a timeout or it to exit
-				CurrentTestCase.Instance = CurrentTestCase.Install.Run();
+				CurrentTestCase.Instance = CurrentTestCase.Install.Run(null);
 				CurrentTestCase.bHasStarted = true;
 				AppStartTime = DateTime.Now;
 				return;
diff --git a/Engine/Source/Programs/AutomationTool/Gauntlet/Unreal/Base/Gauntlet.UnrealSession.cs b/Engine/Source/Programs/AutomationTool/Gauntlet/Unreal/Base/Gauntlet.UnrealSession.cs
--- a/Engine/Source/Programs/AutomationTool/Gauntlet/Unreal/Base/Gauntlet.UnrealSession.cs
+++ b/Engine/Source/Programs/AutomationTool/Gauntlet/Unreal/Base/Gauntlet.UnrealSession.cs
@@ -442,7 +442,7 @@ namespace Gauntlet
 			try
 			{
 				Log.Info("Starting deferred {0} on {1}", Role, CurrentInstall.Device);
-				IAppInstance Instance = CurrentInstall.Run();
+				IAppInstance Instance = CurrentInstall.Run(null);
 				IDeviceUsageReporter.RecordStart(Instance.Device.Name, Instance.Device.Platform, IDeviceUsageReporter.EventType.Test);
 
 				if (Instance != null || Globals.CancelSignalled)
@@ -1360,6 +1360,7 @@ namespace Gauntlet
 				}
 
 				Dictionary<IAppInstall, UnrealSessionRole> InstallsToRoles = new Dictionary<IAppInstall, UnrealSessionRole>();
+				Dictionary<IAppInstall, UnrealAppConfig> InstallsToConfig = new Dictionary<IAppInstall, UnrealAppConfig>();
 
 				// create a copy of our list
 				IEnumerable<ITargetDevice> DevicesToInstallOn = UnrealDeviceReservation.ReservedDevices.ToArray();
@@ -1444,6 +1445,7 @@ namespace Gauntlet
 						Role.ConfigureDevice?.Invoke(Device);
 
 						InstallsToRoles[Install] = Role;
+						InstallsToConfig[Install] = AppConfig;
 
 						if (ReinstallPerPass)
 						{
@@ -1454,6 +1456,7 @@ namespace Gauntlet
 					{
 						Install = RolesToInstalls[Role];
 						InstallsToRoles[Install] = Role;
+						InstallsToConfig[Install] = AppConfig;
 						Log.Info("Using previous install of {0} on {1}", Install.Name, Install.Device.Name);
 					}
 				}
@@ -1497,7 +1500,7 @@ namespace Gauntlet
 						try
 						{
 							Log.Info("Starting {0} on {1}", InstallRoleKV.Value, CurrentInstall.Device);
-							IAppInstance Instance = CurrentInstall.Run();
+							IAppInstance Instance = CurrentInstall.Run(InstallsToConfig[CurrentInstall]);
 							IDeviceUsageReporter.RecordStart(Instance.Device.Name, Instance.Device.Platform, IDeviceUsageReporter.EventType.Test);
 
 							if (Instance != null || Globals.CancelSignalled)
@@ -1745,6 +1748,7 @@ namespace Gauntlet
 			{
 				UnrealSessionRole Role = RoleInstall.Key;
 				IAppInstall Install = RoleInstall.Value;
+				UnrealAppConfig Config = RolesToConfigs[Role];
 
 				// InstallOnly roles don't execute a process
 				if (Role.InstallOnly)
@@ -1766,7 +1770,7 @@ namespace Gauntlet
 				try
 				{
 					Log.Info("Launching {Install} on {Device}", Install, RolesToDevices[Role]);
-					IAppInstance AppInstance = Install.Run();
+					IAppInstance AppInstance = Install.Run(Config);
 
 					if (AppInstance == null)
 					{
diff --git a/Engine/Source/Programs/AutomationTool/Gauntlet/Unreal/Utils/Gauntlet.InstallUnreal.cs b/Engine/Source/Programs/AutomationTool/Gauntlet/Unreal/Utils/Gauntlet.InstallUnreal.cs
--- a/Engine/Source/Programs/AutomationTool/Gauntlet/Unreal/Utils/Gauntlet.InstallUnreal.cs
+++ b/Engine/Source/Programs/AutomationTool/Gauntlet/Unreal/Utils/Gauntlet.InstallUnreal.cs
@@ -103,7 +103,7 @@ namespace Gauntlet
 				Log.Info("Installing build on device {0}", Device.Name);
 
 				IAppInstall Install = Device.InstallApplication(Config);
-				Device.Run(Install);
+				Device.Run(Config, Install);
 
 				TimeSpan Elapsed = (DateTime.Now - StartTime);
 				Log.Info("Installed on device {0} in {1:D2}m:{2:D2}s", Device.Name, Elapsed.Minutes, Elapsed.Seconds);
diff --git a/Engine/Source/Programs/AutomationTool/LowLevelTests/RunLowLevelTests.cs b/Engine/Source/Programs/AutomationTool/LowLevelTests/RunLowLevelTests.cs
--- a/Engine/Source/Programs/AutomationTool/LowLevelTests/RunLowLevelTests.cs
+++ b/Engine/Source/Programs/AutomationTool/LowLevelTests/RunLowLevelTests.cs
@@ -425,7 +425,7 @@ namespace LowLevelTests
 					DeviceWithStateOptions.CachedStateRefresh = QUERY_STATE_INTERVAL;
 				}
 
-				Instance = Device.Run(Install);
+				Instance = Device.Run(AppConfig, Install);
 				IDeviceUsageReporter.RecordStart(Instance.Device.Name, Instance.Device.Platform, IDeviceUsageReporter.EventType.Test);
 			}
 			catch (DeviceException DeviceEx)
@@ -505,7 +505,7 @@ namespace LowLevelTests
 						DeviceWithStateOptions.CachedStateRefresh = QUERY_STATE_INTERVAL;
 					}
 
-					Instance = Device.Run(Install);
+					Instance = Device.Run(AppConfig, Install);
 					IDeviceUsageReporter.RecordStart(Instance.Device.Name, (UnrealTargetPlatform)Instance.Device.Platform, IDeviceUsageReporter.EventType.Test);
 					RunSuccess = true;
 				}
diff --git a/Engine/Source/Programs/AutomationTool/Turnkey/Commands/InstallBuild.cs b/Engine/Source/Programs/AutomationTool/Turnkey/Commands/InstallBuild.cs
--- a/Engine/Source/Programs/AutomationTool/Turnkey/Commands/InstallBuild.cs
+++ b/Engine/Source/Programs/AutomationTool/Turnkey/Commands/InstallBuild.cs
@@ -128,7 +128,7 @@ namespace Turnkey.Commands
 				if (GauntletDevice != null)
 				{
 					IAppInstall Install = GauntletDevice.InstallApplication(Config);
-					GauntletDevice.Run(Install);
+					GauntletDevice.Run(Config, Install);
 				}
 			}
 		}
