Avoid emitting warnings on oplog creation when exporting a snapshot to a zen destination storage server.  oplog creation can fail when the oplog already exists, and this should not be a warning.

diff --git a/Engine/Source/Programs/AutomationTool/BuildGraph/Tasks/ZenExportSnapshotTask.cs b/Engine/Source/Programs/AutomationTool/BuildGraph/Tasks/ZenExportSnapshotTask.cs
--- a/Engine/Source/Programs/AutomationTool/BuildGraph/Tasks/ZenExportSnapshotTask.cs
+++ b/Engine/Source/Programs/AutomationTool/BuildGraph/Tasks/ZenExportSnapshotTask.cs
@@ -382,7 +382,7 @@ namespace AutomationTool.Tasks
 			Writer.WriteObjectEnd();
 		}
 
-		private bool TryRunAndLogWithoutSpew(string App, string CommandLine)
+		private bool TryRunAndLogWithoutSpew(string App, string CommandLine, bool IgnoreFailure)
 		{
 			ProcessResult.SpewFilterCallbackType SilentOutputFilter = new ProcessResult.SpewFilterCallbackType(Line =>
 				{
@@ -394,7 +394,10 @@ namespace AutomationTool.Tasks
 			}
 			catch (CommandUtils.CommandFailedException e)
 			{
-				Logger.LogWarning(e.ToString());
+				if (!IgnoreFailure)
+				{
+					Logger.LogWarning(e.ToString());
+				}
 				return false;
 			}
 			return true;
@@ -406,7 +409,7 @@ namespace AutomationTool.Tasks
 			int Attempt = 0;
 			while (Attempt < AttemptLimit)
 			{
-				if (TryRunAndLogWithoutSpew(App, CommandLine))
+				if (TryRunAndLogWithoutSpew(App, CommandLine, false))
 				{
 					return true;
 				}
@@ -592,7 +595,7 @@ namespace AutomationTool.Tasks
 
 					StringBuilder CreateProjectCommandline = new StringBuilder();
 					CreateProjectCommandline.AppendFormat("project-create --hosturl {0} {1}", Parameters.DestinationZenHost, ProjectName);
-					TryRunAndLogWithoutSpew(ZenExe.FullName, CreateProjectCommandline.ToString());
+					TryRunAndLogWithoutSpew(ZenExe.FullName, CreateProjectCommandline.ToString(), true);
 
 					OplogExportCommandline.AppendFormat(" --zen {0}", Parameters.DestinationZenHost);
 
