Instead of passing in AppConfig to Run, store AppConfig in AndroidInstall.

diff --git a/Engine/Source/Programs/AutomationTool/Gauntlet/Framework/Base/Gauntlet.AppInstall.cs b/Engine/Source/Programs/AutomationTool/Gauntlet/Framework/Base/Gauntlet.AppInstall.cs
--- a/Engine/Source/Programs/AutomationTool/Gauntlet/Framework/Base/Gauntlet.AppInstall.cs
+++ b/Engine/Source/Programs/AutomationTool/Gauntlet/Framework/Base/Gauntlet.AppInstall.cs
@@ -20,7 +20,7 @@ namespace Gauntlet
 		/// Create a running instance of the app
 		/// </summary>
 		/// <returns></returns>
-		IAppInstance Run(UnrealAppConfig AppConfiguration);
+		IAppInstance Run();
 
 		/// <summary>
 		/// The device that we'll be run on
diff --git a/Engine/Source/Programs/AutomationTool/Gauntlet/Framework/Base/Gauntlet.TargetDevice.cs b/Engine/Source/Programs/AutomationTool/Gauntlet/Framework/Base/Gauntlet.TargetDevice.cs
--- a/Engine/Source/Programs/AutomationTool/Gauntlet/Framework/Base/Gauntlet.TargetDevice.cs
+++ b/Engine/Source/Programs/AutomationTool/Gauntlet/Framework/Base/Gauntlet.TargetDevice.cs
@@ -87,7 +87,7 @@ namespace Gauntlet
 
 		IAppInstall InstallApplication(UnrealAppConfig AppConfiguration);
 
-		IAppInstance Run(UnrealAppConfig AppConfiguration, IAppInstall App);
+		IAppInstance Run(IAppInstall App);
 
 		/// Begin new flow ///
 
diff --git a/Engine/Source/Programs/AutomationTool/Gauntlet/Framework/Base/Gauntlet.TargetDeviceDesktopCommon.cs b/Engine/Source/Programs/AutomationTool/Gauntlet/Framework/Base/Gauntlet.TargetDeviceDesktopCommon.cs
--- a/Engine/Source/Programs/AutomationTool/Gauntlet/Framework/Base/Gauntlet.TargetDeviceDesktopCommon.cs
+++ b/Engine/Source/Programs/AutomationTool/Gauntlet/Framework/Base/Gauntlet.TargetDeviceDesktopCommon.cs
@@ -223,7 +223,7 @@ namespace Gauntlet
             LocalDirectoryMappings.Add(EIntendedBaseCopyDirectory.Saved, Path.Combine(BaseDirectory, "Saved"));
 		}
 
-		public abstract IAppInstance Run(UnrealAppConfig AppConfiguration, IAppInstall Install);
+		public abstract IAppInstance Run(IAppInstall Install);
 
 		protected abstract IAppInstall CreateNativeStagedInstall(UnrealAppConfig AppConfig, NativeStagedBuild Build);
 
@@ -323,9 +323,9 @@ namespace Gauntlet
 			CanAlterCommandArgs = true;
 		}
 
-		public virtual IAppInstance Run(UnrealAppConfig AppConfiguration)
+		public virtual IAppInstance Run()
 		{
-			return Device.Run(AppConfiguration, this);
+			return Device.Run(this);
 		}
 
 		/// <summary>
diff --git a/Engine/Source/Programs/AutomationTool/Gauntlet/Platform/Android/Gauntlet.TargetDeviceAndroid.cs b/Engine/Source/Programs/AutomationTool/Gauntlet/Platform/Android/Gauntlet.TargetDeviceAndroid.cs
--- a/Engine/Source/Programs/AutomationTool/Gauntlet/Platform/Android/Gauntlet.TargetDeviceAndroid.cs
+++ b/Engine/Source/Programs/AutomationTool/Gauntlet/Platform/Android/Gauntlet.TargetDeviceAndroid.cs
@@ -438,7 +438,7 @@ namespace Gauntlet
 
 			PopulateDirectoryMappings(Path.GetDirectoryName(DeviceArtifactPath));
 
-			return new AndroidAppInstall(this, AppConfig.ProjectName, Build.AndroidPackageName, AppConfig.CommandLine);
+			return new AndroidAppInstall(this, AppConfig.ProjectName, Build.AndroidPackageName, AppConfig.CommandLine, AppConfig.ProjectFile, AppConfig.Build.Configuration);
 		}
 
 		public void CopyAdditionalFiles(UnrealAppConfig AppConfig, IEnumerable<UnrealFileToCopy> FilesToCopy)
@@ -463,31 +463,22 @@ namespace Gauntlet
 			CopyAdditionalFiles(null, FilesToCopy);
 		}
 
-		public IAppInstance Run(UnrealAppConfig AppConfig, IAppInstall App)
+		public IAppInstance Run(IAppInstall App)
 		{
+			if (App is not AndroidAppInstall Install)
+			{
+				throw new AutomationException("AppInstall is of incorrect type {0}! Must be of type AndroidAppInstall", App.GetType().Name);
+			}
+
 			bool bAFSEnablePlugin;
 			string AFSToken = "";
 			bool bIsShipping;
 			bool bAFSIncludeInShipping;
 			bool bAFSAllowExternalStartInShipping;
 			bool useADB = true;
-			if (AppConfig != null)
+			if (UsingAndroidFileServer(Install.ProjectFile, Install.Configuration, out bAFSEnablePlugin, out AFSToken, out bIsShipping, out bAFSIncludeInShipping, out bAFSAllowExternalStartInShipping))
 			{
-				AndroidBuild Build = AppConfig.Build as AndroidBuild;
-				if (Build == null)
-				{
-					throw new AutomationException("Unsupported build type {0} for Android!", AppConfig.Build.GetType());
-				}
-
-				if (UsingAndroidFileServer(AppConfig.ProjectFile, Build.Configuration, out bAFSEnablePlugin, out AFSToken, out bIsShipping, out bAFSIncludeInShipping, out bAFSAllowExternalStartInShipping))
-				{
-					useADB = false;
-				}
-			}
-
-			if (App is not AndroidAppInstall Install)
-			{
-				throw new AutomationException("AppInstall is of incorrect type {0}! Must be of type AndroidAppInstall", App.GetType().Name);
+				useADB = false;
 			}
 
 			if(!IsOn)
@@ -1412,7 +1403,7 @@ namespace Gauntlet
 
 			CopyCommandlineFile(AppConfig, AppConfig.CommandLine, Build.AndroidPackageName, AppConfig.ProjectName, Build.UsesExternalFilesDir);
 
-			AndroidAppInstall AppInstall = new AndroidAppInstall(this, InApkPath: ApkPath, AppConfig.ProjectName, Build.AndroidPackageName, AppConfig.CommandLine);
+			AndroidAppInstall AppInstall = new AndroidAppInstall(this, InApkPath: ApkPath, AppConfig.ProjectName, Build.AndroidPackageName, AppConfig.CommandLine, AppConfig.ProjectFile, Build.Configuration);
 
 			return AppInstall;
 		}
@@ -1435,16 +1426,22 @@ namespace Gauntlet
 
 		public string ApkPath { get; protected set; }
 
-		public AndroidAppInstall(TargetDeviceAndroid InDevice, string InName, string InAndroidPackageName, string InCommandLine, string InAppTag = "UE")
+		public FileReference ProjectFile { get; protected set; }
+
+		public UnrealTargetConfiguration Configuration { get; protected set; }
+
+		public AndroidAppInstall(TargetDeviceAndroid InDevice, string InName, string InAndroidPackageName, string InCommandLine, FileReference InProjectFile, UnrealTargetConfiguration InConfiguration, string InAppTag = "UE")
 		{
 			AndroidDevice = InDevice;
 			Name = InName;
 			AndroidPackageName = InAndroidPackageName;
 			CommandLine = InCommandLine;
 			AppTag = InAppTag;
+			ProjectFile = InProjectFile;
+			Configuration = InConfiguration;
 		}
 
-		public AndroidAppInstall(TargetDeviceAndroid InDevice, string InApkPath, string InName, string InAndroidPackageName, string InCommandLine, string InAppTag = "UE")
+		public AndroidAppInstall(TargetDeviceAndroid InDevice, string InApkPath, string InName, string InAndroidPackageName, string InCommandLine, FileReference InProjectFile, UnrealTargetConfiguration InConfiguration, string InAppTag = "UE")
 		{
 			AndroidDevice = InDevice;
 			ApkPath = InApkPath;
@@ -1452,11 +1449,13 @@ namespace Gauntlet
 			AndroidPackageName = InAndroidPackageName;
 			CommandLine = InCommandLine;
 			AppTag = InAppTag;
+			ProjectFile = InProjectFile;
+			Configuration = InConfiguration;
 		}
 
-		public IAppInstance Run(UnrealAppConfig AppConfiguration)
+		public IAppInstance Run()
 		{
-			return AndroidDevice.Run(AppConfiguration, this);
+			return AndroidDevice.Run(this);
 		}
 	}
 
diff --git a/Engine/Source/Programs/AutomationTool/Gauntlet/Platform/IOS/Gauntlet.TargetDeviceIOS.cs b/Engine/Source/Programs/AutomationTool/Gauntlet/Platform/IOS/Gauntlet.TargetDeviceIOS.cs
--- a/Engine/Source/Programs/AutomationTool/Gauntlet/Platform/IOS/Gauntlet.TargetDeviceIOS.cs
+++ b/Engine/Source/Programs/AutomationTool/Gauntlet/Platform/IOS/Gauntlet.TargetDeviceIOS.cs
@@ -152,9 +152,9 @@ namespace Gauntlet
 			IOSDevice = InDevice;
 		}
 
-		public IAppInstance Run(UnrealAppConfig AppConfiguration)
+		public IAppInstance Run()
 		{
-			return Device.Run(AppConfiguration, this);
+			return Device.Run(this);
 		}
 	}
 
@@ -260,7 +260,7 @@ namespace Gauntlet
 
 		public CommandUtils.ERunOptions RunOptions { get; set; }
 
-		public IAppInstance Run(UnrealAppConfig AppConfiguration, IAppInstall App)
+		public IAppInstance Run(IAppInstall App)
 		{
 			IOSAppInstall IOSApp = App as IOSAppInstall;
 
diff --git a/Engine/Source/Programs/AutomationTool/Gauntlet/Platform/Linux/Gauntlet.TargetDeviceLinux.cs b/Engine/Source/Programs/AutomationTool/Gauntlet/Platform/Linux/Gauntlet.TargetDeviceLinux.cs
--- a/Engine/Source/Programs/AutomationTool/Gauntlet/Platform/Linux/Gauntlet.TargetDeviceLinux.cs
+++ b/Engine/Source/Programs/AutomationTool/Gauntlet/Platform/Linux/Gauntlet.TargetDeviceLinux.cs
@@ -31,7 +31,7 @@ namespace Gauntlet
 			RunOptions = CommandUtils.ERunOptions.NoWaitForExit | CommandUtils.ERunOptions.NoLoggingOfRunCommand;
 		}
 
-		public override IAppInstance Run(UnrealAppConfig AppConfiguration, IAppInstall App)
+		public override IAppInstance Run(IAppInstall App)
 		{
 			LinuxAppInstall LinuxApp = App as LinuxAppInstall;
 
diff --git a/Engine/Source/Programs/AutomationTool/Gauntlet/Platform/Mac/Gauntlet.TargetDeviceMac.cs b/Engine/Source/Programs/AutomationTool/Gauntlet/Platform/Mac/Gauntlet.TargetDeviceMac.cs
--- a/Engine/Source/Programs/AutomationTool/Gauntlet/Platform/Mac/Gauntlet.TargetDeviceMac.cs
+++ b/Engine/Source/Programs/AutomationTool/Gauntlet/Platform/Mac/Gauntlet.TargetDeviceMac.cs
@@ -96,7 +96,7 @@ namespace Gauntlet
 			return Install;
 		}
 
-		public override IAppInstance Run(UnrealAppConfig AppConfiguration, IAppInstall App)
+		public override IAppInstance Run(IAppInstall App)
 		{
 			MacAppInstall MacInstall = App as MacAppInstall;
 
diff --git a/Engine/Source/Programs/AutomationTool/Gauntlet/Platform/Null/Gauntlet.TargetDeviceNull.cs b/Engine/Source/Programs/AutomationTool/Gauntlet/Platform/Null/Gauntlet.TargetDeviceNull.cs
--- a/Engine/Source/Programs/AutomationTool/Gauntlet/Platform/Null/Gauntlet.TargetDeviceNull.cs
+++ b/Engine/Source/Programs/AutomationTool/Gauntlet/Platform/Null/Gauntlet.TargetDeviceNull.cs
@@ -71,9 +71,9 @@ namespace Gauntlet
 			CommandLine = InCommandLine;
 		}
 
-		public IAppInstance Run(UnrealAppConfig AppConfiguration)
+		public IAppInstance Run()
 		{
-			return Device.Run(AppConfiguration, this);
+			return Device.Run(this);
 		}
 	}
 
@@ -136,7 +136,7 @@ namespace Gauntlet
 
 		public CommandUtils.ERunOptions RunOptions { get; set; }
 
-		public IAppInstance Run(UnrealAppConfig AppConfiguration, IAppInstall App)
+		public IAppInstance Run(IAppInstall App)
 		{
 			NullAppInstall NullApp = App as NullAppInstall;
 
diff --git a/Engine/Source/Programs/AutomationTool/Gauntlet/Platform/Windows/Gauntlet.TargetDeviceWindows.cs b/Engine/Source/Programs/AutomationTool/Gauntlet/Platform/Windows/Gauntlet.TargetDeviceWindows.cs
--- a/Engine/Source/Programs/AutomationTool/Gauntlet/Platform/Windows/Gauntlet.TargetDeviceWindows.cs
+++ b/Engine/Source/Programs/AutomationTool/Gauntlet/Platform/Windows/Gauntlet.TargetDeviceWindows.cs
@@ -59,7 +59,7 @@ namespace Gauntlet
 			return Install;
 		}
 
-		public override IAppInstance Run(UnrealAppConfig AppConfiguration, IAppInstall App)
+		public override IAppInstance Run(IAppInstall App)
 		{
 			WindowsAppInstall WinApp = App as WindowsAppInstall;
 
diff --git a/Engine/Source/Programs/AutomationTool/Gauntlet/SelfTest/Unreal/Gauntlet.SelfTest.TestUnrealInstallAndRunBase.cs b/Engine/Source/Programs/AutomationTool/Gauntlet/SelfTest/Unreal/Gauntlet.SelfTest.TestUnrealInstallAndRunBase.cs
--- a/Engine/Source/Programs/AutomationTool/Gauntlet/SelfTest/Unreal/Gauntlet.SelfTest.TestUnrealInstallAndRunBase.cs
+++ b/Engine/Source/Programs/AutomationTool/Gauntlet/SelfTest/Unreal/Gauntlet.SelfTest.TestUnrealInstallAndRunBase.cs
@@ -106,7 +106,7 @@ namespace Gauntlet.SelfTest
 				}
 
 				// Run the app and wait for either a timeout or it to exit
-				CurrentTestCase.Instance = CurrentTestCase.Install.Run(null);
+				CurrentTestCase.Instance = CurrentTestCase.Install.Run();
 				CurrentTestCase.bHasStarted = true;
 				AppStartTime = DateTime.Now;
 				return;
diff --git a/Engine/Source/Programs/AutomationTool/Gauntlet/Unreal/Base/Gauntlet.UnrealSession.cs b/Engine/Source/Programs/AutomationTool/Gauntlet/Unreal/Base/Gauntlet.UnrealSession.cs
--- a/Engine/Source/Programs/AutomationTool/Gauntlet/Unreal/Base/Gauntlet.UnrealSession.cs
+++ b/Engine/Source/Programs/AutomationTool/Gauntlet/Unreal/Base/Gauntlet.UnrealSession.cs
@@ -442,7 +442,7 @@ namespace Gauntlet
 			try
 			{
 				Log.Info("Starting deferred {0} on {1}", Role, CurrentInstall.Device);
-				IAppInstance Instance = CurrentInstall.Run(null);
+				IAppInstance Instance = CurrentInstall.Run();
 				IDeviceUsageReporter.RecordStart(Instance.Device.Name, Instance.Device.Platform, IDeviceUsageReporter.EventType.Test);
 
 				if (Instance != null || Globals.CancelSignalled)
@@ -1500,7 +1500,7 @@ namespace Gauntlet
 						try
 						{
 							Log.Info("Starting {0} on {1}", InstallRoleKV.Value, CurrentInstall.Device);
-							IAppInstance Instance = CurrentInstall.Run(InstallsToConfig[CurrentInstall]);
+							IAppInstance Instance = CurrentInstall.Run();
 							IDeviceUsageReporter.RecordStart(Instance.Device.Name, Instance.Device.Platform, IDeviceUsageReporter.EventType.Test);
 
 							if (Instance != null || Globals.CancelSignalled)
@@ -1770,7 +1770,7 @@ namespace Gauntlet
 				try
 				{
 					Log.Info("Launching {Install} on {Device}", Install, RolesToDevices[Role]);
-					IAppInstance AppInstance = Install.Run(Config);
+					IAppInstance AppInstance = Install.Run();
 
 					if (AppInstance == null)
 					{
diff --git a/Engine/Source/Programs/AutomationTool/Gauntlet/Unreal/Utils/Gauntlet.InstallUnreal.cs b/Engine/Source/Programs/AutomationTool/Gauntlet/Unreal/Utils/Gauntlet.InstallUnreal.cs
--- a/Engine/Source/Programs/AutomationTool/Gauntlet/Unreal/Utils/Gauntlet.InstallUnreal.cs
+++ b/Engine/Source/Programs/AutomationTool/Gauntlet/Unreal/Utils/Gauntlet.InstallUnreal.cs
@@ -103,7 +103,7 @@ namespace Gauntlet
 				Log.Info("Installing build on device {0}", Device.Name);
 
 				IAppInstall Install = Device.InstallApplication(Config);
-				Device.Run(Config, Install);
+				Device.Run(Install);
 
 				TimeSpan Elapsed = (DateTime.Now - StartTime);
 				Log.Info("Installed on device {0} in {1:D2}m:{2:D2}s", Device.Name, Elapsed.Minutes, Elapsed.Seconds);
diff --git a/Engine/Source/Programs/AutomationTool/LowLevelTests/RunLowLevelTests.cs b/Engine/Source/Programs/AutomationTool/LowLevelTests/RunLowLevelTests.cs
--- a/Engine/Source/Programs/AutomationTool/LowLevelTests/RunLowLevelTests.cs
+++ b/Engine/Source/Programs/AutomationTool/LowLevelTests/RunLowLevelTests.cs
@@ -425,7 +425,7 @@ namespace LowLevelTests
 					DeviceWithStateOptions.CachedStateRefresh = QUERY_STATE_INTERVAL;
 				}
 
-				Instance = Device.Run(AppConfig, Install);
+				Instance = Device.Run(Install);
 				IDeviceUsageReporter.RecordStart(Instance.Device.Name, Instance.Device.Platform, IDeviceUsageReporter.EventType.Test);
 			}
 			catch (DeviceException DeviceEx)
@@ -505,7 +505,7 @@ namespace LowLevelTests
 						DeviceWithStateOptions.CachedStateRefresh = QUERY_STATE_INTERVAL;
 					}
 
-					Instance = Device.Run(AppConfig, Install);
+					Instance = Device.Run(Install);
 					IDeviceUsageReporter.RecordStart(Instance.Device.Name, (UnrealTargetPlatform)Instance.Device.Platform, IDeviceUsageReporter.EventType.Test);
 					RunSuccess = true;
 				}
diff --git a/Engine/Source/Programs/AutomationTool/Turnkey/Commands/InstallBuild.cs b/Engine/Source/Programs/AutomationTool/Turnkey/Commands/InstallBuild.cs
--- a/Engine/Source/Programs/AutomationTool/Turnkey/Commands/InstallBuild.cs
+++ b/Engine/Source/Programs/AutomationTool/Turnkey/Commands/InstallBuild.cs
@@ -128,7 +128,7 @@ namespace Turnkey.Commands
 				if (GauntletDevice != null)
 				{
 					IAppInstall Install = GauntletDevice.InstallApplication(Config);
-					GauntletDevice.Run(Config, Install);
+					GauntletDevice.Run(Install);
 				}
 			}
 		}
