Re-order Ignore check so potential error does not fire inside the compare check.

diff --git a/Engine/Source/Programs/AutomationTool/BuildGraph/TempStorage.cs b/Engine/Source/Programs/AutomationTool/BuildGraph/TempStorage.cs
--- a/Engine/Source/Programs/AutomationTool/BuildGraph/TempStorage.cs
+++ b/Engine/Source/Programs/AutomationTool/BuildGraph/TempStorage.cs
@@ -960,7 +960,7 @@ namespace AutomationTool
 				foreach(TempStorageFile File in Manifest.Files)
 				{
 					string Message;
-					if (!File.Compare(RootDir, out Message) && !IgnoreModified.Matches(File.RelativePath))
+					if (!IgnoreModified.Matches(File.RelativePath) && !File.Compare(RootDir, out Message))
 					{
 						ModifiedFileMessages.Add(Message);
 					}
