[Gauntlet] Adds BuildSource improvements that were made to PS4BuildSources to PS5BuildSources.
Moves shared logic to SonyBuildSources
Packages can be bundled with configurations of other .self files to form packaged builds with different configuration (ex. shipping package + development .self = packaged development build)
Local development binaries will be able to be used with packaged and staged builds (-dev). The actual -dev logic will be added in a follow up CL

diff --git a/Engine/Source/Programs/AutomationTool/Gauntlet/Unreal/Utils/Gauntlet.UnrealHelpers.cs b/Engine/Source/Programs/AutomationTool/Gauntlet/Unreal/Utils/Gauntlet.UnrealHelpers.cs
--- a/Engine/Source/Programs/AutomationTool/Gauntlet/Unreal/Utils/Gauntlet.UnrealHelpers.cs
+++ b/Engine/Source/Programs/AutomationTool/Gauntlet/Unreal/Utils/Gauntlet.UnrealHelpers.cs
@@ -241,13 +241,13 @@ namespace Gauntlet
 					LocalAddress = PreferredInterface.Address;
 				}
 			}
-	
+
 			string HostIP = Globals.Params.ParseValue("hostip", "");
 			HostIP = string.IsNullOrEmpty(HostIP) ? LocalAddress.ToString() : HostIP;
 			return HostIP;
 		}
 
-		static public string GetExecutableName(string ProjectName, UnrealTargetPlatform Platform, UnrealTargetConfiguration Config, UnrealTargetRole Role, string Extension)
+		static public string GetExecutableName(string ProjectName, UnrealTargetPlatform Platform, UnrealTargetConfiguration Config, UnrealTargetRole Role, string Flavor, string Extension)
 		{
 			string ExeName = ProjectName;
 
@@ -265,6 +265,11 @@ namespace Gauntlet
 				ExeName += string.Format("-{0}-{1}", Platform, Config);
 			}
 
+			if(!string.IsNullOrEmpty(Flavor))
+			{
+				ExeName += string.Format("-{0}", Flavor);
+			}
+
 			// todo , how to find this?
 			if (Platform == UnrealTargetPlatform.Android)
 			{
@@ -285,6 +290,11 @@ namespace Gauntlet
 			return ExeName;
 		}
 
+		static public string GetExecutableName(string ProjectName, UnrealTargetPlatform Platform, UnrealTargetConfiguration Config, UnrealTargetRole Role, string Extension)
+		{
+			return GetExecutableName(ProjectName, Platform, Config, Role, string.Empty, Extension);
+		}
+
 		public class ConfigInfo
 		{
 			public UnrealTargetRole 			RoleType;
