BuildGraph: Enable nullable annotations for some BuildGraph files.

diff --git a/Engine/Source/Programs/AutomationTool/BuildGraph/BgCondition.cs b/Engine/Source/Programs/AutomationTool/BuildGraph/BgCondition.cs
--- a/Engine/Source/Programs/AutomationTool/BuildGraph/BgCondition.cs
+++ b/Engine/Source/Programs/AutomationTool/BuildGraph/BgCondition.cs
@@ -8,6 +8,8 @@ using System.Linq;
 using System.Threading.Tasks;
 using UnrealBuildBase;
 
+#nullable enable
+
 namespace AutomationTool
 {
 	/// <summary>
diff --git a/Engine/Source/Programs/AutomationTool/BuildGraph/BgGraphBuilder.cs b/Engine/Source/Programs/AutomationTool/BuildGraph/BgGraphBuilder.cs
--- a/Engine/Source/Programs/AutomationTool/BuildGraph/BgGraphBuilder.cs
+++ b/Engine/Source/Programs/AutomationTool/BuildGraph/BgGraphBuilder.cs
@@ -8,6 +8,8 @@ using System;
 using System.Collections.Generic;
 using System.Text;
 
+#nullable enable
+
 namespace AutomationTool
 {
 	/// <summary>
diff --git a/Engine/Source/Programs/AutomationTool/BuildGraph/BgNodeExecutor.cs b/Engine/Source/Programs/AutomationTool/BuildGraph/BgNodeExecutor.cs
--- a/Engine/Source/Programs/AutomationTool/BuildGraph/BgNodeExecutor.cs
+++ b/Engine/Source/Programs/AutomationTool/BuildGraph/BgNodeExecutor.cs
@@ -370,7 +370,7 @@ namespace AutomationTool
 			{
 				using (IScope Scope = GlobalTracer.Instance.BuildSpan("Task").WithTag("resource", _boundTasks[Idx].GetTraceName()).StartActive())
 				{
-					ITaskExecutor Executor = _boundTasks[Idx].GetExecutor();
+					ITaskExecutor? Executor = _boundTasks[Idx].GetExecutor();
 					if (Executor == null)
 					{
 						// Execute this task directly
@@ -382,10 +382,13 @@ namespace AutomationTool
 						catch (Exception Ex)
 						{
 							ExceptionUtils.AddContext(Ex, "while executing task {0}", _boundTasks[Idx].GetTraceString());
-							if (_boundTasks[Idx].SourceLocation != null)
+
+							BgScriptLocation? sourceLocation = _boundTasks[Idx].SourceLocation;
+							if (sourceLocation != null)
 							{
-								ExceptionUtils.AddContext(Ex, "at {0}({1})", _boundTasks[Idx].SourceLocation.File, _boundTasks[Idx].SourceLocation.LineNumber);
+								ExceptionUtils.AddContext(Ex, "at {0}({1})", sourceLocation.File, sourceLocation.LineNumber);
 							}
+
 							throw;
 						}
 					}
@@ -410,10 +413,13 @@ namespace AutomationTool
 							{
 								ExceptionUtils.AddContext(Ex, "while executing {0}", _boundTasks[TaskIdx].GetTraceString());
 							}
-							if (_boundTasks[FirstIdx].SourceLocation != null)
+
+							BgScriptLocation? sourceLocation = _boundTasks[FirstIdx].SourceLocation;
+							if (sourceLocation != null)
 							{
-								ExceptionUtils.AddContext(Ex, "at {0}({1})", _boundTasks[FirstIdx].SourceLocation.File, _boundTasks[FirstIdx].SourceLocation.LineNumber);
+								ExceptionUtils.AddContext(Ex, "at {0}({1})", sourceLocation.File, sourceLocation.LineNumber);
 							}
+
 							throw;
 						}
 					}
diff --git a/Engine/Source/Programs/AutomationTool/BuildGraph/BgScriptReader.cs b/Engine/Source/Programs/AutomationTool/BuildGraph/BgScriptReader.cs
--- a/Engine/Source/Programs/AutomationTool/BuildGraph/BgScriptReader.cs
+++ b/Engine/Source/Programs/AutomationTool/BuildGraph/BgScriptReader.cs
@@ -98,7 +98,7 @@ namespace AutomationTool
 		/// <param name="logger">Logger for output messages</param>
 		/// <param name="outDocument">If successful, the document that was read</param>
 		/// <returns>True if the document could be read, false otherwise</returns>
-		public static bool TryRead(FileReference file, byte[] data, BgScriptSchema schema, ILogger logger, [NotNullWhen(true)] out BgScriptDocument? outDocument)
+		public static bool TryRead(FileReference file, byte[] data, BgScriptSchema? schema, ILogger logger, [NotNullWhen(true)] out BgScriptDocument? outDocument)
 		{
 			BgScriptDocument document = new BgScriptDocument(file, logger);
 
@@ -302,7 +302,7 @@ namespace AutomationTool
 		/// List of property name to value lookups. Modifications to properties are scoped to nodes and agents. EnterScope() pushes an empty dictionary onto the end of this list, and LeaveScope() removes one. 
 		/// ExpandProperties() searches from last to first lookup when trying to resolve a property name, and takes the first it finds.
 		/// </summary>
-		protected List<Dictionary<string, string>> ScopedProperties { get; } = new List<Dictionary<string, string>>();
+		protected List<Dictionary<string, string?>> ScopedProperties { get; } = new List<Dictionary<string, string?>>();
 
 		/// <summary>
 		/// When declaring a property in a nested scope, we enter its name into a set for each parent scope which prevents redeclaration in an OUTER scope later. Subsequent NESTED scopes can redeclare it.
@@ -333,7 +333,7 @@ namespace AutomationTool
 		/// <summary>
 		/// Schema for the script
 		/// </summary>
-		BgScriptSchema Schema { get; }
+		BgScriptSchema? Schema { get; }
 
 		/// <summary>
 		/// Logger for diagnostic messages
@@ -356,7 +356,7 @@ namespace AutomationTool
 		/// <param name="singleNodeName">If a single node will be processed, the name of that node.</param>
 		/// <param name="schema">Schema for the script</param>
 		/// <param name="logger">Logger for diagnostic messages</param>
-		protected BgScriptReader(IDictionary<string, string> defaultProperties, IReadOnlyDictionary<string, string> arguments, string? singleNodeName, BgScriptSchema schema, ILogger logger)
+		protected BgScriptReader(IDictionary<string, string?> defaultProperties, IReadOnlyDictionary<string, string> arguments, string? singleNodeName, BgScriptSchema? schema, ILogger logger)
 		{
 			Schema = schema;
 			Logger = logger;
@@ -366,7 +366,7 @@ namespace AutomationTool
 			_arguments = new Dictionary<string, string>(arguments, StringComparer.OrdinalIgnoreCase);
 			_singleNodeName = singleNodeName;
 
-			foreach (KeyValuePair<string, string> pair in defaultProperties)
+			foreach (KeyValuePair<string, string?> pair in defaultProperties)
 			{
 				SetPropertyValue(null!, pair.Key, pair.Value);
 			}
@@ -382,7 +382,7 @@ namespace AutomationTool
 		/// <param name="logger">Logger for output messages</param>
 		/// <param name="singleNodeName">If a single node will be processed, the name of that node.</param>
 		/// <returns>True if the graph was read, false if there were errors</returns>
-		public static async Task<BgGraphDef?> ReadAsync(FileReference file, Dictionary<string, string> arguments, Dictionary<string, string> defaultProperties, BgScriptSchema schema, ILogger logger, string? singleNodeName = null)
+		public static async Task<BgGraphDef?> ReadAsync(FileReference file, Dictionary<string, string> arguments, Dictionary<string, string?> defaultProperties, BgScriptSchema? schema, ILogger logger, string? singleNodeName = null)
 		{
 			// Read the file and build the graph
 			BgScriptReader reader = new BgScriptReader(defaultProperties, arguments, singleNodeName, schema, logger);
@@ -539,7 +539,7 @@ namespace AutomationTool
 		/// </summary>
 		protected void EnterScope()
 		{
-			ScopedProperties.Add(new Dictionary<string, string>(StringComparer.OrdinalIgnoreCase));
+			ScopedProperties.Add(new Dictionary<string, string?>(StringComparer.OrdinalIgnoreCase));
 			_shadowProperties.Add(new HashSet<string>(StringComparer.OrdinalIgnoreCase));
 		}
 
@@ -559,7 +559,7 @@ namespace AutomationTool
 		/// <param name="name">Name of the property</param>
 		/// <param name="value">Value for the property</param>
 		/// <param name="createInParentScope">If true, this property should be added to the parent scope and not the current scope. Cannot be used if the parent scope already contains a parameter with this name or if there is no parent scope</param>
-		protected void SetPropertyValue(BgScriptElement element, string name, string value, bool createInParentScope = false)
+		protected void SetPropertyValue(BgScriptElement element, string name, string? value, bool createInParentScope = false)
 		{
 			// Find the scope containing this property, defaulting to the current scope
 			int scopeIdx = 0;
diff --git a/Engine/Source/Programs/AutomationTool/BuildGraph/BgScriptWriter.cs b/Engine/Source/Programs/AutomationTool/BuildGraph/BgScriptWriter.cs
--- a/Engine/Source/Programs/AutomationTool/BuildGraph/BgScriptWriter.cs
+++ b/Engine/Source/Programs/AutomationTool/BuildGraph/BgScriptWriter.cs
@@ -14,6 +14,8 @@ using EpicGames.BuildGraph;
 using EpicGames.Core;
 using Microsoft.Extensions.Logging;
 
+#nullable enable
+
 namespace AutomationTool
 {
 	/// <summary>
@@ -27,7 +29,7 @@ namespace AutomationTool
 		/// <param name="graph">Graph to output</param>
 		/// <param name="file">The file to load</param>
 		/// <param name="schemaFile">Schema file for validation</param>
-		public static void Write(this BgGraphDef graph, FileReference file, FileReference schemaFile)
+		public static void Write(this BgGraphDef graph, FileReference file, FileReference? schemaFile)
 		{
 			XmlWriterSettings settings = new XmlWriterSettings();
 			settings.Indent = true;
diff --git a/Engine/Source/Programs/AutomationTool/BuildGraph/BgTask.cs b/Engine/Source/Programs/AutomationTool/BuildGraph/BgTask.cs
--- a/Engine/Source/Programs/AutomationTool/BuildGraph/BgTask.cs
+++ b/Engine/Source/Programs/AutomationTool/BuildGraph/BgTask.cs
@@ -3,6 +3,8 @@
 using System.Collections.Generic;
 using System.Xml;
 
+#nullable enable
+
 namespace AutomationTool
 {
 	/// <summary>
diff --git a/Engine/Source/Programs/AutomationTool/BuildGraph/BuildGraph.cs b/Engine/Source/Programs/AutomationTool/BuildGraph/BuildGraph.cs
--- a/Engine/Source/Programs/AutomationTool/BuildGraph/BuildGraph.cs
+++ b/Engine/Source/Programs/AutomationTool/BuildGraph/BuildGraph.cs
@@ -19,6 +19,8 @@ using System.Xml;
 using UnrealBuildBase;
 using UnrealBuildTool;
 
+#nullable enable
+
 namespace AutomationTool
 {
 	/// <summary>
@@ -90,7 +92,7 @@ namespace AutomationTool
 			{
 				if (Field.MemberType == MemberTypes.Field)
 				{
-					TaskParameterAttribute ParameterAttribute = Field.GetCustomAttribute<TaskParameterAttribute>();
+					TaskParameterAttribute? ParameterAttribute = Field.GetCustomAttribute<TaskParameterAttribute>();
 					if (ParameterAttribute != null)
 					{
 						ScriptTaskParameters.Add(new ScriptTaskParameterBinding(Field.Name, Field, ParameterAttribute.ValidationType, ParameterAttribute.Optional));
@@ -137,7 +139,7 @@ namespace AutomationTool
 		/// <param name="Branch">The current branch</param>
 		/// <param name="Change">Changelist being built</param>
 		/// <param name="CodeChange">Code changelist being built</param>
-		internal BgEnvironment(string Branch, int? Change, int? CodeChange)
+		internal BgEnvironment(string? Branch, int? Change, int? CodeChange)
 		{
 			this.Stream = Branch ?? "Unknown";
 			this.Change = Change ?? 0;
@@ -263,12 +265,12 @@ namespace AutomationTool
 			}
 
 			// Get the standard P4 properties, defaulting to the environment variables if not set. This allows setting p4-like properties without having a P4 connection.
-			string Branch = P4Enabled ? P4Env.Branch : GetEnvVarOrNull(EnvVarNames.BuildRootP4);
+			string? Branch = P4Enabled ? P4Env.Branch : GetEnvVarOrNull(EnvVarNames.BuildRootP4);
 			int? Change = P4Enabled ? P4Env.Changelist : GetEnvVarIntOrNull(EnvVarNames.Changelist);
 			int? CodeChange = P4Enabled ? P4Env.CodeChangelist : GetEnvVarIntOrNull(EnvVarNames.CodeChangelist);
 
 			// Set up the standard properties which build scripts might need
-			Dictionary<string, string> DefaultProperties = new Dictionary<string, string>(StringComparer.InvariantCultureIgnoreCase);
+			Dictionary<string, string?> DefaultProperties = new Dictionary<string, string?>(StringComparer.InvariantCultureIgnoreCase);
 			DefaultProperties["Branch"] = Branch ?? "Unknown";
 			DefaultProperties["Depot"] = (Branch != null && Branch.StartsWith("//", StringComparison.Ordinal)) ? Branch.Substring(2).Split('/').First() : "Unknown";
 			DefaultProperties["EscapedBranch"] = String.IsNullOrEmpty(Branch) ? "Unknown" : CommandUtils.EscapePath(Branch);
@@ -299,7 +301,7 @@ namespace AutomationTool
 			}
 
 			// Attempt to read existing Build Version information
-			BuildVersion Version;
+			BuildVersion? Version;
 			if (BuildVersion.TryRead(BuildVersion.GetDefaultFileName(), out Version))
 			{
 				DefaultProperties["EngineMajorVersion"] = Version.MajorVersion.ToString();
@@ -374,14 +376,14 @@ namespace AutomationTool
 			}
 
 			// Create the graph
-			BgGraphDef Graph;
+			BgGraphDef? Graph;
 			if (ClassName != null)
 			{
 				// Find all the graph builders
 				Dictionary<string, Type> NameToType = new Dictionary<string, Type>();
 				FindAvailableGraphs(NameToType);
 
-				Type BuilderType;
+				Type? BuilderType;
 				if (!NameToType.TryGetValue(ClassName, out BuilderType))
 				{
 					Logger.LogError("Unable to find graph '{GraphName}'", ClassName);
@@ -394,7 +396,7 @@ namespace AutomationTool
 					return ExitCode.Error_Unknown;
 				}
 
-				BgGraphBuilder Builder = (BgGraphBuilder)Activator.CreateInstance(BuilderType);
+				BgGraphBuilder Builder = (BgGraphBuilder)Activator.CreateInstance(BuilderType)!;
 				BgGraph GraphSpec = Builder.CreateGraph(new BgEnvironment(Branch, Change, CodeChange));
 
 				(byte[] Data, BgThunkDef[] Methods) = BgCompiler.Compile(GraphSpec);
@@ -406,7 +408,7 @@ namespace AutomationTool
 			else
 			{
 				// Import schema if one is passed in
-				BgScriptSchema Schema;
+				BgScriptSchema? Schema;
 				if (ImportSchemaFileName != null)
 				{
 					Schema = BgScriptSchema.Import(FileReference.FromString(ImportSchemaFileName));
@@ -477,7 +479,7 @@ namespace AutomationTool
 			}
 			else
 			{
-				IEnumerable<string> NodesToResolve = null;
+				IEnumerable<string>? NodesToResolve = null;
 
 				// If we're only building a single node and using a preprocessed reference we only need to try to resolve the references
 				// for that node.
@@ -492,7 +494,7 @@ namespace AutomationTool
 
 				foreach (string TargetName in NodesToResolve)
 				{
-					BgNodeDef[] Nodes;
+					BgNodeDef[]? Nodes;
 					if (!Graph.TryResolveReference(TargetName, out Nodes))
 					{
 						Logger.LogError("Target '{TargetName}' is not in graph", TargetName);
@@ -532,7 +534,7 @@ namespace AutomationTool
 				Dictionary<FileReference, string> MissingTokens = new Dictionary<FileReference, string>();
 				foreach (FileReference RequiredToken in RequiredTokens)
 				{
-					string CurrentOwner = ReadTokenFile(RequiredToken);
+					string? CurrentOwner = ReadTokenFile(RequiredToken);
 					if (CurrentOwner != null && CurrentOwner != TokenSignature)
 					{
 						MissingTokens.Add(RequiredToken, CurrentOwner);
@@ -616,7 +618,7 @@ namespace AutomationTool
 			}
 
 			// If we're just building a single node, find it 
-			BgNodeDef SingleNode = null;
+			BgNodeDef? SingleNode = null;
 			if (SingleNodeName != null && !Graph.NameToNode.TryGetValue(SingleNodeName, out SingleNode))
 			{
 				Logger.LogError("Node '{SingleNodeName}' is not in the trimmed graph", SingleNodeName);
@@ -746,7 +748,7 @@ namespace AutomationTool
 				}
 				catch (ReflectionTypeLoadException ex)
 				{
-					Logger.LogWarning("Exception {ex} while trying to get types from assembly {LoadedAssembly}. LoaderExceptions: {Arg2}", ex, LoadedAssembly, string.Join("\n", ex.LoaderExceptions.Select(x => x.Message)));
+					Logger.LogWarning("Exception {ex} while trying to get types from assembly {LoadedAssembly}. LoaderExceptions: {Arg2}", ex, LoadedAssembly, string.Join("\n", ex.LoaderExceptions.Select(x => x?.Message)));
 					continue;
 				}
 
@@ -782,7 +784,7 @@ namespace AutomationTool
 				}
 				catch (ReflectionTypeLoadException ex)
 				{
-					Logger.LogWarning("Exception {ex} while trying to get types from assembly {LoadedAssembly}. LoaderExceptions: {Arg2}", ex, LoadedAssembly, string.Join("\n", ex.LoaderExceptions.Select(x => x.Message)));
+					Logger.LogWarning("Exception {ex} while trying to get types from assembly {LoadedAssembly}. LoaderExceptions: {Arg2}", ex, LoadedAssembly, string.Join("\n", ex.LoaderExceptions.Select(x => x?.Message)));
 					continue;
 				}
 
@@ -811,7 +813,7 @@ namespace AutomationTool
 		/// Reads the contents of the given token
 		/// </summary>
 		/// <returns>Contents of the token, or null if it does not exist</returns>
-		public string ReadTokenFile(FileReference Location)
+		public string? ReadTokenFile(FileReference Location)
 		{
 			return FileReference.Exists(Location) ? File.ReadAllText(Location.FullName) : null;
 		}
@@ -962,13 +964,13 @@ namespace AutomationTool
 		class CleanupScriptRunner : IDisposable
 		{
 			readonly ILogger _logger;
-			readonly FileReference _scriptFile;
+			readonly FileReference? _scriptFile;
 
 			public CleanupScriptRunner(ILogger logger)
 			{
 				_logger = logger;
 
-				string CleanupScriptEnvVar = Environment.GetEnvironmentVariable(CustomTask.CleanupScriptEnvVarName);
+				string? CleanupScriptEnvVar = Environment.GetEnvironmentVariable(CustomTask.CleanupScriptEnvVarName);
 				if (String.IsNullOrEmpty(CleanupScriptEnvVar))
 				{
 					string extension = RuntimeInformation.IsOSPlatform(OSPlatform.Windows) ? "bat" : "sh";
@@ -1025,7 +1027,7 @@ namespace AutomationTool
 			HashSet<TempStorageBlock> InputStorageBlocks = new HashSet<TempStorageBlock>();
 			foreach (BgNodeOutput Input in Node.Inputs)
 			{
-				TempStorageTagManifest FileList = Storage.ReadFileList(Input.ProducingNode.Name, Input.TagName);
+				TempStorageTagManifest FileList = Storage.ReadFileList(Input.ProducingNode.Name, Input.TagName) ?? throw new InvalidOperationException();
 				TagNameToFileSet[Input.TagName] = FileList.ToFileSet(RootDir);
 				InputStorageBlocks.UnionWith(FileList.Blocks);
 			}
@@ -1050,7 +1052,7 @@ namespace AutomationTool
 				TempStorageBlock InputStorageBlock = Pair.Key;
 				foreach (FileReference File in Pair.Value.Files.Select(x => x.ToFileReference(RootDir)))
 				{
-					TempStorageBlock CurrentStorageBlock;
+					TempStorageBlock? CurrentStorageBlock;
 					if (FileToStorageBlock.TryGetValue(File, out CurrentStorageBlock) && !TempStorage.IsDuplicateBuildProduct(File))
 					{
 						Logger.LogError("File '{File}' was produced by {InputStorageBlock} and {CurrentStorageBlock}", File, InputStorageBlock, CurrentStorageBlock);
@@ -1085,7 +1087,7 @@ namespace AutomationTool
 			Dictionary<string, string> ModifiedFiles = new Dictionary<string, string>(StringComparer.InvariantCultureIgnoreCase);
 			foreach (TempStorageFile File in InputManifests.Values.SelectMany(x => x.Files))
 			{
-				string Message;
+				string? Message;
 				if (!ModifiedFiles.ContainsKey(File.RelativePath) && !File.Compare(Unreal.RootDirectory, out Message) && !IgnoreModifiedFilter.Matches(File.RelativePath))
 				{
 					// look up the previous nodes to help with error diagnosis
@@ -1150,7 +1152,7 @@ namespace AutomationTool
 						}
 						else
 						{
-							string OutputName;
+							string? OutputName;
 							if (FileToOutputName.TryGetValue(File, out OutputName) && OutputName.Length > 0)
 							{
 								FileToOutputName[File] = String.Format("{0}+{1}", OutputName, Output.TagName.Substring(1));
@@ -1168,7 +1170,7 @@ namespace AutomationTool
 			Dictionary<string, HashSet<FileReference>> OutputStorageBlockToFiles = new Dictionary<string, HashSet<FileReference>>();
 			foreach (KeyValuePair<FileReference, string> Pair in FileToOutputName)
 			{
-				HashSet<FileReference> Files;
+				HashSet<FileReference>? Files;
 				if (!OutputStorageBlockToFiles.TryGetValue(Pair.Value, out Files))
 				{
 					Files = new HashSet<FileReference>();
@@ -1201,7 +1203,7 @@ namespace AutomationTool
 					HashSet<TempStorageBlock> StorageBlocks = new HashSet<TempStorageBlock>();
 					foreach (FileReference File in Files)
 					{
-						TempStorageBlock StorageBlock;
+						TempStorageBlock? StorageBlock;
 						if (FileToStorageBlock.TryGetValue(File, out StorageBlock))
 						{
 							StorageBlocks.Add(StorageBlock);
@@ -1209,7 +1211,7 @@ namespace AutomationTool
 					}
 
 					IEnumerable<string> Keys = Enumerable.Empty<string>();
-					if (outputNameToArtifact.TryGetValue(Output.TagName, out BgArtifactDef artifact))
+					if (outputNameToArtifact.TryGetValue(Output.TagName, out BgArtifactDef? artifact))
 					{
 						Keys = artifact.Keys;
 					}
@@ -1226,9 +1228,9 @@ namespace AutomationTool
 		/// <summary>
 		/// Gets an environment variable, returning null if it's not set or empty.
 		/// </summary>
-		static string GetEnvVarOrNull(string Name)
+		static string? GetEnvVarOrNull(string Name)
 		{
-			string EnvVar = Environment.GetEnvironmentVariable(Name);
+			string? EnvVar = Environment.GetEnvironmentVariable(Name);
 			return String.IsNullOrEmpty(EnvVar) ? null : EnvVar;
 		}
 
@@ -1237,7 +1239,7 @@ namespace AutomationTool
 		/// </summary>
 		static int? GetEnvVarIntOrNull(string Name)
 		{
-			string EnvVar = GetEnvVarOrNull(Name);
+			string? EnvVar = GetEnvVarOrNull(Name);
 			if (EnvVar != null && Int32.TryParse(EnvVar, out int Value))
 			{
 				return Value;
@@ -1267,7 +1269,7 @@ namespace AutomationTool
 					Document.Load(XmlFileName);
 
 					// Parse all the members, and add them to the map
-					foreach (XmlElement Element in Document.SelectNodes("/doc/members/member"))
+					foreach (XmlElement Element in Document.SelectNodes("/doc/members/member")!)
 					{
 						string Name = Element.GetAttribute("name");
 						MemberNameToElement.Add(Name, Element);
@@ -1323,13 +1325,13 @@ namespace AutomationTool
 					ScriptTaskBinding Task = NameToTask[TaskName];
 
 					// Get the documentation for this task
-					XmlElement TaskElement;
+					XmlElement? TaskElement;
 					if (MemberNameToElement.TryGetValue("T:" + Task.TaskClass.FullName, out TaskElement))
 					{
 						// Write the task heading
 						Writer.WriteLine("### {0}", TaskName);
 						Writer.WriteLine();
-						Writer.WriteLine(ConvertToMarkdown(TaskElement.SelectSingleNode("summary")));
+						Writer.WriteLine(ConvertToMarkdown(TaskElement.SelectSingleNode("summary")!));
 						Writer.WriteLine();
 
 						// Document the parameters
@@ -1340,8 +1342,8 @@ namespace AutomationTool
 							ScriptTaskParameterBinding Parameter = Task.NameToParameter[ParameterName];
 
 							// Get the documentation for this parameter
-							XmlElement ParameterElement;
-							if (MemberNameToElement.TryGetValue("F:" + Parameter.FieldInfo.DeclaringType.FullName + "." + Parameter.Name, out ParameterElement))
+							XmlElement? ParameterElement;
+							if (MemberNameToElement.TryGetValue("F:" + Parameter.FieldInfo.DeclaringType!.FullName + "." + Parameter.Name, out ParameterElement))
 							{
 								Type FieldType = Parameter.FieldInfo.FieldType;
 								if (FieldType.IsGenericType && FieldType.GetGenericTypeDefinition() == typeof(Nullable<>))
@@ -1371,7 +1373,7 @@ namespace AutomationTool
 								Columns[0] = ParameterName;
 								Columns[1] = TypeName;
 								Columns[2] = Parameter.Optional ? "Optional" : "Required";
-								Columns[3] = ConvertToMarkdown(ParameterElement.SelectSingleNode("summary"));
+								Columns[3] = ConvertToMarkdown(ParameterElement.SelectSingleNode("summary")!);
 								Rows.Add(Columns);
 							}
 						}
@@ -1433,12 +1435,12 @@ namespace AutomationTool
 					ScriptTaskBinding Task = NameToTask[TaskName];
 
 					// Get the documentation for this task
-					XmlElement TaskElement;
+					XmlElement? TaskElement;
 					if (MemberNameToElement.TryGetValue("T:" + Task.TaskClass.FullName, out TaskElement))
 					{
 						// Write the task heading
 						Writer.WriteLine("    <h2>{0}</h2>", TaskName);
-						Writer.WriteLine("    <p>{0}</p>", TaskElement.SelectSingleNode("summary").InnerXml.Trim());
+						Writer.WriteLine("    <p>{0}</p>", TaskElement.SelectSingleNode("summary")!.InnerXml.Trim());
 
 						// Start the parameter table
 						Writer.WriteLine("    <table>");
@@ -1456,8 +1458,8 @@ namespace AutomationTool
 							ScriptTaskParameterBinding Parameter = Task.NameToParameter[ParameterName];
 
 							// Get the documentation for this parameter
-							XmlElement ParameterElement;
-							if (MemberNameToElement.TryGetValue("F:" + Parameter.FieldInfo.DeclaringType.FullName + "." + Parameter.Name, out ParameterElement))
+							XmlElement? ParameterElement;
+							if (MemberNameToElement.TryGetValue("F:" + Parameter.FieldInfo.DeclaringType!.FullName + "." + Parameter.Name, out ParameterElement))
 							{
 								string TypeName = Parameter.FieldInfo.FieldType.Name;
 								if (Parameter.ValidationType != TaskParameterValidationType.Default)
@@ -1477,7 +1479,7 @@ namespace AutomationTool
 								Writer.WriteLine("         <td>{0}</td>", ParameterName);
 								Writer.WriteLine("         <td>{0}</td>", TypeName);
 								Writer.WriteLine("         <td>{0}</td>", Parameter.Optional ? "Optional" : "Required");
-								Writer.WriteLine("         <td>{0}</td>", ParameterElement.SelectSingleNode("summary").InnerXml.Trim());
+								Writer.WriteLine("         <td>{0}</td>", ParameterElement.SelectSingleNode("summary")!.InnerXml.Trim());
 								Writer.WriteLine("      </tr>");
 							}
 						}
diff --git a/Engine/Source/Programs/AutomationTool/BuildGraph/JobContext.cs b/Engine/Source/Programs/AutomationTool/BuildGraph/JobContext.cs
--- a/Engine/Source/Programs/AutomationTool/BuildGraph/JobContext.cs
+++ b/Engine/Source/Programs/AutomationTool/BuildGraph/JobContext.cs
@@ -6,6 +6,8 @@ using System.Linq;
 using System.Text;
 using System.Threading.Tasks;
 
+#nullable enable
+
 namespace AutomationTool
 {
 	/// <summary>
diff --git a/Engine/Source/Programs/AutomationTool/BuildGraph/Task.cs b/Engine/Source/Programs/AutomationTool/BuildGraph/Task.cs
--- a/Engine/Source/Programs/AutomationTool/BuildGraph/Task.cs
+++ b/Engine/Source/Programs/AutomationTool/BuildGraph/Task.cs
@@ -18,6 +18,8 @@ using System.Runtime.InteropServices;
 
 using static AutomationTool.CommandUtils;
 
+#nullable enable
+
 namespace AutomationTool
 {
 	/// <summary>
@@ -109,7 +111,7 @@ namespace AutomationTool
 		/// <summary>
 		/// Line number in a source file that this task was declared. Optional; used for log messages.
 		/// </summary>
-		public BgScriptLocation SourceLocation { get; set; }
+		public BgScriptLocation? SourceLocation { get; set; }
 
 		/// <summary>
 		/// Execute this node.
@@ -124,7 +126,7 @@ namespace AutomationTool
 		/// Creates a proxy to execute this node.
 		/// </summary>
 		/// <returns>New proxy instance if one is available to execute this task, otherwise null.</returns>
-		public virtual ITaskExecutor GetExecutor()
+		public virtual ITaskExecutor? GetExecutor()
 		{
 			return null;
 		}
@@ -141,17 +143,17 @@ namespace AutomationTool
 		/// <param name="Parameters">Parameters object that this task is constructed with</param>
 		protected void Write(XmlWriter Writer, object Parameters)
 		{
-			TaskElementAttribute Element = GetType().GetCustomAttribute<TaskElementAttribute>();
+			TaskElementAttribute Element = GetType().GetCustomAttribute<TaskElementAttribute>() ?? throw new InvalidOperationException();
 			Writer.WriteStartElement(Element.Name);
 
 			foreach (FieldInfo Field in Parameters.GetType().GetFields())
 			{
 				if (Field.MemberType == MemberTypes.Field)
 				{
-					TaskParameterAttribute ParameterAttribute = Field.GetCustomAttribute<TaskParameterAttribute>();
+					TaskParameterAttribute? ParameterAttribute = Field.GetCustomAttribute<TaskParameterAttribute>();
 					if (ParameterAttribute != null)
 					{
-						object Value = Field.GetValue(Parameters);
+						object? Value = Field.GetValue(Parameters);
 						if (Value != null)
 						{
 							Writer.WriteAttributeString(Field.Name, Value.ToString());
@@ -182,7 +184,7 @@ namespace AutomationTool
 		/// <returns>The trace name</returns>
 		public virtual string GetTraceName()
 		{
-			TaskElementAttribute TaskElement = GetType().GetCustomAttribute<TaskElementAttribute>();
+			TaskElementAttribute? TaskElement = GetType().GetCustomAttribute<TaskElementAttribute>();
 			return (TaskElement != null)? TaskElement.Name : "unknown";
 		}
 
@@ -193,8 +195,11 @@ namespace AutomationTool
 		/// <param name="Prefix">Prefix for metadata entries</param>
 		public virtual void GetTraceMetadata(ITraceSpan Span, string Prefix)
 		{
-			Span.AddMetadata(Prefix + "source.file", SourceLocation.File.FullName);
-			Span.AddMetadata(Prefix + "source.line", SourceLocation.LineNumber.ToString());
+			if (SourceLocation != null)
+			{
+				Span.AddMetadata(Prefix + "source.file", SourceLocation.File.FullName);
+				Span.AddMetadata(Prefix + "source.line", SourceLocation.LineNumber.ToString());
+			}
 		}
 		
 		/// <summary>
@@ -204,8 +209,11 @@ namespace AutomationTool
 		/// <param name="Prefix">Prefix for metadata entries</param>
 		public virtual void GetTraceMetadata(ISpan Span, string Prefix)
 		{
-			Span.SetTag(Prefix + "source.file", SourceLocation.File.FullName);
-			Span.SetTag(Prefix + "source.line", SourceLocation.LineNumber);
+			if (SourceLocation != null)
+			{
+				Span.SetTag(Prefix + "source.file", SourceLocation.File.FullName);
+				Span.SetTag(Prefix + "source.line", SourceLocation.LineNumber);
+			}
 		}
 
 		/// <summary>
@@ -244,7 +252,7 @@ namespace AutomationTool
 		/// </summary>
 		/// <param name="TagList">List of tags separated by semicolons</param>
 		/// <returns>Tag names from this filespec</returns>
-		protected static IEnumerable<string> FindTagNamesFromList(string TagList)
+		protected static IEnumerable<string> FindTagNamesFromList(string? TagList)
 		{
 			if(!String.IsNullOrEmpty(TagList))
 			{
@@ -277,7 +285,7 @@ namespace AutomationTool
 		/// </summary>
 		/// <param name="Name">Name of the directory. May be null or empty.</param>
 		/// <returns>The resolved directory</returns>
-		public static DirectoryReference ResolveDirectory(string Name)
+		public static DirectoryReference ResolveDirectory(string? Name)
 		{
 			if(String.IsNullOrEmpty(Name))
 			{
@@ -314,7 +322,7 @@ namespace AutomationTool
 			}
 
 			// Find the files which match this tag
-			HashSet<FileReference> Files;
+			HashSet<FileReference>? Files;
 			if(!TagNameToFileSet.TryGetValue(TagName, out Files))
 			{
 				Files = new HashSet<FileReference>();
@@ -451,7 +459,7 @@ namespace AutomationTool
 		/// <param name="Lease">Whether to add the commands to run on lease termination</param>
 		public static async Task AddCleanupCommandsAsync(IEnumerable<string> NewLines, bool Lease = false)
 		{
-			string CleanupScriptEnvVar = Environment.GetEnvironmentVariable(Lease? LeaseCleanupScriptEnvVarName : CleanupScriptEnvVarName);
+			string? CleanupScriptEnvVar = Environment.GetEnvironmentVariable(Lease? LeaseCleanupScriptEnvVarName : CleanupScriptEnvVarName);
 			if (!String.IsNullOrEmpty(CleanupScriptEnvVar))
 			{
 				FileReference CleanupScript = new FileReference(CleanupScriptEnvVar);
@@ -470,7 +478,7 @@ namespace AutomationTool
 		/// <param name="Job">Context for the current job that is being executed</param>
 		public static void UpdateGraphForHorde(JobContext Job)
 		{
-			string exportGraphFile = Environment.GetEnvironmentVariable(GraphUpdateEnvVarName);
+			string? exportGraphFile = Environment.GetEnvironmentVariable(GraphUpdateEnvVarName);
 			if (String.IsNullOrEmpty(exportGraphFile))
 			{
 				throw new Exception($"Missing environment variable {GraphUpdateEnvVarName}. This is required to update graphs on Horde.");
diff --git a/Engine/Source/Programs/AutomationTool/BuildGraph/TempStorage.cs b/Engine/Source/Programs/AutomationTool/BuildGraph/TempStorage.cs
--- a/Engine/Source/Programs/AutomationTool/BuildGraph/TempStorage.cs
+++ b/Engine/Source/Programs/AutomationTool/BuildGraph/TempStorage.cs
@@ -20,6 +20,9 @@ using Microsoft.Extensions.Logging;
 using System.Text;
 
 using static AutomationTool.CommandUtils;
+using System.Diagnostics.CodeAnalysis;
+
+#nullable enable
 
 namespace AutomationTool
 {
@@ -45,6 +48,8 @@ namespace AutomationTool
 		/// </summary>
 		private TempStorageBlock()
 		{
+			NodeName = String.Empty;
+			OutputName = String.Empty;
 		}
 
 		/// <summary>
@@ -63,9 +68,9 @@ namespace AutomationTool
 		/// </summary>
 		/// <param name="Other">The object to compare against</param>
 		/// <returns>True if the blocks are equivalent</returns>
-		public override bool Equals(object Other)
+		public override bool Equals(object? Other)
 		{
-			TempStorageBlock OtherBlock = Other as TempStorageBlock;
+			TempStorageBlock? OtherBlock = Other as TempStorageBlock;
 			return OtherBlock != null && NodeName == OtherBlock.NodeName && OutputName == OtherBlock.OutputName;
 		}
 
@@ -116,13 +121,14 @@ namespace AutomationTool
 		/// Digest for the file. Not all files are hashed.
 		/// </summary>
 		[XmlAttribute]
-		public string Digest;
+		public string? Digest;
 
 		/// <summary>
 		/// Default constructor, for XML serialization.
 		/// </summary>
 		private TempStorageFile()
 		{
+			RelativePath = String.Empty;
 		}
 
 		/// <summary>
@@ -160,7 +166,7 @@ namespace AutomationTool
 		/// <returns>True if the files are identical, false otherwise</returns>
 		public bool Compare(DirectoryReference RootDir)
 		{
-			string Message;
+			string? Message;
 			if(Compare(RootDir, out Message))
 			{
 				if(Message != null)
@@ -185,7 +191,7 @@ namespace AutomationTool
 		/// <param name="RootDir">Root directory for this branch</param>
 		/// <param name="Message">Message describing the difference</param>
 		/// <returns>True if the files are identical, false otherwise</returns>
-		public bool Compare(DirectoryReference RootDir, out string Message)
+		public bool Compare(DirectoryReference RootDir, [NotNullWhen(false)] out string? Message)
 		{
 			FileReference LocalFile = ToFileReference(RootDir);
 
@@ -312,6 +318,7 @@ namespace AutomationTool
 		/// </summary>
 		private TempStorageZipFile()
 		{
+			Name = String.Empty;
 		}
 
 		/// <summary>
@@ -347,13 +354,15 @@ namespace AutomationTool
 		/// <summary>
 		/// Construct a static Xml serializer to avoid throwing an exception searching for the reflection info at runtime
 		/// </summary>
-		static XmlSerializer Serializer = XmlSerializer.FromTypes(new Type[]{ typeof(TempStorageManifest) })[0];
+		static XmlSerializer Serializer = XmlSerializer.FromTypes(new Type[]{ typeof(TempStorageManifest) })[0]!;
 
 		/// <summary>
 		/// Construct an empty temp storage manifest
 		/// </summary>
 		private TempStorageManifest()
 		{
+			Files = Array.Empty<TempStorageFile>();
+			ZipFiles = Array.Empty<TempStorageZipFile>();
 		}
 
 		/// <summary>
@@ -364,6 +373,7 @@ namespace AutomationTool
 		public TempStorageManifest(FileInfo[] InFiles, DirectoryReference RootDir)
 		{
 			Files = InFiles.Select(x => new TempStorageFile(x, RootDir)).ToArray();
+			ZipFiles = Array.Empty<TempStorageZipFile>();
 		}
 
 		/// <summary>
@@ -388,7 +398,7 @@ namespace AutomationTool
 		{
 			using (StreamReader Reader = new(file.FullName))
 			{
-				return (TempStorageManifest)Serializer.Deserialize(Reader);
+				return (TempStorageManifest)(Serializer.Deserialize(Reader) ?? throw new InvalidOperationException());
 			}
 		}
 
@@ -446,13 +456,17 @@ namespace AutomationTool
 		/// <summary>
 		/// Construct a static Xml serializer to avoid throwing an exception searching for the reflection info at runtime
 		/// </summary>
-		static XmlSerializer Serializer = XmlSerializer.FromTypes(new Type[]{ typeof(TempStorageTagManifest) })[0];
+		static XmlSerializer Serializer = XmlSerializer.FromTypes(new Type[]{ typeof(TempStorageTagManifest) })[0]!;
 
 		/// <summary>
 		/// Construct an empty file list for deserialization
 		/// </summary>
 		private TempStorageTagManifest()
 		{
+			LocalFiles = Array.Empty<string>();
+			ExternalFiles = Array.Empty<string>();
+			Blocks = Array.Empty<TempStorageBlock>();
+			ArtifactKeys = Array.Empty<string>();
 		}
 
 		/// <summary>
@@ -492,7 +506,7 @@ namespace AutomationTool
 		{
 			using(StreamReader Reader = new StreamReader(File.FullName))
 			{
-				return (TempStorageTagManifest)Serializer.Deserialize(Reader);
+				return (TempStorageTagManifest)Serializer.Deserialize(Reader)!;
 			}
 		}
 
@@ -551,7 +565,7 @@ namespace AutomationTool
 		/// <summary>
 		/// The shared temp storage directory; typically a network location. May be null.
 		/// </summary>
-		DirectoryReference SharedDir;
+		DirectoryReference? SharedDir;
 
 		/// <summary>
 		/// Whether to allow writes to shared storage
@@ -565,7 +579,7 @@ namespace AutomationTool
 		/// <param name="InLocalDir">The local temp storage directory.</param>
 		/// <param name="InSharedDir">The shared temp storage directory. May be null.</param>
 		/// <param name="bInWriteToSharedStorage">Whether to write to shared storage, or only permit reads from it</param>
-		public TempStorage(DirectoryReference InRootDir, DirectoryReference InLocalDir, DirectoryReference InSharedDir, bool bInWriteToSharedStorage)
+		public TempStorage(DirectoryReference InRootDir, DirectoryReference InLocalDir, DirectoryReference? InSharedDir, bool bInWriteToSharedStorage)
 		{
 			RootDir = InRootDir;
 			LocalDir = InLocalDir;
@@ -681,7 +695,7 @@ namespace AutomationTool
 
 					// Check the manifests are identical, byte by byte
 					byte[] LocalManifestBytes = File.ReadAllBytes(LocalFileListLocation.FullName);
-					byte[] SharedManifestBytes = null;
+					byte[] SharedManifestBytes = Array.Empty<byte>();
 					PerformActionWithRetries(() => SharedManifestBytes = File.ReadAllBytes(SharedFileListLocation.FullName), 3, TimeSpan.FromSeconds(1));
 					if (!LocalManifestBytes.SequenceEqual(SharedManifestBytes))
 					{
@@ -716,7 +730,7 @@ namespace AutomationTool
 
 					// Check the manifests are identical, byte by byte
 					byte[] LocalManifestBytes = File.ReadAllBytes(LocalManifestFile.FullName);
-					byte[] SharedManifestBytes = null;
+					byte[] SharedManifestBytes = Array.Empty<byte>();
 					PerformActionWithRetries(() => SharedManifestBytes = File.ReadAllBytes(SharedManifestFile.FullName), 3, TimeSpan.FromSeconds(1));
 					if(!LocalManifestBytes.SequenceEqual(SharedManifestBytes))
 					{
@@ -740,9 +754,9 @@ namespace AutomationTool
 		/// <param name="NodeName">Name of the node which produced the tag set</param>
 		/// <param name="TagName">Name of the tag, with a '#' prefix</param>
 		/// <returns>The set of files</returns>
-		public TempStorageTagManifest ReadFileList(string NodeName, string TagName)
+		public TempStorageTagManifest? ReadFileList(string NodeName, string TagName)
 		{
-			TempStorageTagManifest FileList = null;
+			TempStorageTagManifest? FileList = null;
 
 			// Try to read the tag set from the local directory
 			FileReference LocalFileListLocation = GetTaggedFileListLocation(LocalDir, NodeName, TagName);
@@ -844,7 +858,7 @@ namespace AutomationTool
 		/// <param name="BuildProducts">Array of build products to be archived</param>
 		/// <param name="bPushToRemote">Allow skipping the copying of this manifest to shared storage, because it's not required by any other agent</param>
 		/// <returns>The created manifest instance (which has already been saved to disk).</returns>
-		public TempStorageManifest Archive(string NodeName, string BlockName, FileReference[] BuildProducts, bool bPushToRemote = true)
+		public TempStorageManifest Archive(string NodeName, string? BlockName, FileReference[] BuildProducts, bool bPushToRemote = true)
 		{
 			using (IScope Scope = GlobalTracer.Instance.BuildSpan("StoreToTempStorage").StartActive())
 			{
@@ -861,7 +875,7 @@ namespace AutomationTool
 				if(bRemote)
 				{
 					// Create the shared directory for this node
-					FileReference SharedManifestFile = GetManifestLocation(SharedDir, NodeName, BlockName);
+					FileReference SharedManifestFile = GetManifestLocation(SharedDir!, NodeName, BlockName);
 					DirectoryReference.CreateDirectory(SharedManifestFile.Directory);
 
 					// Zip all the build products
@@ -896,7 +910,7 @@ namespace AutomationTool
 		/// <param name="OutputName">Name of the block to retrieve. May be null or empty.</param>
 		/// <param name="IgnoreModified">Filter for files to ignore</param>
 		/// <returns>Manifest of the files retrieved</returns>
-		public TempStorageManifest Retrieve(string NodeName, string OutputName, FileFilter IgnoreModified)
+		public TempStorageManifest Retrieve(string NodeName, string? OutputName, FileFilter IgnoreModified)
 		{
 			using (IScope Scope = GlobalTracer.Instance.BuildSpan("RetrieveFromTempStorage").StartActive())
 			{
@@ -905,7 +919,7 @@ namespace AutomationTool
 				bool bLocal = FileReference.Exists(LocalManifestFile);
 
 				// Read the manifest, either from local storage or shared storage
-				TempStorageManifest Manifest = null;
+				TempStorageManifest? Manifest = null;
 				if(bLocal)
 				{
 					Logger.LogInformation("Reading shared manifest from {Arg0}", LocalManifestFile.FullName);
@@ -934,7 +948,7 @@ namespace AutomationTool
 
 					// Unzip all the build products
 					DirectoryReference SharedNodeDir = GetDirectoryForNode(SharedDir, NodeName);
-					FileInfo[] ZipFiles = Manifest.ZipFiles.Select(x => new FileInfo(FileReference.Combine(SharedNodeDir, x.Name).FullName)).ToArray();
+					FileInfo[] ZipFiles = Manifest!.ZipFiles.Select(x => new FileInfo(FileReference.Combine(SharedNodeDir, x.Name).FullName)).ToArray();
 					string Result = ParallelUnzipFiles(ZipFiles, RootDir);
 					if (!string.IsNullOrWhiteSpace(Result))
 					{
@@ -959,7 +973,7 @@ namespace AutomationTool
 				List<string> ModifiedFileMessages = new List<string>();
 				foreach(TempStorageFile File in Manifest.Files)
 				{
-					string Message;
+					string? Message;
 					if (!IgnoreModified.Matches(File.RelativePath) && !File.Compare(RootDir, out Message))
 					{
 						ModifiedFileMessages.Add(Message);
@@ -1050,12 +1064,12 @@ namespace AutomationTool
 			// which has shown to be optimal via profiling, and limits the number of resulting zip files to the number of logical cores.
 			List<Thread> ZipThreads = (
 				from CoreNum in Enumerable.Range(0, bZipInParallel ? Environment.ProcessorCount : 1)
-				select new Thread((object indexObject) =>
+				select new Thread((object? indexObject) =>
 				{
-					int index = (int)indexObject;
+					int index = (int)indexObject!;
 					FileReference ZipFileName = FileReference.Combine(ZipDir, string.Format("{0}{1}.zip", ZipBaseName, bZipInParallel ? "-" + index.ToString("00") : ""));
 					// don't create the zip unless we have at least one file to add
-					FileReference File;
+					FileReference? File;
 					if (FilesToZip.TryDequeue(out File))
 					{
 						try
@@ -1091,14 +1105,14 @@ namespace AutomationTool
 
 			ZipThreads.ForEach(thread => thread.Join());
 
-			if (ZipFiles.Any() && !string.IsNullOrWhiteSpace(StagingDir.FullName))
+			if (ZipFiles.Any() && !string.IsNullOrWhiteSpace(StagingDir?.FullName))
 			{
 				try
 				{
 					string CopyResult = CopyDirectory(ZipDir, OutputDir);
 					string LogPath = CommandUtils.CombinePaths(RootDir.FullName, "Engine/Programs/AutomationTool/Saved/Logs", $"Copy Files to Temp Storage.log");
 					Logger.LogInformation("Saving copy log to {LogPath}", LogPath);
-					Directory.CreateDirectory(Path.GetDirectoryName(LogPath));
+					Directory.CreateDirectory(Path.GetDirectoryName(LogPath)!);
 					File.WriteAllText(LogPath, string.Join(Environment.NewLine, CopyResult));
 					Parallel.ForEach(ZipFiles, (z) => z.Delete());
 					ZipFiles = new ConcurrentBag<FileInfo>(ZipFiles.Select(z => new FileInfo(CommandUtils.MakeRerootedFilePath(z.FullName, StagingDir.FullName, OutputDir.FullName))));
@@ -1167,7 +1181,7 @@ namespace AutomationTool
 										}
 										else
 										{
-											Directory.CreateDirectory(Path.GetDirectoryName(ExtractedFilename));
+											Directory.CreateDirectory(Path.GetDirectoryName(ExtractedFilename)!);
 										}
 
 
@@ -1243,7 +1257,7 @@ namespace AutomationTool
 		/// <param name="BaseDir">A local or shared temp storage root directory.</param>
 		/// <param name="NodeName">Name of the node to get the file for</param>
 		/// <param name="BlockName">Name of the output block to get the manifest for</param>
-		static FileReference GetManifestLocation(DirectoryReference BaseDir, string NodeName, string BlockName)
+		static FileReference GetManifestLocation(DirectoryReference BaseDir, string NodeName, string? BlockName)
 		{
 			return FileReference.Combine(BaseDir, NodeName, String.IsNullOrEmpty(BlockName)? "Manifest.xml" : String.Format("Manifest-{0}.xml", BlockName));
 		}
@@ -1664,7 +1678,7 @@ namespace AutomationTool
 					}
 
 					// Check if there is a marker file, if so skip this folder unless it has been twenty minutes.
-					FileInfo DeleteInProgressFile = BuildsToDelete[Idx].GetFiles("DeleteInProgress.tmp").FirstOrDefault();
+					FileInfo? DeleteInProgressFile = BuildsToDelete[Idx].GetFiles("DeleteInProgress.tmp").FirstOrDefault();
 					if (DeleteInProgressFile != null && DeleteInProgressFile.LastWriteTimeUtc < (DateTimeOffset.UtcNow - TimeSpan.FromMinutes(20)))
 					{
 						Logger.LogInformation("[{Index}/{Total}] {Path} flagged as delete in progress, skipping...", BuildsToDelete.Count - Idx, BuildsToDelete.Count, BuildsToDelete[Idx].FullName);
