- Ability to override Xcodebuild settings when staging and packaging with BuildCookRun, it is not used when building

diff --git a/Engine/Source/Programs/AutomationTool/Apple/ApplePlatform.Automation.cs b/Engine/Source/Programs/AutomationTool/Apple/ApplePlatform.Automation.cs
--- a/Engine/Source/Programs/AutomationTool/Apple/ApplePlatform.Automation.cs
+++ b/Engine/Source/Programs/AutomationTool/Apple/ApplePlatform.Automation.cs
@@ -130,6 +130,11 @@ public abstract class ApplePlatform : Platform
 					// set where the stage directory is
 					$"UE_OVERRIDE_STAGE_DIR=\"{SC.StageDirectory}\"";
 
+				if (Params.XcodeBuildOptions != null)
+				{
+					ExtraOptions += $" {Params.XcodeBuildOptions}";
+				}
+
 				string TargetName = Target.TargetName;
 				if (!Params.IsCodeBasedProject)
 				{
@@ -239,6 +244,11 @@ public abstract class ApplePlatform : Platform
 					ExtraOptions += $" -archivePath \"{XcArchivePath}\"";
 				}
 
+				if (Params.XcodeBuildOptions != null)
+				{
+					ExtraOptions += $" {Params.XcodeBuildOptions}";
+				}
+
 				// if we we packaging for distrbution, we will create a .xcarchive which can be used to submit to app stores, or exported for other distribution methods
 				// the archive will be created in the standard Archives location accessible via Xcode. Using -archive will copy it out into
 				// the specified location for use as needed
diff --git a/Engine/Source/Programs/AutomationTool/AutomationUtils/ProjectParams.cs b/Engine/Source/Programs/AutomationTool/AutomationUtils/ProjectParams.cs
--- a/Engine/Source/Programs/AutomationTool/AutomationUtils/ProjectParams.cs
+++ b/Engine/Source/Programs/AutomationTool/AutomationUtils/ProjectParams.cs
@@ -446,6 +446,7 @@ namespace AutomationTool
 			this.SessionLabel = InParams.SessionLabel;
 			this.ProjectDescriptor = InParams.ProjectDescriptor;
 			this.Upload = InParams.Upload;
+			this.XcodeBuildOptions = InParams.XcodeBuildOptions;
 		}
 
 		/// <summary>
@@ -619,7 +620,8 @@ namespace AutomationTool
 			ParamList<string> InMapsToRebuildLightMaps = null,
 			ParamList<string> InMapsToRebuildHLOD = null,
 			ParamList<string> TitleID = null,
-			string Upload = null
+			string Upload = null,
+			string XcodeBuildOptions = null
 			)
 		{
 			//
@@ -1051,6 +1053,8 @@ namespace AutomationTool
 			}
 
 			this.Upload = Command.ParseParamValue("upload");
+			
+			this.XcodeBuildOptions = ParseParamValueIfNotSpecified(Command, XcodeBuildOptions, "xcodebuildoptions", null);
 
 			if (ClientConfigsToBuild == null)
 			{
@@ -2389,6 +2393,9 @@ namespace AutomationTool
 		[Help("upload", "Arguments for uploading on demand content")]
 		public string Upload { get; set; }
 
+		[Help("XcodeBuildOptions", "Extra options to pass to xcodebuild")]
+		public string XcodeBuildOptions { get; set; }
+
 		private List<SingleTargetProperties> DetectedTargets;
 		private Dictionary<UnrealTargetPlatform, ConfigHierarchy> LoadedEngineConfigs;
 		private Dictionary<UnrealTargetPlatform, ConfigHierarchy> LoadedGameConfigs;
