Don't run PostStagingFileCopy during Apple staging if we are running on a windows host or we are building a DLC plugin

diff --git a/Engine/Source/Programs/AutomationTool/Apple/ApplePlatform.Automation.cs b/Engine/Source/Programs/AutomationTool/Apple/ApplePlatform.Automation.cs
--- a/Engine/Source/Programs/AutomationTool/Apple/ApplePlatform.Automation.cs
+++ b/Engine/Source/Programs/AutomationTool/Apple/ApplePlatform.Automation.cs
@@ -93,6 +93,13 @@ public abstract class ApplePlatform : Platform
 
 	public override void PostStagingFileCopy(ProjectParams Params, DeploymentContext SC)
 	{
+		// If we are cooking a DLC, we should never be staging executables so we can early out
+		if (Params.HasDLCName)
+		{
+			Logger.LogInformation("Skipping PostStagingFileCopy because we are cooking a DLC plugin");
+			return;
+		}
+
 		// staging will put binaries into Staged/<game>/Binaries/<platform> and they aren't needed, and when we pull this into a .app, it 
 		// messes with the resulting .app. So, we remove the game binary now (leaving in helper .app's and raw .dylibs, etc)
 		// they come from BuildProducts, and we could maybe remove from that list, but it could cause issues with Horde/buildmachines
