AppLocalDependencies are copied to the current target platform's binary folders

 - fixes AppLocalDependencies for Windows-based platform extensions

diff --git a/Engine/Source/Programs/AutomationTool/Win/WinPlatform.Automation.cs b/Engine/Source/Programs/AutomationTool/Win/WinPlatform.Automation.cs
--- a/Engine/Source/Programs/AutomationTool/Win/WinPlatform.Automation.cs
+++ b/Engine/Source/Programs/AutomationTool/Win/WinPlatform.Automation.cs
@@ -367,7 +367,7 @@ public class Win64Platform : Platform
 		return true;
 	}
 
-	public void StageAppLocalDependencies(ProjectParams Params, DeploymentContext SC, string PlatformDir)
+	public void StageAppLocalDependencies(ProjectParams Params, DeploymentContext SC, string SourcePlatformDir)
 	{
 		Dictionary<string, string> PathVariables = new Dictionary<string, string>();
 		PathVariables["EngineDir"] = SC.EngineRoot.FullName;
@@ -379,11 +379,11 @@ public class Win64Platform : Platform
 		{
 			string ExpandedAppLocalDir = Utils.ExpandVariables(AppLocalDirectory, PathVariables);
 
-			DirectoryReference BaseAppLocalDependenciesPath = Path.IsPathRooted(ExpandedAppLocalDir) ? new DirectoryReference(CombinePaths(ExpandedAppLocalDir, PlatformDir)) : DirectoryReference.Combine(SC.ProjectRoot, ExpandedAppLocalDir, PlatformDir);
+			DirectoryReference BaseAppLocalDependenciesPath = Path.IsPathRooted(ExpandedAppLocalDir) ? new DirectoryReference(CombinePaths(ExpandedAppLocalDir, SourcePlatformDir)) : DirectoryReference.Combine(SC.ProjectRoot, ExpandedAppLocalDir, SourcePlatformDir);
 			if (DirectoryReference.Exists(BaseAppLocalDependenciesPath))
 			{
-				StageAppLocalDependenciesToDir(SC, BaseAppLocalDependenciesPath, StagedDirectoryReference.Combine("Engine", "Binaries", PlatformDir));
-				StageAppLocalDependenciesToDir(SC, BaseAppLocalDependenciesPath, StagedDirectoryReference.Combine(SC.RelativeProjectRootForStage, "Binaries", PlatformDir));
+				StageAppLocalDependenciesToDir(SC, BaseAppLocalDependenciesPath, StagedDirectoryReference.Combine("Engine", "Binaries", PlatformType.ToString()));
+				StageAppLocalDependenciesToDir(SC, BaseAppLocalDependenciesPath, StagedDirectoryReference.Combine(SC.RelativeProjectRootForStage, "Binaries", PlatformType.ToString()));
 			}
 			else
 			{
