Removing pak padding for iOS, mimicing Android since CL 18200183. Trims ~10% from the total install size at the cost of increased patch sizes, though should be at parity with Android.
#tests Preflight

diff --git a/Engine/Source/Programs/AutomationTool/IOS/IOSPlatform.Automation.cs b/Engine/Source/Programs/AutomationTool/IOS/IOSPlatform.Automation.cs
--- a/Engine/Source/Programs/AutomationTool/IOS/IOSPlatform.Automation.cs
+++ b/Engine/Source/Programs/AutomationTool/IOS/IOSPlatform.Automation.cs
@@ -567,7 +567,7 @@ public class IOSPlatform : ApplePlatform
 
 	public override string GetPlatformPakCommandLine(ProjectParams Params, DeploymentContext SC)
 	{
-		string PakParams = "";
+		string PakParams = " -patchpaddingalign=0";
 
 		string OodleDllPath = DirectoryReference.Combine(SC.ProjectRoot, "Binaries/ThirdParty/Oodle/Mac/libUnrealPakPlugin.dylib").FullName;
 		if (File.Exists(OodleDllPath))
