[Gauntlet] UE.Automation to report fatal error with a callstack property

#rnx

diff --git a/Engine/Source/Programs/AutomationTool/Gauntlet/Unreal/Automation/UE.Automation.cs b/Engine/Source/Programs/AutomationTool/Gauntlet/Unreal/Automation/UE.Automation.cs
--- a/Engine/Source/Programs/AutomationTool/Gauntlet/Unreal/Automation/UE.Automation.cs
+++ b/Engine/Source/Programs/AutomationTool/Gauntlet/Unreal/Automation/UE.Automation.cs
@@ -761,14 +761,14 @@ namespace UE
 						}
 						else
 						{
-							ErrorMessage = "No callstack found in the log.";
+							ErrorMessage = "Engine crashed. No callstack could be found in the log.";
 						}
 					}
 					if (LastTestInProgress != null)
 					{
 						if (!String.IsNullOrEmpty(ErrorMessage))
 						{
-							LastTestInProgress.AddError(ErrorMessage, !HasTimeout);
+							LastTestInProgress.AddError(ErrorMessage, !HasTimeout && FatalError != null);
 						}
 						if (!CanRetry() || JsonTestPassResults.NotRun == 0)
 						{
@@ -1051,8 +1051,14 @@ namespace UE
 								var Errors = CapErrorOrWarningList(Result.ErrorEvents.Distinct());
 								foreach (var Error in Errors)
 								{
-									EventId ErrorEventType = Error.IsCriticalFailure ? KnownLogEvents.Gauntlet_FatalEvent : KnownLogEvents.Gauntlet_UnrealEngineTestEvent;
-									Log.Error(ErrorEventType, "    " + Error.FormatToString());
+									if (Error.IsCriticalFailure)
+									{
+										Log.Error(KnownLogEvents.Gauntlet_FatalEvent, "    {Callstack}", Error.FormatToString());
+									}
+									else
+									{
+										Log.Error(KnownLogEvents.Gauntlet_UnrealEngineTestEvent, "    " + Error.FormatToString());
+									}
 								}
 								NotifyMoreIfNeeded(Result.ErrorEvents);
 								var Warnings = CapErrorOrWarningList(Result.WarningEvents.Distinct());
@@ -1119,8 +1125,14 @@ namespace UE
 								IEnumerable<UnrealAutomationEvent> Events = Result.ErrorEvents.Distinct();
 								foreach (var Event in CapErrorOrWarningList(Events))
 								{
-									EventId ErrorEventType = Event.IsCriticalFailure ? KnownLogEvents.Gauntlet_FatalEvent : KnownLogEvents.Gauntlet_UnrealEngineTestEvent;
-									Log.Error(ErrorEventType, "    " + Event.FormatToString());
+									if (Event.IsCriticalFailure)
+									{
+										Log.Error(KnownLogEvents.Gauntlet_FatalEvent, "    {Callstack}", Event.FormatToString());
+									}
+									else
+									{
+										Log.Error(KnownLogEvents.Gauntlet_UnrealEngineTestEvent, "    " + Event.FormatToString());
+									}
 								}
 								NotifyMoreIfNeeded(Events);
 								Log.Info("");
