[Gauntlet] Add EditorCommand test to run any editor command. Validate only if the Editor was initialized and exited normally or when string in found in the log.

#rnx

diff --git a/Engine/Source/Programs/AutomationTool/Gauntlet/Unreal/Editor/EditorTests.cs b/Engine/Source/Programs/AutomationTool/Gauntlet/Unreal/Editor/EditorTests.cs
--- a/Engine/Source/Programs/AutomationTool/Gauntlet/Unreal/Editor/EditorTests.cs
+++ b/Engine/Source/Programs/AutomationTool/Gauntlet/Unreal/Editor/EditorTests.cs
@@ -432,3 +432,56 @@ namespace UEPerf
 		}
 	}
 }
+
+namespace UE
+{
+	/// <summary>
+	/// Test to run any editor command. Validate only if the Editor was initialized and exited normally or when string in found in the log.
+	/// Exit needs to be explicitly call by the command line argument or executed code, or use -CompletionString= to trigger exit on log string match
+	/// </summary>
+	public class EditorCommand : BootTest
+	{
+		private string CompletionString;
+
+		public EditorCommand(Gauntlet.UnrealTestContext InContext)
+			: base(InContext)
+		{
+		}
+
+		public override UnrealTestConfiguration GetConfiguration()
+		{
+			if (CachedConfig == null)
+			{
+				UnrealTestConfiguration Config = base.GetConfiguration();
+				Config.ClearRoles();
+				UnrealTestRole EditorRole = Config.RequireRole(Config.CookedEditor ? UnrealTargetRole.CookedEditor : UnrealTargetRole.Editor);
+
+				// Take all the command line arguments, filter out those we know about that are not useful for the editor and pass the rest through.
+				CompletionString = Globals.Params.ParseValue("CompletionString", string.Empty);
+				List<string> ToFilterOut = new List<string>
+				{
+					"CompletionString=",
+					"project=",
+					"test=",
+					"configuration=",
+					"build=",
+					"platform=",
+					"MaxDuration="
+				};
+				IEnumerable<string> ParamList = null;
+				ParamList = Globals.Params.AllArguments
+								.Where(P => !ToFilterOut.Any(L => P.StartsWith(L, StringComparison.OrdinalIgnoreCase)));
+				EditorRole.CommandLineParams.AddRawCommandline("-" + string.Join(" -", ParamList));
+
+				CachedConfig = Config;
+			}
+
+			return CachedConfig;
+		}
+
+		protected override string GetCompletionString()
+		{
+			return CompletionString;
+		}
+	}
+}
\ No newline at end of file
