CopyTask does not always respect the ErrorIfNotFound=false parameter.  When encoutnering a file mapping without wildcards, if the source file doesn't exist, it will error due to an exception thrown in CreateFileMapping regardless of presence of ErrorIfNotFound=false.  Using the same approach of enclosing CreateFileMapping in a try/catch as is used by MoveTask.cs.

diff --git a/Engine/Source/Programs/AutomationTool/BuildGraph/Tasks/CopyTask.cs b/Engine/Source/Programs/AutomationTool/BuildGraph/Tasks/CopyTask.cs
--- a/Engine/Source/Programs/AutomationTool/BuildGraph/Tasks/CopyTask.cs
+++ b/Engine/Source/Programs/AutomationTool/BuildGraph/Tasks/CopyTask.cs
@@ -102,35 +102,50 @@ namespace AutomationTool.Tasks
 				Files = ResolveFilespec(SourcePattern.BaseDirectory, Parameters.Files, TagNameToFileSet);
 			}
 
-			// Build the file mapping
-			Dictionary<FileReference, FileReference> TargetFileToSourceFile = FilePattern.CreateMapping(Files, ref SourcePattern, ref TargetPattern);
+			try
+			{
+				// Build the file mapping
+				Dictionary<FileReference, FileReference> TargetFileToSourceFile = FilePattern.CreateMapping(Files, ref SourcePattern, ref TargetPattern);
 
-			// Check we got some files
-			if (TargetFileToSourceFile.Count == 0)
+				// Check we got some files
+				if (TargetFileToSourceFile.Count == 0)
+				{
+					if (Parameters.ErrorIfNotFound)
+					{
+						Logger.LogError("No files found matching '{SourcePattern}'", SourcePattern);
+					}
+					else
+					{
+						Logger.LogInformation("No files found matching '{SourcePattern}'", SourcePattern);
+					}
+					return;
+				}
+
+				// Run the copy
+				Logger.LogInformation("Copying {Arg0} file{Arg1} from {Arg2} to {Arg3}...", TargetFileToSourceFile.Count, (TargetFileToSourceFile.Count == 1) ? "" : "s", SourcePattern.BaseDirectory, TargetPattern.BaseDirectory);
+				await ExecuteAsync(TargetFileToSourceFile, Parameters.Overwrite);
+
+				// Update the list of build products
+				BuildProducts.UnionWith(TargetFileToSourceFile.Keys);
+
+				// Apply the optional output tag to them
+				foreach (string TagName in FindTagNamesFromList(Parameters.Tag))
+				{
+					FindOrAddTagSet(TagNameToFileSet, TagName).UnionWith(TargetFileToSourceFile.Keys);
+				}
+			}
+			catch (FilePatternSourceFileMissingException Ex)
 			{
 				if (Parameters.ErrorIfNotFound)
 				{
-					Logger.LogError("No files found matching '{SourcePattern}'", SourcePattern);
+					Logger.LogError(Ex, "Error while trying to create file pattern match for '{SourcePattern}', error {ExceptionString}", SourcePattern, Ex.Message);
 				}
 				else
 				{
-					Logger.LogInformation("No files found matching '{SourcePattern}'", SourcePattern);
+					Logger.LogInformation(Ex, "Error while trying to create file pattern match for '{SourcePattern}', error {ExceptionString}", SourcePattern, Ex.Message);
 				}
-				return;
 			}
 
-			// Run the copy
-			Logger.LogInformation("Copying {Arg0} file{Arg1} from {Arg2} to {Arg3}...", TargetFileToSourceFile.Count, (TargetFileToSourceFile.Count == 1) ? "" : "s", SourcePattern.BaseDirectory, TargetPattern.BaseDirectory);
-			await ExecuteAsync(TargetFileToSourceFile, Parameters.Overwrite);
-
-			// Update the list of build products
-			BuildProducts.UnionWith(TargetFileToSourceFile.Keys);
-
-			// Apply the optional output tag to them
-			foreach (string TagName in FindTagNamesFromList(Parameters.Tag))
-			{
-				FindOrAddTagSet(TagNameToFileSet, TagName).UnionWith(TargetFileToSourceFile.Keys);
-			}
 		}
 
 		/// <summary>
