[Gauntlet] Support multiword command line argument pass through Gauntlet param ie: -some-argument

#rnx

diff --git a/Engine/Source/Programs/AutomationTool/Gauntlet/SelfTest/Unreal/Gauntlet.SelfTest.TestUnrealBuildSource.cs b/Engine/Source/Programs/AutomationTool/Gauntlet/SelfTest/Unreal/Gauntlet.SelfTest.TestUnrealBuildSource.cs
--- a/Engine/Source/Programs/AutomationTool/Gauntlet/SelfTest/Unreal/Gauntlet.SelfTest.TestUnrealBuildSource.cs
+++ b/Engine/Source/Programs/AutomationTool/Gauntlet/SelfTest/Unreal/Gauntlet.SelfTest.TestUnrealBuildSource.cs
@@ -24,20 +24,22 @@ namespace Gauntlet.SelfTest
 	[TestGroup("Unreal", 4)]
 	class TestUnrealBuildSource : TestUnrealBase
 	{
+		const string extra_arguments = "-somearg -some=\"arg\" -some[arg]:with/path -some-arg";
 		/// <summary>
 		/// Test entry point
 		/// </summary>
 		public override void TickTest()
 		{
 			// create the build source
-			UnrealBuildSource BuildSource = new UnrealBuildSource(this.ProjectName, ProjectFile, this.UnrealPath, UsesSharedBuildType, BuildPath, new string[] { "" });
+			UnrealBuildSource BuildSource = new UnrealBuildSource(this.ProjectName, ProjectFile, this.UnrealPath, UsesSharedBuildType, BuildPath);
 
-			// check editor and statged info is valid
-			CheckResult(BuildSource.EditorValid, "Editor build was invalid");
-			//CheckResult(BuildSource.BuildCount > 0, "staged build was invalid");
+			// check if editor build can be found
+			CheckResult(BuildSource.GetBuildCount(BuildHostPlatform.Current.Platform, BuildFlags.None) > 0, "No Editor build found.");
 
 			// simple check with an editor role
-			UnrealSessionRole EditorRole = new UnrealSessionRole(UnrealTargetRole.Editor, BuildHostPlatform.Current.Platform, UnrealTargetConfiguration.Development);
+			UnrealSessionRole EditorRole = new UnrealSessionRole(UnrealTargetRole.Editor, BuildHostPlatform.Current.Platform, Configuration);
+			// set extra arguments
+			EditorRole.CommandLineParams.AddRawCommandline(extra_arguments);
 
 			List<string> Reasons = new List<string>();
 
@@ -50,6 +52,7 @@ namespace Gauntlet.SelfTest
 			CheckResult(Config != null, "Build source did not return a config for {0}", EditorRole.ToString());
 
 			ValidateEditorConfig(Config, BuildSource);
+			ValidateEditorExtraArguments(Config, BuildSource);
 
 			// Check all editor types (game, server, etc)
 			TestBuildSourceForEditorTypes(BuildSource);
@@ -71,7 +74,7 @@ namespace Gauntlet.SelfTest
 			{
 				if (E.UsesEditor())
 				{
-					UnrealAppConfig Config = BuildSource.CreateConfiguration(new UnrealSessionRole(E, UnrealTargetPlatform.Win64, UnrealTargetConfiguration.Development));
+					UnrealAppConfig Config = BuildSource.CreateConfiguration(new UnrealSessionRole(E, BuildHostPlatform.Current.Platform, UnrealTargetConfiguration.Development));
 
 					CheckResult(Config != null, "Editor config for {0} returned null!", E);
 					CheckResult(string.IsNullOrEmpty(Config.Name) == false, "No config name!");
@@ -136,7 +139,7 @@ namespace Gauntlet.SelfTest
 			CheckResult(Config.ProcessType.UsesEditor(), "Config does not use editor!");
 
 			// Check the project name was the first arg
-			CheckResult(Args.IndexOf(BuildSource.ProjectName.ToLower()) == 0, "Editor-based config for {0} needs to include project name as first argument", Config.ProcessType);
+			CheckResult(Args.Contains(BuildSource.ProjectName.ToLower()), "Editor-based config for {0} needs to include project name as first argument", Config.ProcessType);
 
 			// for clients, check for -game
 			if (Config.ProcessType.IsClient())
@@ -150,5 +153,19 @@ namespace Gauntlet.SelfTest
 				CheckResult(Args.Contains("-server"), "Editor-based game needs to include -server");
 			}
 		}
+
+		/// <summary>
+		/// Simple validation for config extra arguments passthrough
+		/// </summary>
+		/// <param name="Config"></param>
+		/// <param name="BuildSource"></param>
+		void ValidateEditorExtraArguments(UnrealAppConfig Config, UnrealBuildSource BuildSource)
+		{
+			string Args = Config.CommandLine.Trim().ToLower();
+			foreach (var expected in extra_arguments.Split(" "))
+			{
+				CheckResult(Args.Contains(expected), $"Argument {expected} not passed to Editor command line.");
+			}
+		}
 	}
 }
diff --git a/Engine/Source/Programs/AutomationTool/Gauntlet/Unreal/BuildSource/Gauntlet.UnrealBuildSource.cs b/Engine/Source/Programs/AutomationTool/Gauntlet/Unreal/BuildSource/Gauntlet.UnrealBuildSource.cs
--- a/Engine/Source/Programs/AutomationTool/Gauntlet/Unreal/BuildSource/Gauntlet.UnrealBuildSource.cs
+++ b/Engine/Source/Programs/AutomationTool/Gauntlet/Unreal/BuildSource/Gauntlet.UnrealBuildSource.cs
@@ -587,7 +587,7 @@ namespace Gauntlet
 			// Break down Commandline into individual tokens 
 			Dictionary<string, string> CommandlineTokens = new Dictionary<string, string>(StringComparer.OrdinalIgnoreCase);
 			// turn Name(p1,etc) into a collection of Name|(p1,etc) groups
-			MatchCollection Matches = Regex.Matches(InCommandLine, "(?<option>\\-?[\\w\\d.:!\\[\\]\\/\\\\\\?]+)(=(?<value>(\"([^\"]*)\")|(\\S+)))?");
+			MatchCollection Matches = Regex.Matches(InCommandLine, @"(?<option>-?[\w\d.:!\[\]\/\\?-]+)(=(?<value>(""([^""]*)"")|(\S+)))?");
 
 			foreach (Match M in Matches)
 			{
