[Gauntlet] Disable log heartbeat categories if running OOM when loading process stdout into a string

diff --git a/Engine/Source/Programs/AutomationTool/Gauntlet/Unreal/Base/Gauntlet.UnrealTestNode.cs b/Engine/Source/Programs/AutomationTool/Gauntlet/Unreal/Base/Gauntlet.UnrealTestNode.cs
--- a/Engine/Source/Programs/AutomationTool/Gauntlet/Unreal/Base/Gauntlet.UnrealTestNode.cs
+++ b/Engine/Source/Programs/AutomationTool/Gauntlet/Unreal/Base/Gauntlet.UnrealTestNode.cs
@@ -320,7 +320,7 @@ namespace Gauntlet
 		/// Standard semantic versioning for tests. Should be overwritten within individual tests, and individual test maintainers
 		/// are responsible for updating their own versions. See https://semver.org/ for more info on maintaining semantic versions.
 		/// </summary>
-		/// 
+
 		protected Version TestVersion;
 
 		/// <summary>
@@ -338,12 +338,14 @@ namespace Gauntlet
 		protected DateTime TimeOfFirstMissingProcess;
 
 		protected int TimeToWaitForProcesses { get; set; }
-		
+
+		protected bool bDisableHeartbeatLogging = false;
+
 		protected DateTime LastHeartbeatTime = DateTime.MinValue;
 		protected DateTime LastActiveHeartbeatTime = DateTime.MinValue;
 		private DateTime LastHeartbeatLogTime = DateTime.MinValue;
-		
-		// End  UnrealTestNode properties and members 
+
+		// End  UnrealTestNode properties and members
 
 		// artifact paths that have been used in this run
 		static protected HashSet<string> ReservedArtifcactPaths = new HashSet<string>();
@@ -1765,7 +1767,7 @@ namespace Gauntlet
 			{
 				return;
 			}
-			
+
 			List<string> LogCategories = new List<string>();
 			LogCategories.Add("Gauntlet");
 			{
@@ -1777,28 +1779,29 @@ namespace Gauntlet
 
 			void LogHeartbeatCategories(ref int LastLogCount, IAppInstance App, string AppPrefix, bool bUpdateHeartbeatTime)
 			{
-				if (App != null)
+				if (App != null && !bDisableHeartbeatLogging)
 				{
-					UnrealLogStreamParser Parser = new UnrealLogStreamParser();
-					LastLogCount += Parser.ReadStream(App.StdOut, LastLogCount);
-
-					foreach (string TestLine in Parser.GetLogFromShortNameChannels(LogCategories))
+					string StdOut;
+					try
 					{
-						Log.Info(string.Format("{0}: {1}", AppPrefix, TestLine));
+						// Convert the ProcessOutput StringBuilder into a regular string. This can run OOM if the output is sufficently large...
+						StdOut = App.StdOut;
+					}
+					catch (OutOfMemoryException OOMEx)
+					{
+						Log.Warning("{App} encountered an {ExceptionType} when attempting to read the process output. " +
+							"This is usually caused by the output being >4GB in size due to log spam or deformed encoding. " +
+							"Heartbeat logging will be disabled for the remainder of the test.\n" +
+							"{Exception}", App.ToString(), OOMEx.GetType().Name, OOMEx.Message);
 
-						if (bUpdateHeartbeatTime)
-						{
-							if (Regex.IsMatch(TestLine, @".*GauntletHeartbeat\: Active.*"))
-							{
-								LastHeartbeatTime = DateTime.Now;
-								LastActiveHeartbeatTime = DateTime.Now;
-							}
-							else if (Regex.IsMatch(TestLine, @".*GauntletHeartbeat\: Idle.*"))
-							{
-								LastHeartbeatTime = DateTime.Now;
-							}
-						}
+						bDisableHeartbeatLogging = true;
+						CachedConfig.HeartbeatOptions.bExpectHeartbeats = false; // prevents CheckHeartbeat from timing out tests after heartbeats are disabled
+
+						return;
 					}
+
+					UnrealLogStreamParser Parser = new UnrealLogStreamParser();
+					LastLogCount += Parser.ReadStream(StdOut, LastLogCount);
 				}
 			}
 
