[Gauntlet] Adding a log idle timeout command line argument to UE.BootTest

#rnx

diff --git a/Engine/Source/Programs/AutomationTool/Gauntlet/Unreal/Automation/UE.BootTest.cs b/Engine/Source/Programs/AutomationTool/Gauntlet/Unreal/Automation/UE.BootTest.cs
--- a/Engine/Source/Programs/AutomationTool/Gauntlet/Unreal/Automation/UE.BootTest.cs
+++ b/Engine/Source/Programs/AutomationTool/Gauntlet/Unreal/Automation/UE.BootTest.cs
@@ -31,6 +31,11 @@ namespace UE
 		/// </summary>
 		bool DidDetectLaunch = false;
 
+		/// <summary>
+		/// Log idle timeout in seconds
+		/// </summary>
+		private float kTimeOutDuration = 10 * 60;
+
 		/// <summary>
 		/// Default constructor
 		/// </summary>
@@ -54,6 +59,12 @@ namespace UE
 				Client.CommandLineParams.Add("ExecCmds", "Automation SoftQuit");
 			}
 
+			float logIdleTimeout = Globals.Params.ParseValue("LogIdleTimeout", kTimeOutDuration);
+			if (logIdleTimeout > 0)
+			{
+				kTimeOutDuration = logIdleTimeout;
+			}
+
 			return Config;
 		}
 
@@ -94,8 +105,6 @@ namespace UE
 		/// </summary>
 		public override void TickTest()
 		{
-			const int kTimeOutDuration = 10;
-
 			// run the base class tick;
 			base.TickTest();
 
@@ -115,9 +124,9 @@ namespace UE
 
 			// Gauntlet will timeout tests based on the -timeout argument, but we have greater insight here so can bail earlier to save
 			// tests from idling on the farm needlessly.
-			if ((DateTime.Now - LastLogTime).TotalMinutes > kTimeOutDuration)
+			if ((DateTime.Now - LastLogTime).TotalSeconds > kTimeOutDuration)
 			{
-				Log.Error("No logfile activity observed in last {0} minutes. Ending test", kTimeOutDuration);
+				Log.Error("No logfile activity observed in last {Time:0.00} minutes. Ending test", kTimeOutDuration / 60);
 				MarkTestComplete();
 				SetUnrealTestResult(TestResult.TimedOut);
 			}
