Gauntlet: Fix LogFile sometimes not being set in DesktopCommonAppInstall
When the RegEx matches the first group, the LogFile name was empty

diff --git a/Engine/Source/Programs/AutomationTool/Gauntlet/Framework/Base/Gauntlet.TargetDeviceDesktopCommon.cs b/Engine/Source/Programs/AutomationTool/Gauntlet/Framework/Base/Gauntlet.TargetDeviceDesktopCommon.cs
--- a/Engine/Source/Programs/AutomationTool/Gauntlet/Framework/Base/Gauntlet.TargetDeviceDesktopCommon.cs
+++ b/Engine/Source/Programs/AutomationTool/Gauntlet/Framework/Base/Gauntlet.TargetDeviceDesktopCommon.cs
@@ -389,7 +389,7 @@ namespace Gauntlet
 				Match M = LogRegex.Match(CommandArguments);
 				if (M.Success)
 				{
-					LogFile = M.Groups[2].Value;
+					LogFile = !string.IsNullOrEmpty(M.Groups[1].Value) ? M.Groups[1].Value : M.Groups[2].Value;
 				}
 
 				// Explicitly set log file when not already defined if not build machine
