[Gauntlet] (On Windows) Generate mini dump when killing application if timeout is detected.

#rnx

diff --git a/Engine/Source/Programs/AutomationTool/Gauntlet/Framework/Base/Gauntlet.AppInstance.cs b/Engine/Source/Programs/AutomationTool/Gauntlet/Framework/Base/Gauntlet.AppInstance.cs
--- a/Engine/Source/Programs/AutomationTool/Gauntlet/Framework/Base/Gauntlet.AppInstance.cs
+++ b/Engine/Source/Programs/AutomationTool/Gauntlet/Framework/Base/Gauntlet.AppInstance.cs
@@ -50,7 +50,7 @@ namespace Gauntlet
 		/// <summary>
 		/// Kills the process if its running (no need to call WaitForExit)
 		/// </summary>
-		void Kill();
+		void Kill(bool GenerateDumpOnKill = false);
 
 		/// <summary>
 		/// Waits for the process to exit normally
diff --git a/Engine/Source/Programs/AutomationTool/Gauntlet/Framework/Base/Gauntlet.TargetDeviceDesktopCommon.cs b/Engine/Source/Programs/AutomationTool/Gauntlet/Framework/Base/Gauntlet.TargetDeviceDesktopCommon.cs
--- a/Engine/Source/Programs/AutomationTool/Gauntlet/Framework/Base/Gauntlet.TargetDeviceDesktopCommon.cs
+++ b/Engine/Source/Programs/AutomationTool/Gauntlet/Framework/Base/Gauntlet.TargetDeviceDesktopCommon.cs
@@ -11,7 +11,6 @@ using UnrealBuildTool;
 using System.Text.RegularExpressions;
 using Gauntlet.Utils;
 
-
 namespace Gauntlet
 {
 	public abstract class TargetDeviceDesktopCommon : ITargetDevice
@@ -106,8 +105,13 @@ namespace Gauntlet
 			{
 				CleanArtifactDirectory(GetInstallArtifactPath());
 			}
+			// Clean crash dumps
+			ITargetDevice ThisDevice = (ITargetDevice)this;
+			CleanArtifactDirectory(ThisDevice.CrashDumpPath);
 		}
 
+		public virtual bool CopyCrashDumps() { return false; }
+
 		public virtual void InstallBuild(UnrealAppConfig AppConfig)
 		{
 			IBuild Build = AppConfig.Build;
@@ -472,15 +476,21 @@ namespace Gauntlet
 			return ExitCode;
 		}
 
-		virtual public void Kill()
+		virtual public void Kill(bool bGenerateDump)
 		{
 			if (!HasExited)
 			{
+				if (bGenerateDump) 
+				{
+					GenerateDump();
+				}
 				WasKilled = true;
 				ProcessResult.ProcessObject.Kill(true);
 			}
 		}
 
+		virtual protected void GenerateDump() { }
+
 		/// <summary>
 		/// Reader thread when logging to file
 		/// </summary>
diff --git a/Engine/Source/Programs/AutomationTool/Gauntlet/Platform/Android/Gauntlet.TargetDeviceAndroid.cs b/Engine/Source/Programs/AutomationTool/Gauntlet/Platform/Android/Gauntlet.TargetDeviceAndroid.cs
--- a/Engine/Source/Programs/AutomationTool/Gauntlet/Platform/Android/Gauntlet.TargetDeviceAndroid.cs
+++ b/Engine/Source/Programs/AutomationTool/Gauntlet/Platform/Android/Gauntlet.TargetDeviceAndroid.cs
@@ -1558,7 +1558,7 @@ namespace Gauntlet
 			return ExitCode;
 		}
 
-		public void Kill()
+		public void Kill(bool GenerateDump)
 		{
 			if (!HasExited && !AndroidDevice.Disposed)
 			{
diff --git a/Engine/Source/Programs/AutomationTool/Gauntlet/Platform/IOS/Gauntlet.TargetDeviceIOS.cs b/Engine/Source/Programs/AutomationTool/Gauntlet/Platform/IOS/Gauntlet.TargetDeviceIOS.cs
--- a/Engine/Source/Programs/AutomationTool/Gauntlet/Platform/IOS/Gauntlet.TargetDeviceIOS.cs
+++ b/Engine/Source/Programs/AutomationTool/Gauntlet/Platform/IOS/Gauntlet.TargetDeviceIOS.cs
@@ -117,7 +117,7 @@ namespace Gauntlet
 			return ExitCode;
 		}
 
-		public void Kill()
+		public void Kill(bool GenerateDump)
 		{
 			if (!HasExited)
 			{
diff --git a/Engine/Source/Programs/AutomationTool/Gauntlet/Platform/Null/Gauntlet.TargetDeviceNull.cs b/Engine/Source/Programs/AutomationTool/Gauntlet/Platform/Null/Gauntlet.TargetDeviceNull.cs
--- a/Engine/Source/Programs/AutomationTool/Gauntlet/Platform/Null/Gauntlet.TargetDeviceNull.cs
+++ b/Engine/Source/Programs/AutomationTool/Gauntlet/Platform/Null/Gauntlet.TargetDeviceNull.cs
@@ -39,7 +39,7 @@ namespace Gauntlet
 
 		public int ExitCode { get { return 0; }}
 
-		public void Kill()
+		public void Kill(bool GenerateDump)
 		{
 			if (!HasExited)
 			{
diff --git a/Engine/Source/Programs/AutomationTool/Gauntlet/Platform/Windows/Gauntlet.TargetDeviceWindows.cs b/Engine/Source/Programs/AutomationTool/Gauntlet/Platform/Windows/Gauntlet.TargetDeviceWindows.cs
--- a/Engine/Source/Programs/AutomationTool/Gauntlet/Platform/Windows/Gauntlet.TargetDeviceWindows.cs
+++ b/Engine/Source/Programs/AutomationTool/Gauntlet/Platform/Windows/Gauntlet.TargetDeviceWindows.cs
@@ -6,6 +6,9 @@ using AutomationTool;
 using UnrealBuildTool;
 using Gauntlet.Utils;
 using static AutomationTool.ProcessResult;
+using EpicGames.Core;
+using System.ComponentModel;
+using System.Runtime.InteropServices;
 
 namespace Gauntlet
 {
@@ -170,6 +173,26 @@ namespace Gauntlet
 			return (InstallCache as DesktopCommonAppInstall<TargetDeviceWindows>).ArtifactPath;
 		}
 
+		public override bool CopyCrashDumps()
+		{
+			// On Windows platform, check only if any dmp file exists.
+			// It does not need to be fetched from a remote location.
+			ITargetDevice ThisDevice = (ITargetDevice)this;
+			DirectoryInfo DirInfo = new DirectoryInfo(ThisDevice.CrashDumpPath);
+			if (DirInfo.Exists)
+			{
+				// See if there is a crash dump for this test run
+				foreach (FileInfo CrashDumpFileInfo in DirInfo.GetFiles("*.dmp", SearchOption.AllDirectories))
+				{
+					if (CrashDumpFileInfo.Exists)
+					{
+						return true;
+					}
+				}
+			}
+			return false;
+		}
+
 		#region Legacy Implementations
 		public override IAppInstall InstallApplication(UnrealAppConfig AppConfig)
 		{
@@ -316,9 +339,72 @@ namespace Gauntlet
 
 	public class WindowsAppInstance : DesktopCommonAppInstance<WindowsAppInstall, TargetDeviceWindows>
 	{
+		[DllImport("dbghelp", SetLastError = true)]
+		private static extern bool MiniDumpWriteDump(SafeHandle hProcess, uint ProcessId, SafeHandle hFile, MINIDUMP_TYPE DumpType, IntPtr ExceptionParam, IntPtr UserStreamParam, IntPtr CallbackParam);
+
+		private enum MINIDUMP_TYPE
+		{
+			MiniDumpNormal = 0x00000000,
+			MiniDumpWithDataSegs = 0x00000001,
+			MiniDumpWithFullMemory = 0x00000002,
+			MiniDumpWithHandleData = 0x00000004,
+			MiniDumpFilterMemory = 0x00000008,
+			MiniDumpScanMemory = 0x00000010,
+			MiniDumpWithUnloadedModules = 0x00000020,
+			MiniDumpWithIndirectlyReferencedMemory = 0x00000040,
+			MiniDumpFilterModulePaths = 0x00000080,
+			MiniDumpWithProcessThreadData = 0x00000100,
+			MiniDumpWithPrivateReadWriteMemory = 0x00000200,
+			MiniDumpWithoutOptionalData = 0x00000400,
+			MiniDumpWithFullMemoryInfo = 0x00000800,
+			MiniDumpWithThreadInfo = 0x00001000,
+			MiniDumpWithCodeSegs = 0x00002000,
+			MiniDumpWithoutAuxiliaryState = 0x00004000,
+			MiniDumpWithFullAuxiliaryState = 0x00008000,
+			MiniDumpWithPrivateWriteCopyMemory = 0x00010000,
+			MiniDumpIgnoreInaccessibleMemory = 0x00020000,
+			MiniDumpWithTokenInformation = 0x00040000,
+			MiniDumpWithModuleHeaders = 0x00080000,
+			MiniDumpFilterTriage = 0x00100000,
+			MiniDumpWithAvxXStateContext = 0x00200000,
+			MiniDumpWithIptTrace = 0x00400000,
+			MiniDumpScanInaccessiblePartialPages = 0x00800000,
+			MiniDumpFilterWriteCombinedMemory,
+			MiniDumpValidTypeFlags = 0x01ffffff
+		}
+
 		public WindowsAppInstance(WindowsAppInstall InInstall, IProcessResult InProcess, string InProcessLogFile = null)
 			: base(InInstall, InProcess, InProcessLogFile)
 		{ }
+
+		protected override void GenerateDump()
+		{
+			bool WroteDump = false;
+			if (!Directory.Exists(Device.CrashDumpPath))
+			{
+				Directory.CreateDirectory(Device.CrashDumpPath);
+			}
+			string DumpName = Path.Combine(Device.CrashDumpPath, Path.GetFileNameWithoutExtension(ProcessResult.ProcessObject.ProcessName) + ".dmp");
+			using (FileStream CrashDumpStream = File.Create(DumpName))
+			{
+				WroteDump = MiniDumpWriteDump(ProcessResult.ProcessObject.SafeHandle, (uint)ProcessResult.ProcessObject.Id,
+					CrashDumpStream.SafeFileHandle, MINIDUMP_TYPE.MiniDumpNormal, IntPtr.Zero, IntPtr.Zero, IntPtr.Zero);
+				if (!WroteDump)
+				{
+					Int32 Error = Marshal.GetLastWin32Error();
+					string ErrorMessage = new Win32Exception(Error).Message;
+					Log.Error(KnownLogEvents.Gauntlet, "Failed to write minidump. Error: {Error} (GetLastError={LastError})", ErrorMessage, Error);
+				}
+			}
+			if (!WroteDump)
+			{
+				File.Delete(DumpName);
+			}
+			else
+			{
+				Log.Info("Wrote minidump to {FileName}", DumpName);
+			}
+		}
 	}
 
 	public class Win64DeviceFactory : IDeviceFactory
diff --git a/Engine/Source/Programs/AutomationTool/Gauntlet/Unreal/Base/Gauntlet.UnrealSession.cs b/Engine/Source/Programs/AutomationTool/Gauntlet/Unreal/Base/Gauntlet.UnrealSession.cs
--- a/Engine/Source/Programs/AutomationTool/Gauntlet/Unreal/Base/Gauntlet.UnrealSession.cs
+++ b/Engine/Source/Programs/AutomationTool/Gauntlet/Unreal/Base/Gauntlet.UnrealSession.cs
@@ -467,7 +467,7 @@ namespace Gauntlet
 		/// Shutdown the session by killing any remaining processes.
 		/// </summary>
 		/// <returns></returns>
-		public void Shutdown()
+		public void Shutdown(bool GenerateDumpOnKill = false)
 		{
 			// Kill any remaining client processes
 			if (ClientApps != null)
@@ -479,7 +479,7 @@ namespace Gauntlet
 					Log.Info("Shutting down {0} clients", RunningApps.Count);
 					RunningApps.ForEach(App =>
 					{
-						App.Kill();
+						App.Kill(GenerateDumpOnKill);
 						// Apps that are still running have timed out => fail
 						IDeviceUsageReporter.RecordEnd(App.Device.Name, App.Device.Platform, IDeviceUsageReporter.EventType.Test, IDeviceUsageReporter.EventState.Success);
 					});
@@ -501,12 +501,12 @@ namespace Gauntlet
 				if (ServerApp.HasExited == false)
 				{
 					Log.Info("Shutting down server");
-					ServerApp.Kill();
+					ServerApp.Kill(GenerateDumpOnKill);
 				}
 			}
 
 			// kill anything that's left
-			RunningRoles.Where(R => !R.Role.InstallOnly && R.AppInstance.HasExited == false).ToList().ForEach(R => R.AppInstance.Kill());
+			RunningRoles.Where(R => !R.Role.InstallOnly && R.AppInstance.HasExited == false).ToList().ForEach(R => R.AppInstance.Kill(GenerateDumpOnKill));
 
 			// Wait for it all to end
 			RunningRoles.Where(R => !R.Role.InstallOnly).ToList().ForEach(R => R.AppInstance.WaitForExit());
diff --git a/Engine/Source/Programs/AutomationTool/Gauntlet/Unreal/Base/Gauntlet.UnrealTestNode.cs b/Engine/Source/Programs/AutomationTool/Gauntlet/Unreal/Base/Gauntlet.UnrealTestNode.cs
--- a/Engine/Source/Programs/AutomationTool/Gauntlet/Unreal/Base/Gauntlet.UnrealTestNode.cs
+++ b/Engine/Source/Programs/AutomationTool/Gauntlet/Unreal/Base/Gauntlet.UnrealTestNode.cs
@@ -5,12 +5,8 @@ using System.Collections.Generic;
 using System.IO;
 using AutomationTool;
 using UnrealBuildTool;
-using System.Threading;
 using System.Text.RegularExpressions;
-using System.Drawing;
 using System.Linq;
-using System.Text;
-using System.Security.Cryptography;
 using EpicGames.Core;
 
 namespace Gauntlet
@@ -1073,7 +1069,7 @@ namespace Gauntlet
 
 			// Shutdown the instance so we can access all files, but do not null it or shutdown the UnrealApp because we still need
 			// access to these objects and their resources! Final cleanup is done in CleanupTest()
-			TestInstance.Shutdown();
+			TestInstance.Shutdown(InReason == StopReason.MaxDuration || GetTestResult() == TestResult.TimedOut);
 
 			try
 			{
