Fix InstallUnrealBuild command by considering configuration and package type

diff --git a/Engine/Build/Graph/Tasks/DeviceOperation.xml b/Engine/Build/Graph/Tasks/DeviceOperation.xml
new file mode 100644
--- /dev/null
+++ b/Engine/Build/Graph/Tasks/DeviceOperation.xml
@@ -0,0 +1,50 @@
+<?xml version='1.0' ?>
+<BuildGraph xmlns="http://www.epicgames.com/BuildGraph" xmlns:xsi="http://www.w3.org/2001/XMLSchema-instance" xsi:schemaLocation="http://www.epicgames.com/BuildGraph ../Schema.xsd" >
+
+	<Include Script="Inc/SharedProperties.xml"/>
+	
+	<Option Name="InstallBuildFlags" DefaultValue="" Description="Optional BuildFlags used to specific a type of build to install."/>
+	<Option Name="ProjectName" DefaultValue="" Description="The Project Name of the build."/>
+
+	<Option Name="Configurations" DefaultValue="" Description="The Configurations of the build."/>
+	<Option Name="WithShipping" DefaultValue="false" Restrict="true|false" Description="With Shipping Configuration."/>
+	<Property Name="Configurations" Value="$(Configurations);Shipping" If="$(WithShipping)"/>
+	<Option Name="WithTest" DefaultValue="false" Restrict="true|false" Description="With Test Configuration."/>
+	<Property Name="Configurations" Value="$(Configurations);Test" If="$(WithTest)"/>
+	<Option Name="WithDevelopment" DefaultValue="true" Restrict="true|false" Description="With Development Configuration."/>
+	<Property Name="Configurations" Value="$(Configurations);Development" If="$(WithDevelopment)"/>
+
+	<Option Name="BranchTag" DefaultValue="" Description="The tag to the branch."/>
+	<Property Name="BranchTag" Value="$(EscapedBranch)" If="'$(BranchTag)' == ''"/>
+	<Property Name="NetworkPublishDir" Value="$(NetworkPublishRoot)/$(BranchTag)"/>
+	<Option Name="PublishBuildDir" DefaultValue="$(NetworkPublishDir)" Description="The path to the build."/>
+	<Option Name="BuildTag" DefaultValue="" Description="The tag to the build."/>
+	<Option Name="BuildVersion" DefaultValue="" Description="The version to the build."/>
+	<Property Name="BuildTag" Value="$(ProjectName)-CL-$(BuildVersion)" If="'$(BuildTag)' == ''"/>
+
+	<Property Name="TargetRemotePlatforms" Value=""/>
+	<Property Name="InstallTasks" Value=""/>
+	<Macro Name="AddPlatform" Arguments="Platform;IsConsole;IsEnabled" OptionalArguments="IsVirtualized">
+		<Property Name="IsVirtualized" Value="('$(IsVirtualized)' != '' and '$(IsVirtualized)' != 'false')"/>
+		<Property Name="InstallTasks" Value="$(InstallTasks);InstallBuild $(Platform)" If="$(IsEnabled) And $(IsConsole) And !$(IsVirtualized)"/>
+		<Property Name="TargetRemotePlatforms" Value="$(TargetRemotePlatforms);$(Platform)" If="$(IsConsole) And !$(IsVirtualized)"/>
+	</Macro>
+	<!-- Include any scripts from platform extensions -->
+	<Include Script="../../../Platforms/*/Build/Inc/PlatformOptions.xml"/>
+
+	<Option Name="DeviceList" DefaultValue="" Description="List of device to use."/>
+	<Option Name="ExtraInstallArgs" DefaultValue="" Description="Additional install arguments."/>
+	
+	<Agent Name="Install Build Job" Type="DevkitAutomation">
+		<ForEach Name="TargetPlatform" Values="$(TargetRemotePlatforms)">
+			<Node Name="InstallBuild $(TargetPlatform)">
+			<ForEach Name="Configuration" Values="$(Configurations)">
+				<Command Name="InstallUnrealBuild" Arguments="-platform=$(TargetPlatform) $(InstallBuildFlags) -configuration=$(Configuration) -project=$(ProjectName) -device=$(DeviceList) -path=&quot;$(PublishBuildDir)/$(BuildTag)&quot; $(ExtraInstallArgs)"/>
+			</ForEach>
+			</Node>
+		</ForEach>
+	</Agent>
+
+	<Aggregate Name="Install Build" Requires="$(InstallTasks)"/>
+
+</BuildGraph>
\ No newline at end of file
diff --git a/Engine/Source/Programs/AutomationTool/Gauntlet/Unreal/InstallUnrealBuild.cs b/Engine/Source/Programs/AutomationTool/Gauntlet/Unreal/InstallUnrealBuild.cs
--- a/Engine/Source/Programs/AutomationTool/Gauntlet/Unreal/InstallUnrealBuild.cs
+++ b/Engine/Source/Programs/AutomationTool/Gauntlet/Unreal/InstallUnrealBuild.cs
@@ -19,6 +19,7 @@ namespace Gauntlet
 	[ParamHelp("parallel=[1:4]", "Number of installs to run in parallel",
 		ParamType = typeof(string), Choices = new string[] { "1", "2", "3", "4" }, DefaultValue = "1")]
 	[ParamHelp("project=<project>", "Name of the project", Required = true)]
+	[ParamHelp("-sandbox=<sandbox>", "String to specficy the name of a staged install directory.")]
 
 	public class InstallUnrealBuild : BuildCommand
 	{
@@ -39,7 +40,8 @@ namespace Gauntlet
 			string DevicesArg = ParseParamValue("device", string.Empty);
 			DevicesArg = ParseParamValue("devices", DevicesArg);
 
-			int ParallelTasks = Convert.ToInt32(ParseParamValue("parallel", "1"));
+			//Set default to -1, to know that a value has not been passed in and can be overriden.
+			int ParallelTasks = Convert.ToInt32(ParseParamValue("parallel", "-1"));
 
 			bool Success = false;
 
diff --git a/Engine/Source/Programs/AutomationTool/Gauntlet/Unreal/Utils/Gauntlet.InstallUnreal.cs b/Engine/Source/Programs/AutomationTool/Gauntlet/Unreal/Utils/Gauntlet.InstallUnreal.cs
--- a/Engine/Source/Programs/AutomationTool/Gauntlet/Unreal/Utils/Gauntlet.InstallUnreal.cs
+++ b/Engine/Source/Programs/AutomationTool/Gauntlet/Unreal/Utils/Gauntlet.InstallUnreal.cs
@@ -50,7 +50,7 @@ namespace Gauntlet
 				throw new AutomationException("No BuildSources found for platform {0}", Platform);
 			}
 
-			IEnumerable<IBuild> Builds = BuildSources.Where(S => S.CanSupportPlatform(Platform)).SelectMany(S => S.GetBuildsAtPath(ProjectName, BuildPath));
+			IReadOnlyCollection<IBuild> Builds = BuildSources.Where(S => S.CanSupportPlatform(Platform)).SelectMany(S => S.GetBuildsAtPath(ProjectName, BuildPath)).ToList();
 
 			if (Builds.Count() == 0)
 			{
@@ -79,7 +79,7 @@ namespace Gauntlet
 					throw new AutomationException("No IDeviceFactory implmenetation that supports {0}", Platform);
 				}
 
-				DeviceList = Devices.Select(D => Factory.CreateDevice(D, string.Empty)).ToArray();
+				DeviceList = Devices.Select(D => Factory.CreateDevice(D, null)).ToArray();
 			}
 
 			if (DeviceList.Count() == 0)
@@ -87,9 +87,15 @@ namespace Gauntlet
 				throw new AutomationException("No devices found for {0}", Platform);
 			}
 
+			if (ParallelTasks == -1)
+			{
+				//If a value is not passed in as a paramter, set ParallelTasks equal to the number of devices, MAX 4
+				ParallelTasks = (DeviceList.Count() > 4) ? 4 : DeviceList.Count();
+			}
+
 			var POptions = new ParallelOptions { MaxDegreeOfParallelism = ParallelTasks };
 
-			// now copy it four builds at a time
+			// now copy it up to four builds at a time
 			Parallel.ForEach(DeviceList, POptions, Device =>
 			{
 				DateTime StartTime = DateTime.Now;
@@ -97,13 +103,79 @@ namespace Gauntlet
 				UnrealAppConfig Config = new UnrealAppConfig();
 
 				Config.CommandLine = CommandLine;
-				Config.Build = Builds.First();
 				Config.ProjectName = ProjectName;
+				Config.Name = ProjectName;
+
+				string BuildSandbox = Globals.Params.ParseValue("sandbox", string.Empty);
+				Config.Sandbox = string.IsNullOrEmpty(BuildSandbox) ? Config.Sandbox : BuildSandbox;
+
+				// We always (currently..) need to be able to replace the command line
+				BuildFlags Flags = BuildFlags.CanReplaceCommandLine;
+
+				if (Globals.Params.ParseParam("dev"))
+				{
+					Flags |= BuildFlags.CanReplaceExecutable;
+				}
+				if (Globals.Params.ParseParam("bulk"))
+				{
+					Flags |= BuildFlags.Bulk;
+				}
+				if (Globals.Params.ParseParam("notbulk"))
+				{
+					Flags |= BuildFlags.NotBulk;
+				}
+				if (Globals.Params.ParseParam("packaged"))
+				{
+					Flags |= BuildFlags.Packaged;
+				}
+				if (Globals.Params.ParseParam("staged"))
+				{
+					Flags |= BuildFlags.Loose;
+				}
+
+				IEnumerable<IBuild> FlaggedBuilds = Builds.Where(S => S.Flags.HasFlag(Flags));
+
+				if (FlaggedBuilds.Count() == 0)
+				{
+					throw new AutomationException("Unable to find build at {0} with build flag(s): {1} ", BuildPath, Flags.ToString());
+				}
+
+				string BuildConfiguration = Globals.Params.ParseValue("configuration", String.Empty);
 
-				Log.Info("Installing build on device {0}", Device.Name);
+				switch (BuildConfiguration.ToLower())
+				{
+					case "shipping":
+						Config.Build = FlaggedBuilds.Where(S => S.Configuration == UnrealTargetConfiguration.Shipping).FirstOrDefault();
+						break;
+
+					case "test":
+						Config.Build = FlaggedBuilds.Where(S => S.Configuration == UnrealTargetConfiguration.Test).FirstOrDefault();
+						break;
+
+					case "development":
+						Config.Build = FlaggedBuilds.Where(S => S.Configuration == UnrealTargetConfiguration.Development).FirstOrDefault();
+						break;
+
+					default:
+						Log.Info("No build configuration was provided, selecting first available build.");
+						Config.Build = FlaggedBuilds.FirstOrDefault();
+						break;
+				}
 
-				IAppInstall Install = Device.InstallApplication(Config);
-				Device.Run(Install);
+				if (!Device.Connect())
+				{
+					throw new AutomationException("Failed to connect to device: {0}", Device);
+				}
+
+				if (Config.Build != null)
+				{
+					Log.Info("Installing build on device {DeviceName}", Device.Name);
+					IAppInstall Install = Device.InstallApplication(Config);
+				}
+				else
+				{
+					throw new AutomationException("Unable to find build at {0} with build flag(s)={1} and configuration type={2} for {3} platform type", BuildPath, Flags, BuildConfiguration, Platform);
+				}
 
 				TimeSpan Elapsed = (DateTime.Now - StartTime);
 				Log.Info("Installed on device {0} in {1:D2}m:{2:D2}s", Device.Name, Elapsed.Minutes, Elapsed.Seconds);
