[Gauntlet] Detect TVOS build

#rnx

diff --git a/Engine/Source/Programs/AutomationTool/Gauntlet/Platform/IOS/Gauntlet.IOSBuildSource.cs b/Engine/Source/Programs/AutomationTool/Gauntlet/Platform/IOS/Gauntlet.IOSBuildSource.cs
--- a/Engine/Source/Programs/AutomationTool/Gauntlet/Platform/IOS/Gauntlet.IOSBuildSource.cs
+++ b/Engine/Source/Programs/AutomationTool/Gauntlet/Platform/IOS/Gauntlet.IOSBuildSource.cs
@@ -5,14 +5,14 @@ using System.Collections.Generic;
 using System.IO;
 using AutomationTool;
 using UnrealBuildTool;
-using System.Threading;
+using System.Xml.Linq;
 using System.Text.RegularExpressions;
 using System.Linq;
 using EpicGames.Core;
 
 namespace Gauntlet
 {
-	public class IOSBuild : IBuild
+	public class AppleBuild : IBuild
 	{
 		public int PreferenceOrder { get { return 0; } }
 
@@ -30,11 +30,11 @@ namespace Gauntlet
 
 		public string Flavor { get { return ""; } }
 
-		public UnrealTargetPlatform Platform { get { return UnrealTargetPlatform.IOS; } }
+		public virtual UnrealTargetPlatform Platform { get; }
 
 		public bool SupportsAdditionalFileCopy { get; }
 
-		public IOSBuild(UnrealTargetConfiguration InConfig, string InPackageName, string InSourcePath, Dictionary<string, string> InFilesToInstall, BuildFlags InFlags)
+		public AppleBuild(UnrealTargetConfiguration InConfig, string InPackageName, string InSourcePath, Dictionary<string, string> InFilesToInstall, BuildFlags InFlags)
 		{
 			Configuration = InConfig;
 			PackageName = InPackageName;
@@ -130,10 +130,9 @@ namespace Gauntlet
 			return true;
 		}
 
-
-		private static string GetBundleIdentifier(string Source)
+		private static PlistInfo GetPlistInfo(string Source)
 		{
-			string PlistInfo = string.Empty;
+			PlistInfo Info = null;
 			bool IsIPAFile = Path.GetExtension(Source).Equals(".ipa", StringComparison.OrdinalIgnoreCase);
 
 			if (IsIPAFile)
@@ -164,7 +163,7 @@ namespace Gauntlet
 					return null;
 				}
 
-				PlistInfo = Output;
+				Info = new PlistInfo(Output);
 			}
 			else
 			{
@@ -178,42 +177,95 @@ namespace Gauntlet
 				}
 
 				StreamReader PListStream = new StreamReader(PlistFile.FullName);
-				PlistInfo = PListStream.ReadToEnd();
+				Info = new PlistInfo(PListStream.ReadToEnd());
 			}
 
-			// todo: plist parsing, could be better
-			string PackageName = null;
-			string KeyString = "<key>CFBundleIdentifier</key>";
-			int KeyIndex = PlistInfo.IndexOf(KeyString);
-			if (KeyIndex > 0)
+			return Info;
+		}
+
+		private class PlistInfo
+		{
+			private XDocument Document;
+
+			public PlistInfo(string InContent)
 			{
-				int StartIdx = PlistInfo.IndexOf("<string>", KeyIndex + KeyString.Length) + "<string>".Length;
-				int EndIdx = PlistInfo.IndexOf("</string>", StartIdx);
-				if (StartIdx > 0 && EndIdx > StartIdx)
+				try
 				{
-					PackageName = PlistInfo.Substring(StartIdx, EndIdx - StartIdx);
+					Document = XDocument.Parse(InContent);
+				}
+				catch (Exception Ex)
+				{
+					// Ignore errors
+					Log.Warning(KnownLogEvents.Gauntlet_BuildDropEvent, "Fail to parse PlistInfo':\r{Exception}", Ex);
+					Document = new XDocument();
 				}
 			}
 
-			if (String.IsNullOrEmpty(PackageName))
+			/// <summary>
+			/// Get first value from corresponding key
+			/// </summary>
+			/// <param name="Key"></param>
+			/// <returns></returns>
+			public string GetFirstValue(string Key)
 			{
-				Log.Warning(String.Format("Unable to find CFBundleIdentifier in plist info for App {0}", Source));
+				foreach (XElement element in Document.Descendants("key"))
+				{
+					if (element.Value == Key)
+					{
+						XElement NextElement = element.ElementsAfterSelf().FirstOrDefault();
+						if (NextElement != null)
+						{
+							if (NextElement.Name == "string")
+							{
+								return NextElement.Value;
+							}
+							else if (NextElement.Name == "array")
+							{
+								return NextElement.Descendants("string").Select(e => e.Value).FirstOrDefault();
+							}
+						}
+					}
+				}
+
 				return null;
 			}
 
-			Log.Verbose("Found bundle id: {0}", PackageName);
-
-			return PackageName;
+			/// <summary>
+			/// Get all values from corresponding key
+			/// </summary>
+			/// <param name="Key"></param>
+			/// <returns></returns>
+			public IEnumerable<string> GetAllValues(string Key)
+			{
+				foreach (XElement element in Document.Descendants("key"))
+				{
+					if (element.Value == Key)
+					{
+						XElement NextElement = element.ElementsAfterSelf().FirstOrDefault();
+						if (NextElement != null)
+						{
+							if (NextElement.Name == "array")
+							{
+								return NextElement.Descendants("string").Select(e => e.Value);
+							}
+							else if (NextElement.Name == "string")
+							{
+								return new List<string> { NextElement.Value };
+							}
+						}
+					}
+				}
 
+				return null;
+			}
 		}
 
-		public static IOSBuild CreateFromPath(string InProjectName, string InPath)
+		public static T CreateFromPath<T>(string InProjectName, string InPath, AppleBuildSource<T> BuildSource)
+			where T : AppleBuild
 		{
-			string BuildPath = InPath;
-
-			IOSBuild DiscoveredBuilds = null;
+			T DiscoveredBuild = null;
 
-			DirectoryInfo Di = new DirectoryInfo(BuildPath);
+			DirectoryInfo Di = new DirectoryInfo(InPath);
 
 			var UnrealConfig = UnrealHelpers.GetConfigurationFromExecutableName(InProjectName, Di.Name);
 
@@ -240,39 +292,57 @@ namespace Gauntlet
 						Flags |= BuildFlags.NotBulk;
 					}
 
-					string PackageName = GetBundleIdentifier(Di.FullName);
-
-					if (!String.IsNullOrEmpty(PackageName))
+					PlistInfo Info = GetPlistInfo(Di.FullName);
+					if (Info != null)
 					{
-						Dictionary<string, string> FilesToInstall = new Dictionary<string, string>();
+						IEnumerable<string> CFBundlePlatformNames = Info.GetAllValues("CFBundleSupportedPlatforms");
+						if (CFBundlePlatformNames != null && CFBundlePlatformNames.Contains(BuildSource.CFBundlePlatformName))
+						{
+							string PackageName = Info.GetFirstValue("CFBundleIdentifier");
 
-						DiscoveredBuilds = new IOSBuild(UnrealConfig, PackageName, Di.FullName, FilesToInstall, Flags);
+							if (!String.IsNullOrEmpty(PackageName))
+							{
+								Dictionary<string, string> FilesToInstall = new Dictionary<string, string>();
 
-						Log.Verbose("Found {0} {1} build at {2}", UnrealConfig, ((Flags & BuildFlags.Bulk) == BuildFlags.Bulk) ? "(bulk)" : "(not bulk)", AbsPath);
+								DiscoveredBuild = Activator.CreateInstance(typeof(T), new object[] { UnrealConfig, PackageName, Di.FullName, FilesToInstall, Flags }) as T;
+
+								Log.Verbose("Found bundle id: {0}", PackageName);
+								Log.Verbose("Found {0} {1} build at {2}", UnrealConfig, ((Flags & BuildFlags.Bulk) == BuildFlags.Bulk) ? "(bulk)" : "(not bulk)", AbsPath);
+							}
+							else
+							{
+								Log.Warning(String.Format("Unable to find CFBundleIdentifier in plist info for App {0}", Di.FullName));
+							}
+						}
+						else
+						{
+							Log.Verbose("Unable to find matching platform '{0}' for CFBundleSupportedPlatforms in plist info for App {1}", BuildSource.CFBundlePlatformName, Di.FullName);
+						}
 					}
 				}
 			}
 
-			return DiscoveredBuilds;
+			return DiscoveredBuild;
 		}
 	}
 
-	public class IOSBuildSource : IFolderBuildSource
+	public abstract class AppleBuildSource<T> : IFolderBuildSource
+		where T : AppleBuild
 	{
-		public string BuildName { get { return "IOSBuildSource"; } }
+		protected abstract UnrealTargetPlatform Platform { get; }
+
+		public abstract string CFBundlePlatformName { get; }
+
+		public string BuildName { get { return $"{Platform}BuildSource"; } }
 
 		public bool CanSupportPlatform(UnrealTargetPlatform InPlatform)
 		{
-			return InPlatform == UnrealTargetPlatform.IOS;
+			return InPlatform == Platform;
 		}
 
 		public string ProjectName { get; protected set; }
 
-		public IOSBuildSource()
-		{
-		}
-
-		public List<IBuild> GetBuildsAtPath(string InProjectName, string InPath, int MaxRecursion = 3)
+		public virtual List<IBuild> GetBuildsAtPath(string InProjectName, string InPath, int MaxRecursion = 3)
 		{
 			// We only want iOS builds on Mac host
 			if (BuildHostPlatform.Current.Platform != UnrealTargetPlatform.Mac)
@@ -289,14 +359,14 @@ namespace Gauntlet
 			{
 				List<DirectoryInfo> SearchDirs = new List<DirectoryInfo>();
 
-				if (PathDI.Name.IndexOf("IOS", StringComparison.OrdinalIgnoreCase) >= 0)
+				if (PathDI.Name.IndexOf(Platform.ToString(), StringComparison.OrdinalIgnoreCase) >= 0)
 				{
 					SearchDirs.Add(PathDI);
 				}
 				else
 				{
 					// find all directories that begin with IOS
-					SearchDirs.AddRange(PathDI.GetDirectories("IOS*", SearchOption.TopDirectoryOnly));
+					SearchDirs.AddRange(PathDI.GetDirectories($"{Platform}*", SearchOption.TopDirectoryOnly));
 				}
 
 				IEnumerable<DirectoryInfo> DirsToRecurse = new List<DirectoryInfo>(SearchDirs);
@@ -317,10 +387,10 @@ namespace Gauntlet
 					DirsToRecurse = DiscoveredDirs.Except(Packages);
 				}
 
-				string IOSBuildFilter = Globals.Params.ParseValue("IOSBuildFilter", "");
+				string IOSBuildFilter = Globals.Params.ParseValue($"{Platform}BuildFilter", "");
 				foreach (DirectoryInfo Di in AllDirs)
 				{
-					IOSBuild FoundBuild = IOSBuild.CreateFromPath(InProjectName, Di.FullName);
+					AppleBuild FoundBuild = AppleBuild.CreateFromPath<T>(InProjectName, Di.FullName, this);
 
 					if (FoundBuild != null)
 					{
@@ -337,4 +407,34 @@ namespace Gauntlet
 		}
 
 	}
+
+	public class IOSBuild : AppleBuild
+	{
+		public IOSBuild(UnrealTargetConfiguration InConfig, string InPackageName, string InSourcePath, Dictionary<string, string> InFilesToInstall, BuildFlags InFlags)
+			: base(InConfig, InPackageName, InSourcePath, InFilesToInstall, InFlags)
+		{ }
+
+		public override UnrealTargetPlatform Platform => UnrealTargetPlatform.IOS;
+	}
+
+	public class IOSBuildSource : AppleBuildSource<IOSBuild>
+	{
+		protected override UnrealTargetPlatform Platform => UnrealTargetPlatform.IOS;
+		public override string CFBundlePlatformName => "iPhoneOS";
+	}
+
+	public class TVOSBuild : AppleBuild
+	{
+		public TVOSBuild(UnrealTargetConfiguration InConfig, string InPackageName, string InSourcePath, Dictionary<string, string> InFilesToInstall, BuildFlags InFlags)
+			: base(InConfig, InPackageName, InSourcePath, InFilesToInstall, InFlags)
+		{ }
+
+		public override UnrealTargetPlatform Platform => UnrealTargetPlatform.TVOS;
+	}
+
+	public class TVOSBuildSource : AppleBuildSource<TVOSBuild>
+	{
+		protected override UnrealTargetPlatform Platform => UnrealTargetPlatform.TVOS;
+		public override string CFBundlePlatformName => "AppleTVOS";
+	}
 }
\ No newline at end of file
