Ensure when exporting snapshots that the bucket name we export to is always valid and sanitized.

diff --git a/Engine/Source/Programs/AutomationTool/BuildGraph/Tasks/ZenExportSnapshotTask.cs b/Engine/Source/Programs/AutomationTool/BuildGraph/Tasks/ZenExportSnapshotTask.cs
--- a/Engine/Source/Programs/AutomationTool/BuildGraph/Tasks/ZenExportSnapshotTask.cs
+++ b/Engine/Source/Programs/AutomationTool/BuildGraph/Tasks/ZenExportSnapshotTask.cs
@@ -13,6 +13,7 @@ using System.Text.Json.Serialization;
 using System.Threading.Tasks;
 using System.Xml;
 using EpicGames.Core;
+using EpicGames.Horde;
 using EpicGames.Serialization;
 using EpicGames.ProjectStore;
 
@@ -335,6 +336,7 @@ namespace AutomationTool.Tasks
 					{
 						BucketName = ProjectNameAsBucketName;
 					}
+					BucketName = SanitizeBucketName(BucketName);
 
 					string HostName = Parameters.SnapshotDescriptorCloudHost;
 					if (string.IsNullOrEmpty(HostName))
@@ -420,6 +422,11 @@ namespace AutomationTool.Tasks
 			return false;
 		}
 
+		private string SanitizeBucketName(string InString)
+		{
+			return StringId.Sanitize(InString).ToString();
+		}
+
 		/// <summary>
 		/// Execute the task.
 		/// </summary>
@@ -530,6 +537,7 @@ namespace AutomationTool.Tasks
 					{
 						BucketName = ProjectNameAsBucketName;
 					}
+					BucketName = SanitizeBucketName(BucketName);
 
 					OplogExportCommandline.AppendFormat(" --cloud {0} --namespace {1} --bucket {2}", Parameters.DestinationCloudHost, Parameters.DestinationCloudNamespace, BucketName);
 
