[Gauntlet] Remove legacy method of crash dump copying during SaveRoleArtfacts
Each platform already handles crash dump copying during the ArtifactPath accessor

diff --git a/Engine/Source/Programs/AutomationTool/Gauntlet/Framework/Base/Gauntlet.TargetDevice.cs b/Engine/Source/Programs/AutomationTool/Gauntlet/Framework/Base/Gauntlet.TargetDevice.cs
--- a/Engine/Source/Programs/AutomationTool/Gauntlet/Framework/Base/Gauntlet.TargetDevice.cs
+++ b/Engine/Source/Programs/AutomationTool/Gauntlet/Framework/Base/Gauntlet.TargetDevice.cs
@@ -123,23 +123,6 @@ namespace Gauntlet
 		void CopyAdditionalFiles(IEnumerable<UnrealFileToCopy> FilesToCopy);
 
 		/// End new flow ///
-
-		/// <summary>
-		/// Path to the crash dumps on a device
-		/// </summary>
-		string CrashDumpPath
-		{
-			get
-			{
-				return Path.Combine(Globals.TempDir, "CrashDumps", Platform.ToString() + "_" + Name);
-			}
-		}
-
-		/// <summary>
-		/// Ensures the crash dump copy has occurred already - and does the copy if it hasn't happened yet
-		/// Returns true if there were any crash dumps for the run, and false otherwise
-		/// </summary>
-		bool CopyCrashDumps() { return false; }
 	};
 
 	/// <summary>
diff --git a/Engine/Source/Programs/AutomationTool/Gauntlet/Unreal/Base/Gauntlet.UnrealSession.cs b/Engine/Source/Programs/AutomationTool/Gauntlet/Unreal/Base/Gauntlet.UnrealSession.cs
--- a/Engine/Source/Programs/AutomationTool/Gauntlet/Unreal/Base/Gauntlet.UnrealSession.cs
+++ b/Engine/Source/Programs/AutomationTool/Gauntlet/Unreal/Base/Gauntlet.UnrealSession.cs
@@ -1150,40 +1150,6 @@ namespace Gauntlet
 				}
 			}
 
-			bool bRetainCrashdumps = InContext.TestParams.ParseParam("RetainCrashDumps");
-			if (InRunningRole.AppInstance.Device.CopyCrashDumps())
-			{
-				DirectoryInfo CrashDumpDirectory = new DirectoryInfo(InRunningRole.AppInstance.Device.CrashDumpPath);
-				if (CrashDumpDirectory.Exists)
-				{
-					string DesinationCrashDumpDirectory = Path.Combine(DestinationDirectory.FullName, "CrashDumps");
-
-					try
-					{
-						Log.Info("Copying crash dumps from {0} to {1}", CrashDumpDirectory.FullName, DesinationCrashDumpDirectory);
-						SystemHelpers.CopyDirectory(CrashDumpDirectory.FullName, DesinationCrashDumpDirectory);
-					}
-					catch (Exception Exception)
-					{
-						bRetainCrashdumps = true;
-						Log.Warning("Encountered an {Exception} when copying crash dumps from {SourceDirectory} to {DestinationDirectory}. " +
-							"Crash dumps will not be saved locally, but the source crash dumps will not be deleted.", Exception, CrashDumpDirectory, DesinationCrashDumpDirectory);
-					}
-
-					if (!bRetainCrashdumps)
-					{
-						try
-						{
-							CrashDumpDirectory.Delete(true);
-						}
-						catch (Exception Exception)
-						{
-							Log.Info("Encountered an {Exception} when deleting source crash dumps at {SourceDirectory}. Crash dumps will remain on the device.", Exception, SourceDirectory);
-						}
-					}
-				}
-			}
-
 			// Convert any screenshots to jpegs and create a gif when not running a server
 			if (!InRunningRole.Role.RoleType.IsServer())
 			{
