[Gauntlet] Fix CI errors from crash dump removal merge

diff --git a/Engine/Source/Programs/AutomationTool/Gauntlet/Framework/Base/Gauntlet.TargetDeviceDesktopCommon.cs b/Engine/Source/Programs/AutomationTool/Gauntlet/Framework/Base/Gauntlet.TargetDeviceDesktopCommon.cs
--- a/Engine/Source/Programs/AutomationTool/Gauntlet/Framework/Base/Gauntlet.TargetDeviceDesktopCommon.cs
+++ b/Engine/Source/Programs/AutomationTool/Gauntlet/Framework/Base/Gauntlet.TargetDeviceDesktopCommon.cs
@@ -105,13 +105,8 @@ namespace Gauntlet
 			{
 				CleanArtifactDirectory(GetInstallArtifactPath());
 			}
-			// Clean crash dumps
-			ITargetDevice ThisDevice = (ITargetDevice)this;
-			CleanArtifactDirectory(ThisDevice.CrashDumpPath);
 		}
 
-		public virtual bool CopyCrashDumps() { return false; }
-
 		public virtual void InstallBuild(UnrealAppConfig AppConfig)
 		{
 			IBuild Build = AppConfig.Build;
diff --git a/Engine/Source/Programs/AutomationTool/Gauntlet/Platform/Windows/Gauntlet.TargetDeviceWindows.cs b/Engine/Source/Programs/AutomationTool/Gauntlet/Platform/Windows/Gauntlet.TargetDeviceWindows.cs
--- a/Engine/Source/Programs/AutomationTool/Gauntlet/Platform/Windows/Gauntlet.TargetDeviceWindows.cs
+++ b/Engine/Source/Programs/AutomationTool/Gauntlet/Platform/Windows/Gauntlet.TargetDeviceWindows.cs
@@ -173,26 +173,6 @@ namespace Gauntlet
 			return (InstallCache as DesktopCommonAppInstall<TargetDeviceWindows>).ArtifactPath;
 		}
 
-		public override bool CopyCrashDumps()
-		{
-			// On Windows platform, check only if any dmp file exists.
-			// It does not need to be fetched from a remote location.
-			ITargetDevice ThisDevice = (ITargetDevice)this;
-			DirectoryInfo DirInfo = new DirectoryInfo(ThisDevice.CrashDumpPath);
-			if (DirInfo.Exists)
-			{
-				// See if there is a crash dump for this test run
-				foreach (FileInfo CrashDumpFileInfo in DirInfo.GetFiles("*.dmp", SearchOption.AllDirectories))
-				{
-					if (CrashDumpFileInfo.Exists)
-					{
-						return true;
-					}
-				}
-			}
-			return false;
-		}
-
 		#region Legacy Implementations
 		public override IAppInstall InstallApplication(UnrealAppConfig AppConfig)
 		{
@@ -379,12 +359,13 @@ namespace Gauntlet
 
 		protected override void GenerateDump()
 		{
+			string CrashDumpPath = string.Empty;
 			bool WroteDump = false;
-			if (!Directory.Exists(Device.CrashDumpPath))
+			if (!Directory.Exists(CrashDumpPath))
 			{
-				Directory.CreateDirectory(Device.CrashDumpPath);
+				Directory.CreateDirectory(CrashDumpPath);
 			}
-			string DumpName = Path.Combine(Device.CrashDumpPath, Path.GetFileNameWithoutExtension(ProcessResult.ProcessObject.ProcessName) + ".dmp");
+			string DumpName = Path.Combine(CrashDumpPath, Path.GetFileNameWithoutExtension(ProcessResult.ProcessObject.ProcessName) + ".dmp");
 			using (FileStream CrashDumpStream = File.Create(DumpName))
 			{
 				WroteDump = MiniDumpWriteDump(ProcessResult.ProcessObject.SafeHandle, (uint)ProcessResult.ProcessObject.Id,
