[Gauntlet] Direct windows crashdump generation on process kill to saved crashes directory

diff --git a/Engine/Source/Programs/AutomationTool/Gauntlet/Platform/Windows/Gauntlet.TargetDeviceWindows.cs b/Engine/Source/Programs/AutomationTool/Gauntlet/Platform/Windows/Gauntlet.TargetDeviceWindows.cs
--- a/Engine/Source/Programs/AutomationTool/Gauntlet/Platform/Windows/Gauntlet.TargetDeviceWindows.cs
+++ b/Engine/Source/Programs/AutomationTool/Gauntlet/Platform/Windows/Gauntlet.TargetDeviceWindows.cs
@@ -359,12 +359,13 @@ namespace Gauntlet
 
 		protected override void GenerateDump()
 		{
-			string CrashDumpPath = string.Empty;
+			string CrashDumpPath = Path.Combine(Install.DesktopDevice.UserDir, "Saved", "Crashes");
 			bool WroteDump = false;
 			if (!Directory.Exists(CrashDumpPath))
 			{
 				Directory.CreateDirectory(CrashDumpPath);
 			}
+
 			string DumpName = Path.Combine(CrashDumpPath, Path.GetFileNameWithoutExtension(ProcessResult.ProcessObject.ProcessName) + ".dmp");
 			using (FileStream CrashDumpStream = File.Create(DumpName))
 			{
