[Gauntlet] Adds -dev local binarry overrides (Overlay Executables) when running with -ExperimentalLaunchFlow
    - Adds the OverlayExecutable type
    - Adds OverlayExecutable member to UnrealAppConfig. Only initialized if -dev is present
    - This object can check if there is a local newer binary for the relevant Role and return the path to it

diff --git a/Engine/Source/Programs/AutomationTool/Gauntlet/Framework/Base/Gauntlet.OverlayExecutable.cs b/Engine/Source/Programs/AutomationTool/Gauntlet/Framework/Base/Gauntlet.OverlayExecutable.cs
new file mode 100644
--- /dev/null
+++ b/Engine/Source/Programs/AutomationTool/Gauntlet/Framework/Base/Gauntlet.OverlayExecutable.cs
@@ -0,0 +1,81 @@
+// Copyright Epic Games, Inc. All Rights Reserved.
+
+using AutomationTool;
+using System;
+using System.IO;
+using UnrealBuildTool;
+
+namespace Gauntlet
+{
+	/// <summary>
+	/// Manages the overlaying of a development executable over an existing build.
+	/// These are useful for testing code changes without needing to re-cook or re-package builds.
+	/// To use overlay executables, simply specify -dev on the UAT commandline.
+	/// The overlay executable will automatically be applied as long as the following conditions are met
+	///		- The build being used supports Overlay executables. See BuildFlags.CanReplaceExecutable
+	///		- The development executable's write time is newer than that of the supplied builds write time
+	/// </summary>
+	public class OverlayExecutable
+	{
+		private string ProjectName;
+		private string Flavor;
+		private UnrealTargetPlatform Platform;
+		private UnrealTargetConfiguration Configuration;
+		private UnrealTargetRole Role;
+
+		public OverlayExecutable(UnrealSessionRole SessionRole, string ProjectName)
+		{
+			this.ProjectName = ProjectName;
+			this.Configuration = SessionRole.Configuration;
+			this.Role = SessionRole.RoleType;
+			this.Platform = SessionRole.Platform.Value;
+			this.Flavor = SessionRole.RequiredFlavor;
+		}
+
+		/// <summary>
+		/// Attempts to get the path to a local executable if it is newer than the base executable
+		/// </summary>
+		/// <param name="BaseExecutable">The base executable to use if an overlay does not exist</param>
+		/// <param name="ExtensionOverride">What extension the overlay file should have</param>
+		/// <returns>True if a local, newer executable that matches the role's requirements exists</returns>
+		/// <exception cref="AutomationException"></exception>
+		public bool GetOverlay(string BaseExecutable, out string OverlayExecutable, string ExtensionOverride = null)
+		{
+			OverlayExecutable = null;
+
+			// If no base executable is specified we'll skip the "newer" check and just return a local binary that matches the role
+			if (string.IsNullOrEmpty(BaseExecutable))
+			{
+				throw new AutomationException("Overlay Error: No base executable was specified.");
+			}
+
+			if (!File.Exists(BaseExecutable))
+			{
+				throw new AutomationException("Overlay Error: Could not find base executable {0}.", BaseExecutable);
+			}
+
+			// Ex. /UnrealGame/Binaries/Win64/UnrealClient.exe
+			string PlatformBinariesDirectory = Path.Combine(Environment.CurrentDirectory, ProjectName, "Binaries", Platform.ToString());
+			string ExecutableName = UnrealHelpers.GetExecutableName(ProjectName, Platform, Configuration, Role, Flavor, string.Empty);
+			string ExecutableExtension = string.IsNullOrEmpty(ExtensionOverride) ? Path.GetExtension(BaseExecutable) : ExtensionOverride;
+			string LocalExecutable = Path.Combine(PlatformBinariesDirectory, ExecutableName + ExecutableExtension);
+
+			if (!File.Exists(LocalExecutable))
+			{
+				Log.Verbose("No local executable for {Platform} exists. Skipping overlay for this role", Platform);
+				return false;
+			}
+
+			if (File.GetLastWriteTime(BaseExecutable) > File.GetLastWriteTime(LocalExecutable))
+			{
+				Log.Verbose("Local executable for {Platform} is not newer than base executable. Skipping overlay for this role", Platform);
+				return false;
+			}
+
+			Log.Info("Applying newer local executable as overlay {LocalExecutable}", LocalExecutable);
+
+			OverlayExecutable = LocalExecutable;
+			return true;
+		}
+	}
+}
\ No newline at end of file
diff --git a/Engine/Source/Programs/AutomationTool/Gauntlet/Framework/Gauntlet.Utils.cs b/Engine/Source/Programs/AutomationTool/Gauntlet/Framework/Gauntlet.Utils.cs
--- a/Engine/Source/Programs/AutomationTool/Gauntlet/Framework/Gauntlet.Utils.cs
+++ b/Engine/Source/Programs/AutomationTool/Gauntlet/Framework/Gauntlet.Utils.cs
@@ -41,6 +41,7 @@ namespace Gauntlet
 		static List<Action> InnerPostAbortHandlers = new List<Action>();
 		public static bool CancelSignalled { get; private set; }
 
+		public static bool IsRunningDev => Params.ParseParam("dev");
 
 		/// <summary>
 		/// Returns the device pool id 
diff --git a/Engine/Source/Programs/AutomationTool/Gauntlet/Unreal/Base/Gauntlet.UnrealAppConfig.cs b/Engine/Source/Programs/AutomationTool/Gauntlet/Unreal/Base/Gauntlet.UnrealAppConfig.cs
--- a/Engine/Source/Programs/AutomationTool/Gauntlet/Unreal/Base/Gauntlet.UnrealAppConfig.cs
+++ b/Engine/Source/Programs/AutomationTool/Gauntlet/Unreal/Base/Gauntlet.UnrealAppConfig.cs
@@ -101,6 +101,8 @@ namespace Gauntlet
 
 		public IBuild Build { get; set; }
 
+		public OverlayExecutable OverlayExecutable{ get; set; }
+
 		// Prevents installing a build on device
 		public bool SkipInstall => ForceSkipInstall.HasValue ? ForceSkipInstall.Value : _SkipInstall;
 
@@ -112,10 +114,10 @@ namespace Gauntlet
 
 		// Force a skip install
 		public static bool? ForceSkipInstall = null;
-		
+
 		[AutoParamWithNames(false, "SkipInstall", "SkipDeploy", "SkipCopy")]
 		private bool _SkipInstall { get; set; }
-		
+
 		[AutoParamWithNames(false, "FullClean")]
 		private bool _FullClean { get; set; }
 
diff --git a/Engine/Source/Programs/AutomationTool/Gauntlet/Unreal/Base/Gauntlet.UnrealSession.cs b/Engine/Source/Programs/AutomationTool/Gauntlet/Unreal/Base/Gauntlet.UnrealSession.cs
--- a/Engine/Source/Programs/AutomationTool/Gauntlet/Unreal/Base/Gauntlet.UnrealSession.cs
+++ b/Engine/Source/Programs/AutomationTool/Gauntlet/Unreal/Base/Gauntlet.UnrealSession.cs
@@ -179,7 +179,7 @@ namespace Gauntlet
 
 			RequiredBuildFlags = BuildFlags.None;
 
-			if (Globals.Params.ParseParam("dev") && !RoleType.UsesEditor())
+			if (Globals.IsRunningDev && !RoleType.UsesEditor())
 			{
 				RequiredBuildFlags |= BuildFlags.CanReplaceExecutable;
 			}
@@ -827,7 +827,7 @@ namespace Gauntlet
 				}
 				catch (Exception Ex)
 				{
-					if(IsOutOfSpaceException(Ex))
+					if(IsOutOfSpaceException(Ex) || IsOverlayException(Ex))
 					{
 						RemainingAttempts = 0;
 						ReleaseSessionDevices();
@@ -990,10 +990,9 @@ namespace Gauntlet
 
 			// We only want to move artifacts for editor data if there was a crash on a buildmachine.
 			// Also, don't move artifacts in dev mode, because peoples saved data could be huuuuuuuge!
-			bool IsDevBuild = InContext.TestParams.ParseParam("dev");
 			bool IsEditorBuild = InRunningRole.Role.RoleType.UsesEditor();
 			bool IsBuildMachine = CommandUtils.IsBuildMachine;
-			bool SkipArchivingAssets = IsDevBuild || (IsEditorBuild && (IsBuildMachine == false || InRunningRole.AppInstance.ExitCode == 0));
+			bool SkipArchivingAssets = Globals.IsRunningDev || (IsEditorBuild && (IsBuildMachine == false || InRunningRole.AppInstance.ExitCode == 0));
 			bool bRetainArtifacts = InContext.TestParams.ParseParam("RetainDeviceArtifacts");
 
 			// Check if we should copy artifacts
@@ -1067,7 +1066,7 @@ namespace Gauntlet
 				{
 					Log.Info("Skipping archival of assets for editor {Role}", RoleName);
 				}
-				else if (IsDevBuild)
+				else if (Globals.IsRunningDev)
 				{
 					Log.Info("Skipping archival of assets for dev build");
 				}
@@ -1689,6 +1688,12 @@ namespace Gauntlet
 						Log.Error(KnownLogEvents.Gauntlet_DeviceEvent, Message);
 						throw;
 					}
+					else if(IsOverlayException(Ex))
+					{
+						// Errors with Overlay executables are caused by missing files or improper setup, inform user and exit early
+						Log.Error(KnownLogEvents.Gauntlet_DeviceEvent, Message);
+						throw;
+					}
 					else
 					{
 						UnrealDeviceReservation.MarkProblemDevice(Device, Message);
@@ -1939,6 +1944,11 @@ namespace Gauntlet
 		{
 			return Ex.Message.Contains("not enough space", StringComparison.OrdinalIgnoreCase);
 		}
+
+		private bool IsOverlayException(Exception Ex)
+		{
+			return Ex.Message.Contains("Overlay Error", StringComparison.OrdinalIgnoreCase);
+		}
 	}
 
 }
diff --git a/Engine/Source/Programs/AutomationTool/Gauntlet/Unreal/BuildSource/Gauntlet.UnrealBuildSource.cs b/Engine/Source/Programs/AutomationTool/Gauntlet/Unreal/BuildSource/Gauntlet.UnrealBuildSource.cs
--- a/Engine/Source/Programs/AutomationTool/Gauntlet/Unreal/BuildSource/Gauntlet.UnrealBuildSource.cs
+++ b/Engine/Source/Programs/AutomationTool/Gauntlet/Unreal/BuildSource/Gauntlet.UnrealBuildSource.cs
@@ -573,6 +573,12 @@ namespace Gauntlet
             {
                 Config.FilesToCopy = Role.FilesToCopy;
             }
+
+			if(Globals.IsRunningDev)
+			{
+				Config.OverlayExecutable = new OverlayExecutable(Role, Config.ProjectName);
+			}
+
 			return Config;
 		}
 
@@ -729,9 +735,8 @@ namespace Gauntlet
 					}
 				}
 
-				bool IsRunningDev = Globals.Params.ParseParam("dev");
-
 				// Turn FooGame into Foo
+				bool IsRunningDev = Globals.IsRunningDev;
 				string ExeBase = ProjectName.Replace("Game", "");
 
 				if (TargetPlatform == UnrealTargetPlatform.Android)
@@ -765,8 +770,6 @@ namespace Gauntlet
 					{
 						throw new AutomationException("No suitable build for {0} found at {1}", TargetPlatform, string.Join(",", BuildPaths));
 					}
-
-					//ExePath = AndroidSource.SourceApkPath;			
 				}
 				else
 				{
