[Gauntlet] Adds -dev support for Desktop when running with ExperimentalLaunchFlow

diff --git a/Engine/Source/Programs/AutomationTool/Gauntlet/Framework/Base/Gauntlet.TargetDeviceDesktopCommon.cs b/Engine/Source/Programs/AutomationTool/Gauntlet/Framework/Base/Gauntlet.TargetDeviceDesktopCommon.cs
--- a/Engine/Source/Programs/AutomationTool/Gauntlet/Framework/Base/Gauntlet.TargetDeviceDesktopCommon.cs
+++ b/Engine/Source/Programs/AutomationTool/Gauntlet/Framework/Base/Gauntlet.TargetDeviceDesktopCommon.cs
@@ -414,6 +414,12 @@ namespace Gauntlet
 				// clear artifact path
 				CleanDeviceArtifacts();
 			}
+
+			if(Globals.IsRunningDev && AppConfig.OverlayExecutable.GetOverlay(ExecutablePath, out string OverlayExecutable))
+			{
+				CommandArguments += string.Format(" -basedir=\"{0}\"", Path.GetDirectoryName(ExecutablePath));
+				ExecutablePath = OverlayExecutable;
+			}
 		}
 	}
 
