[Gauntlet] Adds -dev support for Android when running with ExperimentalLaunchFlow

diff --git a/Engine/Source/Programs/AutomationTool/Gauntlet/Platform/Android/Gauntlet.TargetDeviceAndroid.cs b/Engine/Source/Programs/AutomationTool/Gauntlet/Platform/Android/Gauntlet.TargetDeviceAndroid.cs
--- a/Engine/Source/Programs/AutomationTool/Gauntlet/Platform/Android/Gauntlet.TargetDeviceAndroid.cs
+++ b/Engine/Source/Programs/AutomationTool/Gauntlet/Platform/Android/Gauntlet.TargetDeviceAndroid.cs
@@ -408,11 +408,16 @@ namespace Gauntlet
 
 			EnablePermissions(Package);
 
+
 			// Install apk
-			CopyFileToDevice(AppConfig, Package, Build.SourceApkPath, string.Empty);
+			string APK = Globals.IsRunningDev && AppConfig.OverlayExecutable.GetOverlay(Build.SourceApkPath, out string OverlayAPK)
+				? OverlayAPK
+				: Build.SourceApkPath;
+			CopyFileToDevice(AppConfig, Package, APK, string.Empty);
 
 			// Copy obbs from bulk builds
-			if(Build.FilesToInstall != null && Build.FilesToInstall.Any())
+			bool bSkipOBBInstall = Globals.Params.ParseParam("SkipOBBCopy"); // useful when iterating on dev executables
+			if(Build.FilesToInstall != null && Build.FilesToInstall.Any() && !bSkipOBBInstall)
 			{
 				CopyOBBFiles(AppConfig, Build.FilesToInstall, Package);
 			}
