[Gauntlet] Finalize/tidy up android AFS addition

diff --git a/Engine/Source/Programs/AutomationTool/Gauntlet/Platform/Android/Gauntlet.TargetDeviceAndroid.cs b/Engine/Source/Programs/AutomationTool/Gauntlet/Platform/Android/Gauntlet.TargetDeviceAndroid.cs
--- a/Engine/Source/Programs/AutomationTool/Gauntlet/Platform/Android/Gauntlet.TargetDeviceAndroid.cs
+++ b/Engine/Source/Programs/AutomationTool/Gauntlet/Platform/Android/Gauntlet.TargetDeviceAndroid.cs
@@ -107,53 +107,6 @@ namespace Gauntlet
 
 		protected Dictionary<EIntendedBaseCopyDirectory, string> LocalDirectoryMappings;
 
-		public static bool UsingAndroidFileServer(FileReference RawProjectPath, UnrealTargetConfiguration TargetConfiguration, out bool bEnablePlugin, out string AFSToken, out bool bIsShipping, out bool bIncludeInShipping, out bool bAllowExternalStartInShipping)
-		{
-			UnrealTargetPlatform TargetPlatform = UnrealTargetPlatform.Android;
-			//UnrealTargetConfiguration TargetConfiguration = SC != null ? SC.StageTargetConfigurations[0] : Params.ClientConfigsToBuild[0];
-			bIsShipping = TargetConfiguration == UnrealTargetConfiguration.Shipping;
-
-			ConfigHierarchy Ini = ConfigCache.ReadHierarchy(ConfigHierarchyType.Engine, DirectoryReference.FromFile(RawProjectPath), TargetPlatform);
-			if (!Ini.GetBool("/Script/AndroidFileServerEditor.AndroidFileServerRuntimeSettings", "bEnablePlugin", out bEnablePlugin))
-			{
-				bEnablePlugin = true;
-			}
-			if (!Ini.GetString("/Script/AndroidFileServerEditor.AndroidFileServerRuntimeSettings", "SecurityToken", out AFSToken))
-			{
-				AFSToken = "";
-			}
-			if (!Ini.GetBool("/Script/AndroidFileServerEditor.AndroidFileServerRuntimeSettings", "bIncludeInShipping", out bIncludeInShipping))
-			{
-				bIncludeInShipping = false;
-			}
-			if (!Ini.GetBool("/Script/AndroidFileServerEditor.AndroidFileServerRuntimeSettings", "bAllowExternalStartInShipping", out bAllowExternalStartInShipping))
-			{
-				bAllowExternalStartInShipping = false;
-			}
-
-			if (bIsShipping && !(bIncludeInShipping && bAllowExternalStartInShipping))
-			{
-				return false;
-			}
-			return bEnablePlugin;
-		}
-
-		public bool UseAFS(UnrealAppConfig AppConfig)
-		{
-			AndroidBuild Build = AppConfig.Build as AndroidBuild;
-			if (Build == null)
-			{
-				throw new AutomationException("Unsupported build type {0} for Android!", AppConfig.Build.GetType());
-			}
-
-			bool bAFSEnablePlugin;
-			string AFSToken;
-			bool bIsShipping;
-			bool bAFSIncludeInShipping;
-			bool bAFSAllowExternalStartInShipping;
-			return UsingAndroidFileServer(AppConfig.ProjectFile, Build.Configuration, out bAFSEnablePlugin, out AFSToken, out bIsShipping, out bAFSIncludeInShipping, out bAFSAllowExternalStartInShipping);
-		}
-
 		public TargetDeviceAndroid(string InDeviceName = "", AndroidDeviceData DeviceData = null, string InCachePath = null)
 		{
 			AdbCredentialCache.AddInstance(DeviceData);
@@ -413,7 +366,7 @@ namespace Gauntlet
 			string APK = Globals.IsRunningDev && AppConfig.OverlayExecutable.GetOverlay(Build.SourceApkPath, out string OverlayAPK)
 				? OverlayAPK
 				: Build.SourceApkPath;
-			CopyFileToDevice(AppConfig, Package, APK, string.Empty);
+			CopyFileToDevice(Package, APK, string.Empty);
 
 			// Copy obbs from bulk builds
 			bool bSkipOBBInstall = Globals.Params.ParseParam("SkipOBBCopy"); // useful when iterating on dev executables
@@ -446,7 +399,7 @@ namespace Gauntlet
 			return new AndroidAppInstall(this, AppConfig.ProjectName, Build.AndroidPackageName, AppConfig.CommandLine, AppConfig.ProjectFile, AppConfig.Build.Configuration);
 		}
 
-		public void CopyAdditionalFiles(UnrealAppConfig AppConfig, IEnumerable<UnrealFileToCopy> FilesToCopy)
+		public void CopyAdditionalFiles(IEnumerable<UnrealFileToCopy> FilesToCopy)
 		{
 			foreach (UnrealFileToCopy FileToCopy in FilesToCopy)
 			{
@@ -455,7 +408,7 @@ namespace Gauntlet
 				{
 					FileInfo SourceFile = new FileInfo(FileToCopy.SourceFileLocation);
 					SourceFile.IsReadOnly = false;
-					CopyFileToDevice(AppConfig, null, SourceFile.FullName, DestinationFile);
+					CopyFileToDevice(null, SourceFile.FullName, DestinationFile); // todo: add AFS support
 				}
 				else
 				{
@@ -463,10 +416,6 @@ namespace Gauntlet
 				}
 			}
 		}
-		public void CopyAdditionalFiles(IEnumerable<UnrealFileToCopy> FilesToCopy)
-		{
-			CopyAdditionalFiles(null, FilesToCopy);
-		}
 
 		public IAppInstance Run(IAppInstall App)
 		{
@@ -475,17 +424,6 @@ namespace Gauntlet
 				throw new AutomationException("AppInstall is of incorrect type {0}! Must be of type AndroidAppInstall", App.GetType().Name);
 			}
 
-			bool bAFSEnablePlugin;
-			string AFSToken = "";
-			bool bIsShipping;
-			bool bAFSIncludeInShipping;
-			bool bAFSAllowExternalStartInShipping;
-			bool useADB = true;
-			if (UsingAndroidFileServer(Install.ProjectFile, Install.Configuration, out bAFSEnablePlugin, out AFSToken, out bIsShipping, out bAFSIncludeInShipping, out bAFSAllowExternalStartInShipping))
-			{
-				useADB = false;
-			}
-
 			if(!IsOn)
 			{
 				PowerOn();
@@ -503,14 +441,14 @@ namespace Gauntlet
 			RunAdbDeviceCommand("logcat --clear");
 
 			// Ensure artifact directories exist
-			if (useADB)
+			if (UsingAndroidFileServer(Install.ProjectFile, Install.Configuration, out _, out string AFSToken, out _, out _, out _))
 			{
-				RunAdbDeviceCommand(string.Format("shell mkdir -p {0}/", DeviceExternalStorageSavedPath));
-				RunAdbDeviceCommand(string.Format("shell mkdir -p {0}/", DeviceExternalFilesSavedPath));
+				RunAFSDeviceCommand(string.Format("-p \"{0}\" -k \"{1}\" mkdir \"^saved\"", Install.AndroidPackageName, AFSToken));
 			}
 			else
 			{
-				RunAFSDeviceCommand(string.Format("-p \"{0}\" -k \"{1}\" mkdir \"^saved\"", Install.AndroidPackageName, AFSToken));
+				RunAdbDeviceCommand(string.Format("shell mkdir -p {0}/", DeviceExternalStorageSavedPath));
+				RunAdbDeviceCommand(string.Format("shell mkdir -p {0}/", DeviceExternalFilesSavedPath));
 			}
 
 			// start the app on device!
@@ -541,6 +479,57 @@ namespace Gauntlet
 			LocalDirectoryMappings.Add(EIntendedBaseCopyDirectory.Saved, ProjectDir + "Saved");
 		}
 
+		public static bool UsingAndroidFileServer(FileReference RawProjectPath, UnrealTargetConfiguration TargetConfiguration, out bool bEnablePlugin, out string AFSToken, out bool bIsShipping, out bool bIncludeInShipping, out bool bAllowExternalStartInShipping)
+		{
+			if(RawProjectPath == null || !FileReference.Exists(RawProjectPath) || TargetConfiguration == UnrealTargetConfiguration.Unknown)
+			{
+				bEnablePlugin = false;
+				AFSToken = string.Empty;
+				bIsShipping = false;
+				bIncludeInShipping = false;
+				bAllowExternalStartInShipping = false;
+				return false;
+			}
+
+			UnrealTargetPlatform TargetPlatform = UnrealTargetPlatform.Android;
+			bIsShipping = TargetConfiguration == UnrealTargetConfiguration.Shipping;
+
+			ConfigHierarchy Ini = ConfigCache.ReadHierarchy(ConfigHierarchyType.Engine, DirectoryReference.FromFile(RawProjectPath), TargetPlatform);
+			if (!Ini.GetBool("/Script/AndroidFileServerEditor.AndroidFileServerRuntimeSettings", "bEnablePlugin", out bEnablePlugin))
+			{
+				bEnablePlugin = true;
+			}
+			if (!Ini.GetString("/Script/AndroidFileServerEditor.AndroidFileServerRuntimeSettings", "SecurityToken", out AFSToken))
+			{
+				AFSToken = "";
+			}
+			if (!Ini.GetBool("/Script/AndroidFileServerEditor.AndroidFileServerRuntimeSettings", "bIncludeInShipping", out bIncludeInShipping))
+			{
+				bIncludeInShipping = false;
+			}
+			if (!Ini.GetBool("/Script/AndroidFileServerEditor.AndroidFileServerRuntimeSettings", "bAllowExternalStartInShipping", out bAllowExternalStartInShipping))
+			{
+				bAllowExternalStartInShipping = false;
+			}
+
+			if (bIsShipping && !(bIncludeInShipping && bAllowExternalStartInShipping))
+			{
+				return false;
+			}
+			return bEnablePlugin;
+		}
+
+		public bool UseAFS(UnrealAppConfig AppConfig)
+		{
+			AndroidBuild Build = AppConfig.Build as AndroidBuild;
+			if (Build == null)
+			{
+				throw new AutomationException("Unsupported build type {0} for Android!", AppConfig.Build.GetType());
+			}
+
+			return UsingAndroidFileServer(AppConfig.ProjectFile, Build.Configuration, out _, out _, out _, out _, out _);
+		}
+
 		/// <summary>
 		/// Runs an ADB command, automatically adding the name of the current device to
 		/// the arguments sent to adb
@@ -631,7 +620,7 @@ namespace Gauntlet
 
 			return Process;
 		}
-		
+
 		public static IProcessResult RunAFSGlobalCommand(string Args, bool Wait = true, bool bShouldLogCommand = false, bool bPauseErrorParsing = false)
 		{
 			CommandUtils.ERunOptions RunOptions = CommandUtils.ERunOptions.AppMustExist | CommandUtils.ERunOptions.NoWaitForExit | CommandUtils.ERunOptions.SpewIsVerbose;
@@ -680,11 +669,11 @@ namespace Gauntlet
 				Log.Verbose("Removing {0}", CommandLineFilePath);
 				if (UseAFS(AppConfig))
 				{
-					DeleteFileFromDevice(AppConfig, Package, "^commandfile");
+					DeleteFileFromDevice(Package, "^commandfile", AppConfig.ProjectFile, AppConfig.Configuration);
 				}
 				else
 				{
-					DeleteFileFromDevice(AppConfig, Package, CommandLineFilePath);
+					DeleteFileFromDevice(Package, CommandLineFilePath);
 				}
 			}
 		}
@@ -750,31 +739,14 @@ namespace Gauntlet
 			}
 		}
 
-		public bool CopyFileToDevice(UnrealAppConfig AppConfig, string PackageName, string SourcePath, string DestPath, bool IgnoreDependencies = false)
+		public bool CopyFileToDevice(string PackageName, string SourcePath, string DestPath,
+			FileReference ProjectFile = null, UnrealTargetConfiguration Configuration = UnrealTargetConfiguration.Unknown, bool IgnoreDependencies = false)
 		{
-			bool bAFSEnablePlugin;
-			string AFSToken = "";
-			bool bIsShipping;
-			bool bAFSIncludeInShipping;
-			bool bAFSAllowExternalStartInShipping;
-			bool useADB = true;
-			if (AppConfig != null)
-			{
-				AndroidBuild Build = AppConfig.Build as AndroidBuild;
-				if (Build == null)
-				{
-					throw new AutomationException("Unsupported build type {0} for Android!", AppConfig.Build.GetType());
-				}
-
-				if (UsingAndroidFileServer(AppConfig.ProjectFile, Build.Configuration, out bAFSEnablePlugin, out AFSToken, out bIsShipping, out bAFSIncludeInShipping, out bAFSAllowExternalStartInShipping))
-				{
-					useADB = false;
-				}
-			}
-
+			string AFSToken;
+			bool UseAFS = UsingAndroidFileServer(ProjectFile, Configuration, out _, out AFSToken, out _, out _, out _);
 			bool IsAPK = string.Equals(Path.GetExtension(SourcePath), ".apk", StringComparison.OrdinalIgnoreCase);
 
-			// for the APK there's no easy/reliable way to get the date of the version installed, so 
+			// for the APK there's no easy/reliable way to get the date of the version installed, so
 			// we write this out to a dependency file in the demote dir and check it each time.
 			// current file time
 			DateTime LocalModifiedTime = File.GetLastWriteTime(SourcePath);
@@ -809,7 +781,11 @@ namespace Gauntlet
 			else
 			{
 				// for other files get the file info
-				if (useADB)
+				if (UseAFS)
+				{
+					CurrentFileInfo = AFSGetFileInfo(PackageName, AFSToken, DestPath);
+				}
+				else
 				{
 					AdbResult = RunAdbDeviceCommand(string.Format("shell ls -l {0}", DestPath));
 					if (AdbResult.ExitCode == 0)
@@ -817,10 +793,6 @@ namespace Gauntlet
 						CurrentFileInfo = AdbResult.Output.ToString().Trim();
 					}
 				}
-				else
-				{
-					CurrentFileInfo = AFSGetFileInfo(PackageName, AFSToken, DestPath);
-				}
 
 			}
 
@@ -881,18 +853,18 @@ namespace Gauntlet
 				{
 					string FileDirectory = Path.GetDirectoryName(DestPath).Replace('\\', '/');
 
-					if (useADB)
+					if (UseAFS)
+					{
+						Log.Info("Copying {0} to {1} via AFS", QuotedSourcePath, DestPath);
+						AdbResult = RunAFSDeviceCommand(string.Format("-p \"{0}\" -k \"{1}\" push \"{2}\" \"{3}\"", PackageName, AFSToken, SourcePath, DestPath));
+					}
+					else
 					{
 						RunAdbDeviceCommand(string.Format("shell mkdir -p {0}/", FileDirectory));
 
 						Log.Info("Copying {0} to {1} via adb push", QuotedSourcePath, DestPath);
 						AdbResult = RunAdbDeviceCommand(string.Format("push {0} {1}", QuotedSourcePath, DestPath));
 					}
-					else
-					{
-						Log.Info("Copying {0} to {1} via AFS", QuotedSourcePath, DestPath);
-						AdbResult = RunAFSDeviceCommand(string.Format("-p \"{0}\" -k \"{1}\" push \"{2}\" \"{3}\"", PackageName, AFSToken, SourcePath, DestPath));
-					}
 
 					// Note: Presently, AdbResult NEVER reports failures to push when using AFS. Should be fixed at some point.
 
@@ -905,15 +877,15 @@ namespace Gauntlet
 							// Try to reconnect
 							if (Connect())
 							{
-								if (useADB)
+								if (UseAFS)
 								{
-									Log.Info("Retrying to copy via adb push...");
-									AdbResult = RunAdbDeviceCommand(string.Format("push {0} {1}", QuotedSourcePath, DestPath));
+									Log.Info("Retrying to copy via AFS...");
+									AdbResult = RunAFSDeviceCommand(string.Format("-p {0} -k {1} push \"{2}\" \"{3}\"", PackageName, AFSToken, SourcePath, DestPath));
 								}
 								else
 								{
-									Log.Info("Retrying to copy via AFS...");
-									AdbResult = RunAFSDeviceCommand(string.Format("-p {0} -k {1} push \"{2}\" \"{3}\"", PackageName, AFSToken, SourcePath, DestPath));
+									Log.Info("Retrying to copy via adb push...");
+									AdbResult = RunAdbDeviceCommand(string.Format("push {0} {1}", QuotedSourcePath, DestPath));
 								}
 								if (AdbResult.ExitCode != 0)
 								{
@@ -932,14 +904,14 @@ namespace Gauntlet
 					}
 
 					// Now pull info about the file which we'll write as a dep
-					if (useADB)
+					if (UseAFS)
 					{
-						AdbResult = RunAdbDeviceCommand(string.Format("shell ls -l {0}", DestPath));
-						CurrentFileInfo = AdbResult.Output.ToString().Trim();
+						CurrentFileInfo = AFSGetFileInfo(PackageName, AFSToken, DestPath);
 					}
 					else
 					{
-						CurrentFileInfo = AFSGetFileInfo(PackageName, AFSToken, DestPath);
+						AdbResult = RunAdbDeviceCommand(string.Format("shell ls -l {0}", DestPath));
+						CurrentFileInfo = AdbResult.Output.ToString().Trim();
 					}
 				}
 
@@ -1052,7 +1024,7 @@ namespace Gauntlet
 					}
 
 					DestinationFile = Regex.Replace(DestinationFile, "%STORAGE%", StoragePath, RegexOptions.IgnoreCase);
-					CopyFileToDevice(AppConfig, null, SourceFile, DestinationFile);
+					CopyFileToDevice(null, SourceFile, DestinationFile, AppConfig.ProjectFile, AppConfig.Configuration);
 				}
 				else
 				{
@@ -1062,12 +1034,12 @@ namespace Gauntlet
 					if (parts[parts.Length - 1] == "obb")
 					{
 						string obbname = string.Format("^{0}obb", parts[0]);
-						CopyFileToDevice(AppConfig, Package, SourceFile, obbname);
+						CopyFileToDevice(Package, SourceFile, obbname, AppConfig.ProjectFile, AppConfig.Configuration);
 					}
 					else
 					{
 						DestinationFile = Regex.Replace(DestinationFile, "%STORAGE%", StoragePath, RegexOptions.IgnoreCase);
-						CopyFileToDevice(AppConfig, Package, SourceFile, DestinationFile);
+						CopyFileToDevice(Package, SourceFile, DestinationFile, AppConfig.ProjectFile, AppConfig.Configuration);
 					}
 				}
 			}
@@ -1090,46 +1062,30 @@ namespace Gauntlet
 
 			if (UseAFS(AppConfig))
 			{
-				CopyFileToDevice(AppConfig, Package, CommandLineTmpFile, "^commandfile");
+				CopyFileToDevice(Package, CommandLineTmpFile, "^commandfile", AppConfig.ProjectFile, AppConfig.Configuration);
 			}
 			else
 			{
-				CopyFileToDevice(AppConfig, Package, CommandLineTmpFile, CommandLineFilePath);
+				CopyFileToDevice(Package, CommandLineTmpFile, CommandLineFilePath);
 			}
 
 			File.Delete(CommandLineTmpFile);
 		}
 
-		protected bool DeleteFileFromDevice(UnrealAppConfig AppConfig, string Package, string DestPath)
+		protected bool DeleteFileFromDevice(string Package, string DestPath,
+			FileReference ProjectFile = null, UnrealTargetConfiguration Configuration = UnrealTargetConfiguration.Unknown)
 		{
-			bool bAFSEnablePlugin;
-			string AFSToken = "";
-			bool bIsShipping;
-			bool bAFSIncludeInShipping;
-			bool bAFSAllowExternalStartInShipping;
-			bool useADB = true;
-			if (AppConfig != null)
-			{
-				AndroidBuild Build = AppConfig.Build as AndroidBuild;
-				if (Build == null)
-				{
-					throw new AutomationException("Unsupported build type {0} for Android!", AppConfig.Build.GetType());
-				}
-
-				if (UsingAndroidFileServer(AppConfig.ProjectFile, Build.Configuration, out bAFSEnablePlugin, out AFSToken, out bIsShipping, out bAFSIncludeInShipping, out bAFSAllowExternalStartInShipping))
-				{
-					useADB = false;
-				}
-			}
+			string AFSToken;
+			bool UseAFS = UsingAndroidFileServer(ProjectFile, Configuration, out _, out AFSToken, out _, out _, out _);
 
 			IProcessResult AdbResult;
-			if (useADB)
+			if (UseAFS)
 			{
-				AdbResult = RunAdbDeviceCommand(string.Format("shell rm -f {0}", DestPath));
+				AdbResult = RunAFSDeviceCommand(string.Format("-p \"{0}\" -k \"{1}\" rm \"{2}\"", Package, AFSToken, DestPath));
 			}
 			else
 			{
-				AdbResult = RunAFSDeviceCommand(string.Format("-p \"{0}\" -k \"{1}\" rm \"{2}\"", Package, AFSToken, DestPath));
+				AdbResult = RunAdbDeviceCommand(string.Format("shell rm -f {0}", DestPath));
 			}
 			return AdbResult.ExitCode == 0;
 		}
@@ -1266,7 +1222,7 @@ namespace Gauntlet
 				}
 
 				// first install the APK
-				CopyFileToDevice(AppConfig, Build.AndroidPackageName, ApkPath, string.Empty);
+				CopyFileToDevice(Build.AndroidPackageName, ApkPath, string.Empty, AppConfig.ProjectFile, AppConfig.Configuration);
 
 				// remote dir on the device, create it if it doesn't exist
 				if (useADB)
@@ -1382,7 +1338,7 @@ namespace Gauntlet
 					Console.WriteLine("Copying {0} to {1}", LocalFile, RemoteFile);
 					if (useADB)
 					{
-						CopyFileToDevice(AppConfig, Build.AndroidPackageName, LocalFile, RemoteFile);
+						CopyFileToDevice(Build.AndroidPackageName, LocalFile, RemoteFile, AppConfig.ProjectFile, AppConfig.Configuration);
 					}
 					else
 					{
@@ -1392,11 +1348,11 @@ namespace Gauntlet
 						if (parts[parts.Length - 1] == "obb")
 						{
 							string obbname = string.Format("^{0}obb", parts[0]);
-							CopyFileToDevice(AppConfig, Build.AndroidPackageName, LocalFile, obbname);
+							CopyFileToDevice(Build.AndroidPackageName, LocalFile, obbname, AppConfig.ProjectFile, AppConfig.Configuration);
 						}
 						else
 						{
-							CopyFileToDevice(AppConfig, Build.AndroidPackageName, LocalFile, RemoteFile);
+							CopyFileToDevice(Build.AndroidPackageName, LocalFile, RemoteFile, AppConfig.ProjectFile, AppConfig.Configuration);
 						}
 					}
 				}
