[Gauntlet] Fix exception thrown by UnrealTestNode during CreateReport if no log file is written

diff --git a/Engine/Source/Programs/AutomationTool/Gauntlet/Unreal/Base/Gauntlet.UnrealTestNode.cs b/Engine/Source/Programs/AutomationTool/Gauntlet/Unreal/Base/Gauntlet.UnrealTestNode.cs
--- a/Engine/Source/Programs/AutomationTool/Gauntlet/Unreal/Base/Gauntlet.UnrealTestNode.cs
+++ b/Engine/Source/Programs/AutomationTool/Gauntlet/Unreal/Base/Gauntlet.UnrealTestNode.cs
@@ -1355,8 +1355,13 @@ namespace Gauntlet
 			{
 				foreach (UnrealRoleResult RoleResult in RoleResults)
 				{
-					string LogName = Path.GetFullPath(RoleResult.Artifacts.LogPath).Replace(Path.GetFullPath(Context.Options.LogDir), "").TrimStart(Path.DirectorySeparatorChar);
-					HordeTestReport.AttachArtifact(RoleResult.Artifacts.LogPath, LogName);
+					string LogName = string.Empty;
+					string LogPath = RoleResult.Artifacts.LogPath;
+					if(!string.IsNullOrEmpty(LogPath) && File.Exists(LogPath))
+					{
+						LogName = Path.GetFullPath(LogPath).Replace(Path.GetFullPath(Context.Options.LogDir), "").TrimStart(Path.DirectorySeparatorChar);
+						HordeTestReport.AttachArtifact(RoleResult.Artifacts.LogPath, LogName);
+					}
 
 					UnrealLog LogSummary = RoleResult.LogSummary;
 					if (LogSummary.Errors.Count() > 0)
