[Gauntlet] Fixing log output string
#rnx

diff --git a/Engine/Source/Programs/AutomationTool/Gauntlet/Platform/IOS/Gauntlet.IOSBuildSource.cs b/Engine/Source/Programs/AutomationTool/Gauntlet/Platform/IOS/Gauntlet.IOSBuildSource.cs
--- a/Engine/Source/Programs/AutomationTool/Gauntlet/Platform/IOS/Gauntlet.IOSBuildSource.cs
+++ b/Engine/Source/Programs/AutomationTool/Gauntlet/Platform/IOS/Gauntlet.IOSBuildSource.cs
@@ -196,7 +196,7 @@ namespace Gauntlet
 				catch (Exception Ex)
 				{
 					// Ignore errors
-					Log.Warning(KnownLogEvents.Gauntlet_BuildDropEvent, "Fail to parse PlistInfo':\r{Exception}", Ex);
+					Log.Warning(KnownLogEvents.Gauntlet_BuildDropEvent, "Fail to parse PlistInfo:\n{Exception}", Ex);
 					Document = new XDocument();
 				}
 			}
