[Gauntlet] Restore Heartbeat category logging that was mistakenly removed in 33204261

diff --git a/Engine/Source/Programs/AutomationTool/Gauntlet/Unreal/Base/Gauntlet.UnrealTestNode.cs b/Engine/Source/Programs/AutomationTool/Gauntlet/Unreal/Base/Gauntlet.UnrealTestNode.cs
--- a/Engine/Source/Programs/AutomationTool/Gauntlet/Unreal/Base/Gauntlet.UnrealTestNode.cs
+++ b/Engine/Source/Programs/AutomationTool/Gauntlet/Unreal/Base/Gauntlet.UnrealTestNode.cs
@@ -1803,6 +1803,24 @@ namespace Gauntlet
 
 					UnrealLogStreamParser Parser = new UnrealLogStreamParser();
 					LastLogCount += Parser.ReadStream(StdOut, LastLogCount);
+
+					foreach (string TestLine in Parser.GetLogFromShortNameChannels(LogCategories))
+					{
+						Log.Info(string.Format("{0}: {1}", AppPrefix, TestLine));
+
+						if (bUpdateHeartbeatTime)
+						{
+							if (Regex.IsMatch(TestLine, @".*GauntletHeartbeat\: Active.*"))
+							{
+								LastHeartbeatTime = DateTime.Now;
+								LastActiveHeartbeatTime = DateTime.Now;
+							}
+							else if (Regex.IsMatch(TestLine, @".*GauntletHeartbeat\: Idle.*"))
+							{
+								LastHeartbeatTime = DateTime.Now;
+							}
+						}
+					}
 				}
 			}
 
