Disable play protect for Android installation

diff --git a/Engine/Build/LowLevelTests.xml b/Engine/Build/LowLevelTests.xml
--- a/Engine/Build/LowLevelTests.xml
+++ b/Engine/Build/LowLevelTests.xml
@@ -177,7 +177,7 @@
 					<Property Name="TestTags" Value=""/>
 					<Property Name="TestTags" Value="$(Tags)" If="'$(Tags)' != ''"/>
 					<Property Name="TestTags" Value="$(OverrideTags)" If="'$(OverrideTags)' != ''"/>
-					<Property Name="TestArgsExt" Value="$(TestArgsExt) -cleandevice" If="'$(Platform)' == 'Android'"/>
+					<Property Name="TestArgsExt" Value="$(TestArgsExt) -cleandevice -no-play-protect" If="'$(Platform)' == 'Android'"/>
 					<Property Name="TestArgsExt" Value="$(TestArgsExt) -tags=&quot;$(TestTags)&quot;"/>
 					<Property Name="TestArgsExt" Value="$(TestArgsExt) --buildmachine" If="$(IsBuildMachine)"/>
 					<Property Name="TestArgsExt" Value="$(TestArgsExt) -timeout=$(PerTestTimeout)"/>
diff --git a/Engine/Source/Programs/AutomationTool/Gauntlet/Platform/Android/Gauntlet.TargetDeviceAndroid.cs b/Engine/Source/Programs/AutomationTool/Gauntlet/Platform/Android/Gauntlet.TargetDeviceAndroid.cs
--- a/Engine/Source/Programs/AutomationTool/Gauntlet/Platform/Android/Gauntlet.TargetDeviceAndroid.cs
+++ b/Engine/Source/Programs/AutomationTool/Gauntlet/Platform/Android/Gauntlet.TargetDeviceAndroid.cs
@@ -107,6 +107,8 @@ namespace Gauntlet
 
 		protected Dictionary<EIntendedBaseCopyDirectory, string> LocalDirectoryMappings;
 
+		private string NoPlayProtectSetting = null;
+
 		public TargetDeviceAndroid(string InDeviceName = "", AndroidDeviceData DeviceData = null, string InCachePath = null)
 		{
 			AdbCredentialCache.AddInstance(DeviceData);
@@ -220,6 +222,12 @@ namespace Gauntlet
 			{
 				try
 				{
+					if (!string.IsNullOrEmpty(NoPlayProtectSetting))
+					{
+						Log.Verbose("Restoring play protect for this session...");
+						RunAdbDeviceCommand($"shell settings put global package_verifier_user_consent {NoPlayProtectSetting}");
+					}
+
 					if (!IsExistingDevice)
 					{
 						// disconnect
@@ -1221,6 +1229,17 @@ namespace Gauntlet
 					}
 				}
 
+				bool NoPlayProtect = Globals.Params.ParseParam("no-play-protect");
+				if (NoPlayProtect)
+				{
+					NoPlayProtectSetting = RunAdbDeviceCommandAndGetOutput("shell settings get global package_verifier_user_consent");
+					if (NoPlayProtectSetting != "-1")
+					{
+						Log.Verbose("Removing play protect for this session...");
+						RunAdbDeviceCommand("shell settings put global package_verifier_user_consent -1");
+					}
+				}
+
 				// first install the APK
 				CopyFileToDevice(Build.AndroidPackageName, ApkPath, string.Empty, AppConfig.ProjectFile, AppConfig.Configuration);
 
diff --git a/Engine/Source/Programs/AutomationTool/LowLevelTests/RunLowLevelTests.cs b/Engine/Source/Programs/AutomationTool/LowLevelTests/RunLowLevelTests.cs
--- a/Engine/Source/Programs/AutomationTool/LowLevelTests/RunLowLevelTests.cs
+++ b/Engine/Source/Programs/AutomationTool/LowLevelTests/RunLowLevelTests.cs
@@ -167,6 +167,9 @@ namespace LowLevelTests
 
 		public bool VerifyLogin;
 
+		// Android specific
+		public bool NoPlayProtect;
+
 		[AutoParam(false)]
 		public bool SkipStage;
 
@@ -237,6 +240,7 @@ namespace LowLevelTests
 			Tags = Params.ParseValue("tags=", null);
 			AttachToDebugger = Params.ParseParam("attachtodebugger");
 			VerifyLogin = Params.ParseParam("verifylogin");
+			NoPlayProtect = Params.ParseParam("no-play-protect");
 
 			TestExtraArgs = Params.ParseValue("extra-args=", null);
 
diff --git a/Engine/Source/Runtime/CoreUObject/Tests/UObject/InstanceDataObjectUtilsTest.cpp b/Engine/Source/Runtime/CoreUObject/Tests/UObject/InstanceDataObjectUtilsTest.cpp
--- a/Engine/Source/Runtime/CoreUObject/Tests/UObject/InstanceDataObjectUtilsTest.cpp
+++ b/Engine/Source/Runtime/CoreUObject/Tests/UObject/InstanceDataObjectUtilsTest.cpp
@@ -2,7 +2,7 @@
 
 #include "InstanceDataObjectUtilsTest.h"
 
-#if WITH_TESTS
+#if WITH_TESTS && (PLATFORM_WINDOWS || PLATFORM_LINUX || PLATFORM_MAC)
 
 #include "Tests/TestHarnessAdapter.h"
 #include "UObject/InstanceDataObjectUtils.h"
