[Gauntlet] Add DateTime to ITestEvent

#rnx

diff --git a/Engine/Source/Programs/AutomationTool/Gauntlet/Framework/Base/Gauntlet.TestNode.cs b/Engine/Source/Programs/AutomationTool/Gauntlet/Framework/Base/Gauntlet.TestNode.cs
--- a/Engine/Source/Programs/AutomationTool/Gauntlet/Framework/Base/Gauntlet.TestNode.cs
+++ b/Engine/Source/Programs/AutomationTool/Gauntlet/Framework/Base/Gauntlet.TestNode.cs
@@ -76,12 +76,17 @@ namespace Gauntlet
 	public interface ITestEvent
 	{
 		/// <summary>
-		/// Level of severity that descrives this event
+		/// Level of severity that describes this event
 		/// </summary>
 		EventSeverity Severity { get; }
 
 		/// <summary>
-		/// High level single-line summary of what occurred. Should neber be null
+		/// Time at which the event was fired
+		/// </summary>
+		DateTime Time { get; }
+
+		/// <summary>
+		/// High level single-line summary of what occurred. Should never be null
 		/// </summary>
 		string Summary { get; }
 
diff --git a/Engine/Source/Programs/AutomationTool/Gauntlet/Unreal/Base/Gauntlet.UnrealTestEvent.cs b/Engine/Source/Programs/AutomationTool/Gauntlet/Unreal/Base/Gauntlet.UnrealTestEvent.cs
--- a/Engine/Source/Programs/AutomationTool/Gauntlet/Unreal/Base/Gauntlet.UnrealTestEvent.cs
+++ b/Engine/Source/Programs/AutomationTool/Gauntlet/Unreal/Base/Gauntlet.UnrealTestEvent.cs
@@ -1,5 +1,6 @@
 ﻿// Copyright Epic Games, Inc. All Rights Reserved.
 
+using System;
 using System.Collections.Generic;
 using System.Linq;
 
@@ -15,6 +16,9 @@ namespace Gauntlet
 		/// </summary>
 		public EventSeverity Severity { get; protected set; }
 
+		// Time at which the event was fired
+		public DateTime Time { get; }
+
 		// Title/Summary of event
 		public string Summary { get; protected set; }
 
@@ -24,12 +28,19 @@ namespace Gauntlet
 		// Callstack 
 		public IEnumerable<string> Callstack { get; protected set; }
 
-		//True if this is an ensure (Gauntlet does not define a level for this, but we log differently).
+		// True if this is an ensure (Gauntlet does not define a level for this, but we log differently).
 		public bool IsEnsure { get; protected set; }
 
+		// True if the event severity is Error or Fatal
+		public bool IsError => Severity == EventSeverity.Error || Severity == EventSeverity.Fatal;
+
+		// True if the event severity is Warning
+		public bool IsWarning => Severity == EventSeverity.Warning;
+
 		// Constructor that requires all properties
 		public UnrealTestEvent(EventSeverity InSeverity, string InSummary, IEnumerable<string> InDetails, UnrealLog.CallstackMessage InCallstack = null)
 		{
+			Time = DateTime.UtcNow;
 			Severity = InSeverity;
 			Summary = InSummary;
 			Details = InDetails.ToArray();
