Screen record the Gauntlet run, then put the recording in artifacts.

diff --git a/Engine/Source/Programs/AutomationTool/Gauntlet/Platform/Android/Gauntlet.TargetDeviceAndroid.cs b/Engine/Source/Programs/AutomationTool/Gauntlet/Platform/Android/Gauntlet.TargetDeviceAndroid.cs
--- a/Engine/Source/Programs/AutomationTool/Gauntlet/Platform/Android/Gauntlet.TargetDeviceAndroid.cs
+++ b/Engine/Source/Programs/AutomationTool/Gauntlet/Platform/Android/Gauntlet.TargetDeviceAndroid.cs
@@ -1519,12 +1519,19 @@ namespace Gauntlet
 		private DateTime ActivityCheckTime = DateTime.UtcNow;
 		private bool ActivityExited = false;
 
+		public AdbScreenRecorder Recorder;
+
 		public AndroidAppInstance(TargetDeviceAndroid InDevice, AndroidAppInstall InInstall, IProcessResult InProcess)
 		{
 			AndroidDevice = InDevice;
 			Install = InInstall;
 			LaunchProcess = InProcess;
 			LogProcess = AndroidDevice.RunAdbDeviceCommand($"logcat -s {Install.AppTag} debug Debug DEBUG -v raw", false, false);
+
+			if (Globals.Params.ParseParam("screenrecord"))
+			{
+				StartRecording();
+			}
 		}
 
 		public int WaitForExit()
@@ -1545,11 +1552,29 @@ namespace Gauntlet
 				WasKilled = true;
 				Install.AndroidDevice.KillRunningProcess(Install.AndroidPackageName);
 				LogProcess.StopProcess();
+				StopRecording();
+			}
+		}
+		public void StartRecording()
+		{
+			Recorder = new AdbScreenRecorder();
+			Recorder.StartRecording(Install.AndroidDevice.DeviceName, $"{Install.AndroidDevice.DeviceArtifactPath}/Logs/screen_recording.mp4");
+		}
+
+		public void StopRecording()
+		{
+			if (Recorder != null)
+			{
+				Recorder.StopRecording();
+				Recorder.Dispose();
+				Recorder = null;
 			}
 		}
 
 		protected void SaveArtifacts()
 		{
+			StopRecording();
+
 			// Pull all the artifacts
 			string ArtifactPullCommand = string.Format("pull {0} \"{1}\"", Install.AndroidDevice.DeviceArtifactPath, Install.AndroidDevice.LocalCachePath);
 			IProcessResult PullCmd = Install.AndroidDevice.RunAdbDeviceCommand(ArtifactPullCommand, bShouldLogCommand: Log.IsVerbose);
@@ -1615,6 +1640,102 @@ namespace Gauntlet
 
 		}
 	}
+	public class AdbScreenRecorder : IDisposable
+	{
+		private Process ADBProcess;
+		private bool Disposed;
+
+		public void StartRecording(string Device, string OutputFilePath)
+		{
+			if (ADBProcess != null && !ADBProcess.HasExited)
+			{
+				throw new InvalidOperationException("A screen recording session is already in progress.");
+			}
+
+			string AdbCommand = Environment.ExpandEnvironmentVariables("%ANDROID_HOME%/platform-tools/adb" + (RuntimePlatform.IsWindows ? ".exe" : ""));
+
+			ADBProcess = new Process
+			{
+				StartInfo = new ProcessStartInfo
+				{
+					FileName = AdbCommand,
+					Arguments = $"-s {Device} shell screenrecord {OutputFilePath}",
+					RedirectStandardOutput = true,
+					RedirectStandardInput = true,
+					UseShellExecute = false,
+					CreateNoWindow = true
+				}
+			};
+
+			ADBProcess.Start();
+		}
+		public void StopRecording(int timeout = 500)
+		{
+			if (ADBProcess != null && !ADBProcess.HasExited)
+			{
+				try
+				{
+					ADBProcess.StandardInput.WriteLine("\x03"); // Send Ctrl+C to stop the recording
+					if (!ADBProcess.WaitForExit(timeout))
+					{
+						// Process did not exit within the specified timeout
+						KillAdbProcess();
+					}
+				}
+				catch (InvalidOperationException)
+				{
+					// Handle the exception if StandardInput is not available
+					KillAdbProcess();
+				}
+			}
+		}
+		private void KillAdbProcess()
+		{
+			try
+			{
+				if (ADBProcess != null && !ADBProcess.HasExited)
+				{
+					ADBProcess.Kill();
+				}
+			}
+			catch (InvalidOperationException)
+			{
+				// Process has already exited, do nothing
+			}
+		}
+
+		public void Dispose()
+		{
+			Dispose(true);
+			GC.SuppressFinalize(this);
+		}
+
+		protected virtual void Dispose(bool disposing)
+		{
+			if (!Disposed)
+			{
+				if (disposing)
+				{
+					// Release managed resources
+					StopRecording();
+				}
+
+				// Release unmanaged resources
+				if (ADBProcess != null)
+				{
+					ADBProcess.Dispose();
+					ADBProcess = null;
+				}
+
+				Disposed = true;
+			}
+		}
+
+		~AdbScreenRecorder()
+		{
+			Dispose(false);
+		}
+	}
 
 	// Device data from json
 	public sealed class AndroidDeviceData
