In Gauntlet, MANAGE_EXTERNAL_STORAGE allow permission. Also move perm allow to AFTER the install.

diff --git a/Engine/Source/Programs/AutomationTool/Gauntlet/Platform/Android/Gauntlet.TargetDeviceAndroid.cs b/Engine/Source/Programs/AutomationTool/Gauntlet/Platform/Android/Gauntlet.TargetDeviceAndroid.cs
--- a/Engine/Source/Programs/AutomationTool/Gauntlet/Platform/Android/Gauntlet.TargetDeviceAndroid.cs
+++ b/Engine/Source/Programs/AutomationTool/Gauntlet/Platform/Android/Gauntlet.TargetDeviceAndroid.cs
@@ -367,15 +367,14 @@ namespace Gauntlet
 			string Package = Build.AndroidPackageName;
 			KillRunningProcess(Package);
 
-			EnablePermissions(Package);
-
-
 			// Install apk
 			string APK = Globals.IsRunningDev && AppConfig.OverlayExecutable.GetOverlay(Build.SourceApkPath, out string OverlayAPK)
 				? OverlayAPK
 				: Build.SourceApkPath;
 			CopyFileToDevice(Package, APK, string.Empty);
 
+			EnablePermissions(Package);
+
 			// Copy obbs from bulk builds
 			bool bSkipOBBInstall = Globals.Params.ParseParam("SkipOBBCopy"); // useful when iterating on dev executables
 			if(Build.FilesToInstall != null && Build.FilesToInstall.Any() && !bSkipOBBInstall)
@@ -956,6 +955,8 @@ namespace Gauntlet
 				Log.Verbose(string.Format("Enabling permission: {0} {1}", AndroidPackageName, Permission));
 				RunAdbDeviceCommand(CommandLine, true, false, true);
 			});
+			Log.Verbose($"Enabling permission: {AndroidPackageName} MANAGE_EXTERNAL_STORAGE");
+			RunAdbDeviceCommand($"shell appops set {AndroidPackageName} MANAGE_EXTERNAL_STORAGE allow");
 		}
 
 		public void KillRunningProcess(string AndroidPackageName)
