Store screen recordings in Saved directory not Saved/Logs.

diff --git a/Engine/Source/Programs/AutomationTool/Gauntlet/Platform/Android/Gauntlet.TargetDeviceAndroid.cs b/Engine/Source/Programs/AutomationTool/Gauntlet/Platform/Android/Gauntlet.TargetDeviceAndroid.cs
--- a/Engine/Source/Programs/AutomationTool/Gauntlet/Platform/Android/Gauntlet.TargetDeviceAndroid.cs
+++ b/Engine/Source/Programs/AutomationTool/Gauntlet/Platform/Android/Gauntlet.TargetDeviceAndroid.cs
@@ -1559,7 +1559,7 @@ namespace Gauntlet
 		public void StartRecording()
 		{
 			Recorder = new AdbScreenRecorder();
-			Recorder.StartRecording(Install.AndroidDevice.DeviceName, $"{Install.AndroidDevice.DeviceArtifactPath}/Logs/screen_recording.mp4");
+			Recorder.StartRecording(Install.AndroidDevice.DeviceName, $"{Install.AndroidDevice.DeviceArtifactPath}/screen_recording.mp4");
 		}
 
 		public void StopRecording()
