Save screen_recording to Saved/Logs directory. Make sure the directory is created before saving!

diff --git a/Engine/Source/Programs/AutomationTool/Gauntlet/Platform/Android/Gauntlet.TargetDeviceAndroid.cs b/Engine/Source/Programs/AutomationTool/Gauntlet/Platform/Android/Gauntlet.TargetDeviceAndroid.cs
--- a/Engine/Source/Programs/AutomationTool/Gauntlet/Platform/Android/Gauntlet.TargetDeviceAndroid.cs
+++ b/Engine/Source/Programs/AutomationTool/Gauntlet/Platform/Android/Gauntlet.TargetDeviceAndroid.cs
@@ -1557,9 +1557,18 @@ namespace Gauntlet
 			}
 		}
 		public void StartRecording()
-		{
+		{	
+			// Ensure artifact directories exist
+			if (TargetDeviceAndroid.UsingAndroidFileServer(Install.ProjectFile, Install.Configuration, out _, out string AFSToken, out _, out _, out _))
+			{
+				Install.AndroidDevice.RunAFSDeviceCommand(string.Format("-p \"{0}\" -k \"{1}\" mkdir \"^saved/Logs\"", Install.AndroidPackageName, AFSToken));
+			}
+			else
+			{ 
+				Install.AndroidDevice.RunAdbDeviceCommand($"shell mkdir -p {Install.AndroidDevice.DeviceArtifactPath}/Logs/");
+			}
 			Recorder = new AdbScreenRecorder();
-			Recorder.StartRecording(Install.AndroidDevice.DeviceName, $"{Install.AndroidDevice.DeviceArtifactPath}/screen_recording.mp4");
+			Recorder.StartRecording(Install.AndroidDevice.DeviceName, $"{Install.AndroidDevice.DeviceArtifactPath}/Logs/screen_recording.mp4");
 		}
 
 		public void StopRecording()
