Re-order Ignore check so potential error does not fire inside the compare check.
Convert relative file paths to full paths to match the query which is based on full path.

diff --git a/Engine/Source/Programs/AutomationTool/BuildGraph/BuildGraph.cs b/Engine/Source/Programs/AutomationTool/BuildGraph/BuildGraph.cs
--- a/Engine/Source/Programs/AutomationTool/BuildGraph/BuildGraph.cs
+++ b/Engine/Source/Programs/AutomationTool/BuildGraph/BuildGraph.cs
@@ -933,7 +933,8 @@ namespace AutomationTool
 			HashSet<BgNodeDef> CleanedNodes = new HashSet<BgNodeDef>();
 			foreach (BgNodeDef NodeToExecute in NodesToExecute)
 			{
-				if (NodeToExecute.InputDependencies.Any(x => CleanedNodes.Contains(x)) || !Storage.CheckLocalIntegrity(NodeToExecute.Name, NodeToExecute.Outputs.Select(x => x.TagName)))
+				FileFilter IgnoreModifiedFilter = new FileFilter(NodeToExecute.IgnoreModified);
+				if (NodeToExecute.InputDependencies.Any(x => CleanedNodes.Contains(x)) || !Storage.CheckLocalIntegrity(NodeToExecute.Name, NodeToExecute.Outputs.Select(x => x.TagName), IgnoreModifiedFilter))
 				{
 					Storage.CleanLocalNode(NodeToExecute.Name);
 					CleanedNodes.Add(NodeToExecute);
@@ -1088,7 +1089,7 @@ namespace AutomationTool
 			foreach (TempStorageFile File in InputManifests.Values.SelectMany(x => x.Files))
 			{
 				string? Message;
-				if (!ModifiedFiles.ContainsKey(File.RelativePath) && !File.Compare(Unreal.RootDirectory, out Message) && !IgnoreModifiedFilter.Matches(File.RelativePath))
+				if (!ModifiedFiles.ContainsKey(File.RelativePath) && !IgnoreModifiedFilter.Matches(File.ToFileReference(Unreal.RootDirectory).FullName) && !File.Compare(Unreal.RootDirectory, out Message))
 				{
 					// look up the previous nodes to help with error diagnosis
 					List<string> PreviousNodeNames = InputManifests.Where(x => x.Value.Files.Contains(File))
diff --git a/Engine/Source/Programs/AutomationTool/BuildGraph/TempStorage.cs b/Engine/Source/Programs/AutomationTool/BuildGraph/TempStorage.cs
--- a/Engine/Source/Programs/AutomationTool/BuildGraph/TempStorage.cs
+++ b/Engine/Source/Programs/AutomationTool/BuildGraph/TempStorage.cs
@@ -662,8 +662,9 @@ namespace AutomationTool
 		/// </summary>
 		/// <param name="NodeName">The node to retrieve build products for</param>
 		/// <param name="TagNames">List of tag names from this node.</param>
+		/// <param name="IgnoreModified">Filter for files to ignore when performing integrity check. Specified per node. </param>
 		/// <returns>True if the node is complete and valid, false if not (and typically followed by a call to CleanNode()).</returns>
-		public bool CheckLocalIntegrity(string NodeName, IEnumerable<string> TagNames)
+		public bool CheckLocalIntegrity(string NodeName, IEnumerable<string> TagNames, FileFilter IgnoreModified)
 		{
 			// If the node is not locally complete, fail immediately.
 			FileReference CompleteMarkerFile = GetCompleteMarkerFile(LocalDir, NodeName);
@@ -740,7 +741,7 @@ namespace AutomationTool
 
 				// Read the manifest and check the files
 				TempStorageManifest LocalManifest = TempStorageManifest.Load(LocalManifestFile);
-				if(LocalManifest.Files.Any(x => !x.Compare(RootDir)))
+				if(LocalManifest.Files.Any(x => !IgnoreModified.Matches(x.ToFileReference(RootDir).FullName) && !x.Compare(RootDir)))
 				{
 					return false;
 				}
@@ -974,7 +975,7 @@ namespace AutomationTool
 				foreach(TempStorageFile File in Manifest.Files)
 				{
 					string? Message;
-					if (!IgnoreModified.Matches(File.RelativePath) && !File.Compare(RootDir, out Message))
+					if (!IgnoreModified.Matches(File.ToFileReference(RootDir).FullName) && !File.Compare(RootDir, out Message))
 					{
 						ModifiedFileMessages.Add(Message);
 					}
