Add support for the cooked cooker to stage the .target file it is built with & for this file to be read by the editor
This allows dlls to be loaded based on the target file when not running in monolithic mode
#rnx

diff --git a/Engine/Source/Developer/DesktopPlatform/Private/TargetReceipt.cpp b/Engine/Source/Developer/DesktopPlatform/Private/TargetReceipt.cpp
--- a/Engine/Source/Developer/DesktopPlatform/Private/TargetReceipt.cpp
+++ b/Engine/Source/Developer/DesktopPlatform/Private/TargetReceipt.cpp
@@ -85,7 +85,11 @@ bool FTargetReceipt::Read(const FString& FileName, bool bExpandVariables)
 	// Get the launch path
 	if(!Object->TryGetStringField(TEXT("Launch"), Launch))
 	{
-		return false;
+		//If the regular launch field is not present try using the LaunchCmd field
+		if (!Object->TryGetStringField(TEXT("LaunchCmd"), Launch))
+		{
+			return false;
+		}
 	}
 	if (bExpandVariables)
 	{
diff --git a/Engine/Source/Programs/AutomationTool/CookedEditor/CookedEditor.Automation.cs b/Engine/Source/Programs/AutomationTool/CookedEditor/CookedEditor.Automation.cs
--- a/Engine/Source/Programs/AutomationTool/CookedEditor/CookedEditor.Automation.cs
+++ b/Engine/Source/Programs/AutomationTool/CookedEditor/CookedEditor.Automation.cs
@@ -87,6 +87,7 @@ public class ModifyStageContext
 	public bool bStageUAT = false;
 	public bool bIsForExternalDistribution = false;
 	public bool bStagePython = false;
+	public bool bStageTargetFiles = false;
 
 	public ConfigHelper ConfigHelper;
 
@@ -1006,8 +1007,11 @@ public class MakeCookedEditor : BuildCommand
 		// we already cooked assets, so remove assets we may have found, except for the Uncook ones
 		Context.UFSFilesToStage.RemoveAll(x => x.GetExtension() == ".uasset");
 
-		// don't need the .target files
-		Context.NonUFSFilesToStage.RemoveAll(x => x.GetExtension() == ".target");
+		if (!Context.bStageTargetFiles)
+		{
+			// don't need the .target files
+			Context.NonUFSFilesToStage.RemoveAll(x => x.GetExtension() == ".target");
+		}
 
 		if (!Context.bStageShaderDirs)
 		{
