Fix a null deref if you didn't specify -build with the LowLevelTests.

It'll later complain about the lack of build, but it has already tried to unconditionally deref the null before then.

diff --git a/Engine/Source/Programs/AutomationTool/LowLevelTests/RunLowLevelTests.cs b/Engine/Source/Programs/AutomationTool/LowLevelTests/RunLowLevelTests.cs
--- a/Engine/Source/Programs/AutomationTool/LowLevelTests/RunLowLevelTests.cs
+++ b/Engine/Source/Programs/AutomationTool/LowLevelTests/RunLowLevelTests.cs
@@ -234,7 +234,17 @@ namespace LowLevelTests
 
 			CaptureOutput = Params.ParseParam("captureoutput");
 
-			Build = Params.ParseValue("build=", null).Replace('/', Path.DirectorySeparatorChar).Replace('\\', Path.DirectorySeparatorChar);
+			var buildValue = Params.ParseValue("build=", null);
+
+			if (buildValue != null)
+			{
+				Build = buildValue.Replace('/', Path.DirectorySeparatorChar).Replace('\\', Path.DirectorySeparatorChar);
+			}
+			else
+			{
+				Build = null;
+			}
+
 			TestApp = Globals.Params.ParseValue("testapp=", "");
 
 			Tags = Params.ParseValue("tags=", null);
