Add a marker for the end of ASan reports on Windows. This allows Gauntlet to extract the callstack and present it during autotests.

#tests using local AIBotSoakASAN
#rnx

diff --git a/Engine/Source/Programs/AutomationTool/Gauntlet/Unreal/Utils/Gauntlet.UnrealLogParser.cs b/Engine/Source/Programs/AutomationTool/Gauntlet/Unreal/Utils/Gauntlet.UnrealLogParser.cs
--- a/Engine/Source/Programs/AutomationTool/Gauntlet/Unreal/Utils/Gauntlet.UnrealLogParser.cs
+++ b/Engine/Source/Programs/AutomationTool/Gauntlet/Unreal/Utils/Gauntlet.UnrealLogParser.cs
@@ -983,7 +983,7 @@ namespace Gauntlet
 			foreach (Match TraceInitMatch in Matches)
 			{
 				string TraceID = TraceInitMatch.Groups[1].Value;
-				string EndPattern = @$"=={TraceID}==ABORTING";
+				string EndPattern = @$"(=={TraceID}==ABORTING)|(End of Address Sanitizer report)";
 				int TraceInitIndex = TraceInitMatch.Index + TraceInitMatch.Length;
 
 				UnrealLog.CallstackMessage NewTrace = new UnrealLog.CallstackMessage();
diff --git a/Engine/Source/Runtime/Launch/Private/Windows/LaunchWindows.cpp b/Engine/Source/Runtime/Launch/Private/Windows/LaunchWindows.cpp
--- a/Engine/Source/Runtime/Launch/Private/Windows/LaunchWindows.cpp
+++ b/Engine/Source/Runtime/Launch/Private/Windows/LaunchWindows.cpp
@@ -27,6 +27,8 @@ static void ASanErrorCallback(const char* ErrorStr)
 		UE_LOG(LogASan, Warning, TEXT("Critical error raised prior to ASan error; engine memory may be unstable in this state!"));
 	}
 	UE_LOG(LogASan, Fatal, TEXT("ASan Error: %s"), ANSI_TO_TCHAR(ErrorStr));
+	/* Marking the end of sanitizer report for Gauntlet. */
+	UE_LOG(LogASan, Fatal, TEXT("ASan Error: %s\nEnd of Address Sanitizer report"), ANSI_TO_TCHAR(ErrorStr));
 }
 
 #endif
