[Gauntlet] Add optional flag to ignore drive mapping failures on ps4 and ps5
Also modify InstallUnreal to use InstallBuild over InstallApplication

diff --git a/Engine/Source/Programs/AutomationTool/Gauntlet/Unreal/Utils/Gauntlet.InstallUnreal.cs b/Engine/Source/Programs/AutomationTool/Gauntlet/Unreal/Utils/Gauntlet.InstallUnreal.cs
--- a/Engine/Source/Programs/AutomationTool/Gauntlet/Unreal/Utils/Gauntlet.InstallUnreal.cs
+++ b/Engine/Source/Programs/AutomationTool/Gauntlet/Unreal/Utils/Gauntlet.InstallUnreal.cs
@@ -170,7 +170,7 @@ namespace Gauntlet
 				if (Config.Build != null)
 				{
 					Log.Info("Installing build on device {DeviceName}", Device.Name);
-					IAppInstall Install = Device.InstallApplication(Config);
+					Device.InstallBuild(Config);
 				}
 				else
 				{
