Change method of ordering AFS startup sequence (am start -W doesn't work for older Android)

#android

diff --git a/Engine/Source/Programs/AutomationTool/Android/AndroidFileClient.cs b/Engine/Source/Programs/AutomationTool/Android/AndroidFileClient.cs
--- a/Engine/Source/Programs/AutomationTool/Android/AndroidFileClient.cs
+++ b/Engine/Source/Programs/AutomationTool/Android/AndroidFileClient.cs
@@ -1382,11 +1382,13 @@ namespace AutomationTool
 
 		private class WaitCallbackData
 		{
+			public double AfterTime;
 			public string WaitForLine;
 			public string MessageBuffer;
 
-			public WaitCallbackData(string InWaitForLine)
+			public WaitCallbackData(double InAfterTime, string InWaitForLine)
 			{
+				AfterTime = InAfterTime;
 				WaitForLine = InWaitForLine;
 				MessageBuffer = "";
 			}
@@ -1397,17 +1399,78 @@ namespace AutomationTool
 				MessageBuffer += InNewData;
 
 				// check if we now have the search line
-				if (MessageBuffer.IndexOf(WaitForLine) >= 0)
+				while (true)
 				{
-					MessageBuffer = "";
-					return true;
+					int WaitLineIndex = MessageBuffer.IndexOf(WaitForLine);
+					if (WaitLineIndex < 0)
+					{
+						// stop if not found
+						break;
+					}
+	
+					// done if don't care about time
+					if (AfterTime < 0.0)
+					{
+						MessageBuffer = "";
+						return true;
+					}
+
+					// get from beginning of line with the match
+					string Trimmed = MessageBuffer.Substring(0, WaitLineIndex);
+					int EOLIndex = Trimmed.LastIndexOf("\n");
+					if (EOLIndex >= 0)
+					{
+						Trimmed = Trimmed.Substring(EOLIndex + 1);
+					}
+
+					// keep only past this found line for next iteration
+					MessageBuffer = MessageBuffer.Substring(WaitLineIndex + 1);
+					int FirstLineIndex = MessageBuffer.IndexOf('\n');
+					if (FirstLineIndex >= 0)
+					{
+						MessageBuffer = MessageBuffer.Substring(FirstLineIndex + 1);
+					}
+
+					// parse the time
+					string Result = "0.000";
+					Trimmed = Trimmed.TrimStart();
+					if (string.IsNullOrWhiteSpace(Trimmed) || !char.IsDigit(Trimmed[0]))
+					{
+						continue;
+					}
+
+					int SpaceIndex = Trimmed.IndexOf(' ');
+					int TabIndex = Trimmed.IndexOf('\t');
+					if (SpaceIndex > 0 && TabIndex > 0)
+					{
+						Result = Trimmed.Substring(0, Math.Min(SpaceIndex, TabIndex));
+					}
+					else if (SpaceIndex > 0)
+					{
+						Result = Trimmed.Substring(0, SpaceIndex);
+					}
+					else if (TabIndex > 0)
+					{
+						Result = Trimmed.Substring(0, TabIndex);
+					}
+
+					double LineTime = 0.0;
+					if (Double.TryParse(Result, out LineTime))
+					{
+						if (LineTime > AfterTime)
+						{
+//								Console.WriteLine("GOT IT at: " + LineTime);
+							MessageBuffer = "";
+							return true;
+						}
+					}
 				}
 
 				// drop old lines to keep the search faster and free memory
 				int LastLineIndex = MessageBuffer.LastIndexOf('\n');
 				if (LastLineIndex >= 0)
 				{
-					MessageBuffer.Substring(LastLineIndex + 1);
+					MessageBuffer = MessageBuffer.Substring(LastLineIndex + 1);
 				}
 
 				return false;
@@ -1420,9 +1483,9 @@ namespace AutomationTool
 			return WorkData.CheckForLine(Message) ? false : true;
 		}
 
-		public static bool WaitForLogcat(string Device, string Options, string Line, int Timeout)
+		public static bool WaitForLogcat(string Device, string Options, string Line, int Timeout, double AfterTime = -1.0)
 		{
-			WaitCallbackData WorkData = new WaitCallbackData(Line);
+			WaitCallbackData WorkData = new WaitCallbackData(AfterTime, Line);
 
 			string Result = Logcat(Device, Options, Timeout, LogcatWaitCallback, (object)WorkData);
 			return Result != "TIMEOUT";
@@ -1663,6 +1726,46 @@ namespace AutomationTool
 			return false;
 		}
 
+		private static string CurrentAndroidSDKLevelDevice = "";
+		private static int CurrentAndroidSDKLevel = -1;
+
+		private static int GetDeviceAndroidSDKLevel(string Device)
+		{
+			if (CurrentAndroidSDKLevelDevice != Device || CurrentAndroidSDKLevel == -1)
+			{
+				string Result = adb.Shell(Device, "getprop ro.build.version.sdk");
+				if (int.TryParse(Result.Trim(), out CurrentAndroidSDKLevel))
+				{
+					CurrentAndroidSDKLevelDevice = Device;
+				}
+			}
+			return CurrentAndroidSDKLevel;
+		}
+
+		private static bool ActivityManagerSendAndWait(string Device, string PackageName, string Arguments)
+		{
+			// am start -W seems to work past Android 10 reliably (but use a timeout just in case)
+			if (GetDeviceAndroidSDKLevel(Device) > 29)
+			{
+				adb.Shell(Device, "am start -W -a com.epicgames.unreal.RemoteFileManager.intent.COMMAND2 -n " + PackageName + "/com.epicgames.unreal.RemoteFileManagerActivity " + Arguments, 1500);
+				return true;
+			}
+
+			// need to monitor logcat for response from receiver (with timeout)
+			string LastTime = GetLastLogcatTime(Device);
+			double LastTimeValue = 0.0;
+			Double.TryParse(LastTime, out LastTimeValue);
+//			Console.WriteLine("Sending " + Arguments + " to " + PackageName + " at " + LastTime);
+
+			adb.Shell(Device, "am start -a com.epicgames.unreal.RemoteFileManager.intent.COMMAND2 -n " + PackageName + "/com.epicgames.unreal.RemoteFileManagerActivity " + Arguments);
+
+//			Stopwatch stopwatch = System.Diagnostics.Stopwatch.StartNew();
+			bool retval = WaitForLogcat(Device, "-v monotonic -s UEFS", "package = " + PackageName, 2500, LastTimeValue);
+//			stopwatch.Stop();
+//			Console.WriteLine("Time: " + (double)stopwatch.ElapsedMilliseconds / 1000.0);
+			return retval;
+		}
+
 		public static bool StopAnyServers(string Device, int Port, bool bWaitForStop = true)
 		{
 			// get a list of installed receivers and verify package is available first
@@ -1681,7 +1784,7 @@ namespace AutomationTool
 			foreach (string Activity in InstalledActivities)
 			{
 				bDidSendStops = true;
-				adb.Shell(Device, "am start -W -a com.epicgames.unreal.RemoteFileManager.intent.COMMAND2 -n " + Activity + "/com.epicgames.unreal.RemoteFileManagerActivity -e cmd stop");
+				ActivityManagerSendAndWait(Device, Activity, "-e cmd stop");
 			}
 
 			// it can take up to 2 seconds for running servers to terminate so check if any binds to port are still active
@@ -1796,9 +1899,7 @@ namespace AutomationTool
 			string WifiAddress;
 
 			// sent start request (won't do anything if already started)
-//			string StartCommand = "am broadcast -a com.epicgames.unreal.RemoteFileManager.intent.COMMAND -n " + PackageName + "/com.epicgames.unreal.RemoteFileManagerReceiver -e cmd start -e token " + Token + " -ei port " + ServerPort;
-			string StartCommand = "am start -W -a com.epicgames.unreal.RemoteFileManager.intent.COMMAND2 -n " + PackageName + "/com.epicgames.unreal.RemoteFileManagerActivity -e cmd start -e token " + Token + " -ei port " + ServerPort;
-			adb.Shell(Device, StartCommand);
+			ActivityManagerSendAndWait(Device, PackageName, "-e cmd start -e token " + Token + " -ei port " + ServerPort);
 
 			// see if we can check listen status
 			if (GetListenStatus(Device, ServerPort, out bUSB, out bWifi, out WifiAddress))
@@ -1818,7 +1919,7 @@ namespace AutomationTool
 					Thread.Sleep(200);
 
 					// sent start request again (won't do anything if already started)
-					adb.Shell(Device, StartCommand);
+					ActivityManagerSendAndWait(Device, PackageName, "-e cmd start -e token " + Token + " -ei port " + ServerPort);
 
 					GetListenStatus(Device, ServerPort, out bUSB, out bWifi, out WifiAddress);
 				}
@@ -1839,7 +1940,7 @@ namespace AutomationTool
 				}
 
 				// sent start request again (won't do anything if already started)
-				adb.Shell(Device, StartCommand);
+				ActivityManagerSendAndWait(Device, PackageName, "-e cmd start -e token " + Token + " -ei port " + ServerPort);
 			}
 
 			Logger.LogInformation("Timed out on connection attempts");
diff --git a/Engine/Source/Programs/UnrealAndroidFileTool/Program.cs b/Engine/Source/Programs/UnrealAndroidFileTool/Program.cs
--- a/Engine/Source/Programs/UnrealAndroidFileTool/Program.cs
+++ b/Engine/Source/Programs/UnrealAndroidFileTool/Program.cs
@@ -17,7 +17,7 @@ namespace UnrealAndroidFileTool
 {
 	class Program
 	{
-		private static string AppVersion = "1.0.4";
+		private static string AppVersion = "1.0.5";
 		private static ILogger Logger = Log.Logger;
 
 		static void ShowHelp(string command)
