Add HostArchitecture to default BuildGraph properties

#rnx

diff --git a/Engine/Source/Programs/AutomationTool/BuildGraph/BuildGraph.cs b/Engine/Source/Programs/AutomationTool/BuildGraph/BuildGraph.cs
--- a/Engine/Source/Programs/AutomationTool/BuildGraph/BuildGraph.cs
+++ b/Engine/Source/Programs/AutomationTool/BuildGraph/BuildGraph.cs
@@ -278,6 +278,7 @@ namespace AutomationTool
 			DefaultProperties["CodeChange"] = (CodeChange ?? 0).ToString();
 			DefaultProperties["IsBuildMachine"] = IsBuildMachine ? "true" : "false";
 			DefaultProperties["HostPlatform"] = HostPlatform.Current.HostEditorPlatform.ToString();
+			DefaultProperties["HostArchitecture"] = RuntimeInformation.ProcessArchitecture.ToString();
 			DefaultProperties["RestrictedFolderNames"] = String.Join(";", RestrictedFolder.GetNames());
 			DefaultProperties["RestrictedFolderFilter"] = String.Join(";", RestrictedFolder.GetNames().Select(x => String.Format(".../{0}/...", x)));
 			DefaultProperties["DataDrivenPlatforms"] = String.Join(";", DataDrivenPlatformInfo.GetAllPlatformInfos().Keys);
