AutomationTool: Permit MakeCookedEditor subclasses to override ProjectParams construction.

diff --git a/Engine/Source/Programs/AutomationTool/CookedEditor/CookedEditor.Automation.cs b/Engine/Source/Programs/AutomationTool/CookedEditor/CookedEditor.Automation.cs
--- a/Engine/Source/Programs/AutomationTool/CookedEditor/CookedEditor.Automation.cs
+++ b/Engine/Source/Programs/AutomationTool/CookedEditor/CookedEditor.Automation.cs
@@ -1040,17 +1040,9 @@ public class MakeCookedEditor : BuildCommand
 		return CommandUtils.CombinePaths(ProjectFile.Directory.FullName, "Releases", ReleaseVersionName, ReleaseTargetName);
 	}
 
-	private ProjectParams GetParams()
+	protected virtual ProjectParams MakeParams(string DLCName, string BasedOnReleaseVersion)
 	{
-		// setup DLC defaults, then ask project if it should 
-		string DLCName;
-		string BasedOnReleaseVersion;
-		TargetType ReleaseType;
-		SetupDLCMode(ProjectFile, out DLCName, out BasedOnReleaseVersion, out ReleaseType);
-		bool bIsDLC = DLCName != null;
-
-		var Params = new ProjectParams
-		(
+		return new ProjectParams(
 			Command: this
 			, RawProjectPath: ProjectFile
 
@@ -1060,8 +1052,19 @@ public class MakeCookedEditor : BuildCommand
 			, DedicatedServer: bIsCookedCooker
 			, NoClient: bIsCookedCooker
 			, OptionalContent: true
-
 		);
+	}
+
+	private ProjectParams GetParams()
+	{
+		// setup DLC defaults, then ask project if it should 
+		string DLCName;
+		string BasedOnReleaseVersion;
+		TargetType ReleaseType;
+		SetupDLCMode(ProjectFile, out DLCName, out BasedOnReleaseVersion, out ReleaseType);
+		bool bIsDLC = DLCName != null;
+
+		ProjectParams Params = MakeParams(DLCName, BasedOnReleaseVersion);
 
 		// cook the cooked editor targetplatorm as the "client"
 		//Params.ClientCookedTargets.Add("CrashReportClientEditor");
