BuildGraph: Add a task which will read secrets from Horde and substitute named placeholders for them in a text file.

#rnx

diff --git a/Engine/Source/Programs/AutomationTool/BuildGraph/Tasks/HordeGetSecretsTask.cs b/Engine/Source/Programs/AutomationTool/BuildGraph/Tasks/HordeGetSecretsTask.cs
new file mode 100644
--- /dev/null
+++ b/Engine/Source/Programs/AutomationTool/BuildGraph/Tasks/HordeGetSecretsTask.cs
@@ -0,0 +1,162 @@
+// Copyright Epic Games, Inc. All Rights Reserved.
+
+using EpicGames.Core;
+using EpicGames.Horde;
+using EpicGames.Horde.Secrets;
+using Microsoft.Extensions.DependencyInjection;
+using Microsoft.Extensions.Logging;
+using System;
+using System.Collections.Generic;
+using System.IO;
+using System.Linq;
+using System.Text;
+using System.Text.Json;
+using System.Threading.Tasks;
+using System.Xml;
+
+#nullable enable
+
+namespace AutomationTool.Tasks
+{
+	/// <summary>
+	/// Parameters for a <see cref="HordeGetSecretsTask"/>.
+	/// </summary>
+	public class HordeGetSecretsTaskParameters
+	{
+		/// <summary>
+		/// File to update with secrets
+		/// </summary>
+		[TaskParameter]
+		public string File = String.Empty;
+
+		/// <summary>
+		/// Pairs of strings and secret names to expand in the text file, in the form SOURCE_TEXT=secret-name;SOURCE_TEXT_2=secret-name-2
+		/// </summary>
+		[TaskParameter]
+		public string Replace = String.Empty;
+	}
+
+	/// <summary>
+	/// Replaces strings in a text file with secrets obtained from Horde
+	/// </summary>
+	[TaskElement("Horde-GetSecrets", typeof(HordeGetSecretsTaskParameters))]
+	public class HordeGetSecretsTask : BgTaskImpl
+	{
+		record class ReplacementInfo(string Property, string Variable);
+
+		readonly HordeGetSecretsTaskParameters _parameters;
+
+		/// <summary>
+		/// Constructor.
+		/// </summary>
+		/// <param name="parameters">Parameters for this task.</param>
+		public HordeGetSecretsTask(HordeGetSecretsTaskParameters parameters)
+		{
+			_parameters = parameters;
+		}
+
+		/// <summary>
+		/// Execute the task.
+		/// </summary>
+		/// <param name="Job">Information about the current job.</param>
+		/// <param name="BuildProducts">Set of build products produced by this node.</param>
+		/// <param name="TagNameToFileSet">Mapping from tag names to the set of files they include.</param>
+		public override async Task ExecuteAsync(JobContext Job, HashSet<FileReference> BuildProducts, Dictionary<string, HashSet<FileReference>> TagNameToFileSet)
+		{
+			// Read the input text
+			FileReference file = ResolveFile(_parameters.File);
+			string text = await FileReference.ReadAllTextAsync(file);
+
+			// Parse the secrets to replace
+			Dictionary<SecretId, List<ReplacementInfo>> secretToReplacementInfo = new Dictionary<SecretId, List<ReplacementInfo>>();
+			foreach (string clause in _parameters.Replace.Split(';', StringSplitOptions.RemoveEmptyEntries | StringSplitOptions.TrimEntries))
+			{
+				int idx = clause.LastIndexOf('=');
+				if (idx == -1)
+				{
+					throw new AutomationException($"Invalid replacement clause in Horde-GetSecrets task: {clause} (expected VARIABLE=Secret.Property)");
+				}
+
+				int propertyIdx = clause.IndexOf('.', idx + 1);
+				if (propertyIdx == -1)
+				{
+					throw new AutomationException($"Invalid replacement clause for secret in Horde-GetSecrets task: {clause} (expected VARIABLE=Secret.Property)");
+				}
+
+				string variable = clause.Substring(0, idx);
+				if (!text.Contains(variable, StringComparison.Ordinal))
+				{
+					Logger.LogWarning("Variable '{Variable}' not found in {File}", variable, file);
+					continue;
+				}
+
+				SecretId secretId = new SecretId(clause.Substring(idx + 1, propertyIdx - (idx + 1)));
+				string propertyName = clause.Substring(propertyIdx + 1);
+
+				List<ReplacementInfo>? replacements;
+				if (!secretToReplacementInfo.TryGetValue(secretId, out replacements))
+				{
+					replacements = new List<ReplacementInfo>();
+					secretToReplacementInfo.Add(secretId, replacements);
+				}
+
+				replacements.Add(new ReplacementInfo(propertyName, variable));
+			}
+
+			// Read the secrets from Horde, and substitute them in the output file
+			if (secretToReplacementInfo.Count > 0)
+			{
+				ServiceCollection serviceCollection = new ServiceCollection();
+				serviceCollection.AddHorde(options => options.AllowAuthPrompt = !CommandUtils.IsBuildMachine);
+
+				await using (ServiceProvider serviceProvider = serviceCollection.BuildServiceProvider())
+				{
+					IHordeClient hordeClient = serviceProvider.GetRequiredService<IHordeClient>();
+
+					using HordeHttpClient hordeHttpClient = hordeClient.CreateHttpClient();
+
+					foreach ((SecretId secretId, List<ReplacementInfo> replacements) in secretToReplacementInfo)
+					{
+						GetSecretResponse secret = await hordeHttpClient.GetSecretAsync(secretId);
+						foreach (ReplacementInfo replacement in replacements)
+						{
+							string? value;
+							if (secret.Data.TryGetValue(replacement.Property, out value))
+							{
+								text = text.Replace(replacement.Variable, value, StringComparison.Ordinal);
+							}
+							else
+							{
+								Logger.LogWarning("Property '{PropertyName}' not found in secret {SecretId}", replacement.Property, secretId);
+							}
+						}
+					}
+				}
+			}
+
+			// Write the output file
+			await FileReference.WriteAllTextAsync(file, text);
+			Logger.LogInformation("Updated {File} with secrets from Horde.", file);
+		}
+
+		/// <summary>
+		/// Output this task out to an XML writer.
+		/// </summary>
+		public override void Write(XmlWriter writer)
+		{
+			Write(writer, _parameters);
+		}
+
+		/// <summary>
+		/// Find all the tags which are used as inputs to this task
+		/// </summary>
+		/// <returns>The tag names which are read by this task</returns>
+		public override IEnumerable<string> FindConsumedTagNames() => Enumerable.Empty<string>();
+
+		/// <summary>
+		/// Find all the tags which are modified by this task
+		/// </summary>
+		/// <returns>The tag names which are modified by this task</returns>
+		public override IEnumerable<string> FindProducedTagNames() => Enumerable.Empty<string>();
+	}
+}
