BuildGraph: Add support for using properties rather than fields in parameter class definitions.

diff --git a/Engine/Source/Programs/AutomationTool/BuildGraph/BgNodeExecutor.cs b/Engine/Source/Programs/AutomationTool/BuildGraph/BgNodeExecutor.cs
--- a/Engine/Source/Programs/AutomationTool/BuildGraph/BgNodeExecutor.cs
+++ b/Engine/Source/Programs/AutomationTool/BuildGraph/BgNodeExecutor.cs
@@ -229,7 +229,7 @@ namespace AutomationTool
 						object? FieldValue = ParseValue(Value, Parameter.ValueType);
 						if (FieldValue != null)
 						{
-							Parameter.FieldInfo.SetValue(ParametersObject, FieldValue);
+							Parameter.SetValue(ParametersObject, FieldValue);
 						}
 						else if (!Parameter.Optional)
 						{
@@ -239,11 +239,11 @@ namespace AutomationTool
 					else
 					{
 						// Get the collection, or create one if necessary
-						object? CollectionValue = Parameter.FieldInfo.GetValue(ParametersObject);
+						object? CollectionValue = Parameter.GetValue(ParametersObject);
 						if (CollectionValue == null)
 						{
-							CollectionValue = Activator.CreateInstance(Parameter.FieldInfo.FieldType)!;
-							Parameter.FieldInfo.SetValue(ParametersObject, CollectionValue);
+							CollectionValue = Activator.CreateInstance(Parameter.ParameterType)!;
+							Parameter.SetValue(ParametersObject, CollectionValue);
 						}
 
 						// Parse the values and add them to the collection
diff --git a/Engine/Source/Programs/AutomationTool/BuildGraph/BuildGraph.cs b/Engine/Source/Programs/AutomationTool/BuildGraph/BuildGraph.cs
--- a/Engine/Source/Programs/AutomationTool/BuildGraph/BuildGraph.cs
+++ b/Engine/Source/Programs/AutomationTool/BuildGraph/BuildGraph.cs
@@ -26,21 +26,64 @@ namespace AutomationTool
 	/// <summary>
 	/// Implementation of ScriptTaskParameter corresponding to a field in a parameter class
 	/// </summary>
-	class ScriptTaskParameterBinding : BgScriptTaskParameter
+	abstract class ScriptTaskParameterBinding : BgScriptTaskParameter
 	{
-		/// <summary>
-		/// Field for this parameter
-		/// </summary>
-		public FieldInfo FieldInfo { get; }
+		public abstract Type ParameterType { get; }
 
-		/// <summary>
-		/// Constructor
-		/// </summary>
-		public ScriptTaskParameterBinding(string Name, FieldInfo FieldInfo, TaskParameterValidationType ValidationType, bool bOptional)
-			: base(Name, FieldInfo.FieldType, ValidationType, bOptional)
+		public ScriptTaskParameterBinding(string Name, Type ValueType, TaskParameterValidationType ValidationType, bool bOptional)
+			: base(Name, ValueType, ValidationType, bOptional)
+		{
+		}
+
+		public abstract string GetXmlDocName();
+		public abstract object? GetValue(object SourceObject);
+		public abstract void SetValue(object TargetObject, object? Value);
+	}
+
+	class ScriptTaskParameterPropertyBinding : ScriptTaskParameterBinding
+	{
+		readonly PropertyInfo _propertyInfo;
+
+		public ScriptTaskParameterPropertyBinding(string name, PropertyInfo propertyInfo, TaskParameterValidationType validationType, bool optional)
+			: base(name, propertyInfo.PropertyType, validationType, optional)
+		{
+			_propertyInfo = propertyInfo;
+		}
+
+		public override Type ParameterType
+			=> _propertyInfo.PropertyType;
+
+		public override string GetXmlDocName()
+			=> "P:" + _propertyInfo.DeclaringType!.FullName + "." + Name;
+
+		public override object? GetValue(object SourceObject)
+			=> _propertyInfo.GetValue(SourceObject);
+
+		public override void SetValue(object TargetObject, object? Value)
+			=> _propertyInfo.SetValue(TargetObject, Value);
+	}
+
+	class ScriptTaskParameterFieldBinding : ScriptTaskParameterBinding
+	{
+		readonly FieldInfo _fieldInfo;
+
+		public ScriptTaskParameterFieldBinding(string name, FieldInfo fieldInfo, TaskParameterValidationType validationType, bool optional)
+			: base(name, fieldInfo.FieldType, validationType, optional)
 		{
-			this.FieldInfo = FieldInfo;
+			_fieldInfo = fieldInfo;
 		}
+
+		public override Type ParameterType
+			=> _fieldInfo.FieldType;
+
+		public override string GetXmlDocName()
+			=> "F:" + _fieldInfo.DeclaringType!.FullName + "." + Name;
+
+		public override object? GetValue(object SourceObject)
+			=> _fieldInfo.GetValue(SourceObject);
+
+		public override void SetValue(object TargetObject, object? Value)
+			=> _fieldInfo.SetValue(TargetObject, Value);
 	}
 
 	/// <summary>
@@ -51,12 +94,12 @@ namespace AutomationTool
 		/// <summary>
 		/// Type of the task to construct with this info
 		/// </summary>
-		public Type TaskClass;
+		public Type TaskClass { get; }
 
 		/// <summary>
 		/// Type to construct with the parsed parameters
 		/// </summary>
-		public Type ParametersClass;
+		public Type ParametersClass { get; }
 
 		/// <summary>
 		/// Map from name to parameter
@@ -95,10 +138,18 @@ namespace AutomationTool
 					TaskParameterAttribute? ParameterAttribute = Field.GetCustomAttribute<TaskParameterAttribute>();
 					if (ParameterAttribute != null)
 					{
-						ScriptTaskParameters.Add(new ScriptTaskParameterBinding(Field.Name, Field, ParameterAttribute.ValidationType, ParameterAttribute.Optional));
+						ScriptTaskParameters.Add(new ScriptTaskParameterFieldBinding(Field.Name, Field, ParameterAttribute.ValidationType, ParameterAttribute.Optional));
 					}
 				}
 			}
+			foreach (PropertyInfo Property in ParametersClass.GetProperties(BindingFlags.Public | BindingFlags.Instance))
+			{
+				TaskParameterAttribute? ParameterAttribute = Property.GetCustomAttribute<TaskParameterAttribute>();
+				if (ParameterAttribute != null)
+				{
+					ScriptTaskParameters.Add(new ScriptTaskParameterPropertyBinding(Property.Name, Property, ParameterAttribute.ValidationType, ParameterAttribute.Optional));
+				}
+			}
 			return ScriptTaskParameters;
 		}
 	}
@@ -1345,9 +1396,9 @@ namespace AutomationTool
 
 							// Get the documentation for this parameter
 							XmlElement? ParameterElement;
-							if (MemberNameToElement.TryGetValue("F:" + Parameter.FieldInfo.DeclaringType!.FullName + "." + Parameter.Name, out ParameterElement))
+							if (MemberNameToElement.TryGetValue(Parameter.GetXmlDocName(), out ParameterElement))
 							{
-								Type FieldType = Parameter.FieldInfo.FieldType;
+								Type FieldType = Parameter.ParameterType;
 								if (FieldType.IsGenericType && FieldType.GetGenericTypeDefinition() == typeof(Nullable<>))
 								{
 									FieldType = FieldType.GetGenericArguments()[0];
@@ -1461,9 +1512,9 @@ namespace AutomationTool
 
 							// Get the documentation for this parameter
 							XmlElement? ParameterElement;
-							if (MemberNameToElement.TryGetValue("F:" + Parameter.FieldInfo.DeclaringType!.FullName + "." + Parameter.Name, out ParameterElement))
+							if (MemberNameToElement.TryGetValue(Parameter.GetXmlDocName(), out ParameterElement))
 							{
-								string TypeName = Parameter.FieldInfo.FieldType.Name;
+								string TypeName = Parameter.ParameterType.Name;
 								if (Parameter.ValidationType != TaskParameterValidationType.Default)
 								{
 									StringBuilder NewTypeName = new StringBuilder(Parameter.ValidationType.ToString());
diff --git a/Engine/Source/Programs/AutomationTool/BuildGraph/Task.cs b/Engine/Source/Programs/AutomationTool/BuildGraph/Task.cs
--- a/Engine/Source/Programs/AutomationTool/BuildGraph/Task.cs
+++ b/Engine/Source/Programs/AutomationTool/BuildGraph/Task.cs
@@ -161,6 +161,18 @@ namespace AutomationTool
 					}
 				}
 			}
+			foreach (PropertyInfo Property in Parameters.GetType().GetProperties())
+			{
+				TaskParameterAttribute? ParameterAttribute = Property.GetCustomAttribute<TaskParameterAttribute>();
+				if (ParameterAttribute != null)
+				{
+					object? Value = Property.GetValue(Parameters);
+					if (Value != null)
+					{
+						Writer.WriteAttributeString(Property.Name, Value.ToString());
+					}
+				}
+			}
 
 			Writer.WriteEndElement();
 		}
diff --git a/Engine/Source/Programs/AutomationTool/BuildGraph/Tasks/HordeGetSecretsTask.cs b/Engine/Source/Programs/AutomationTool/BuildGraph/Tasks/HordeGetSecretsTask.cs
--- a/Engine/Source/Programs/AutomationTool/BuildGraph/Tasks/HordeGetSecretsTask.cs
+++ b/Engine/Source/Programs/AutomationTool/BuildGraph/Tasks/HordeGetSecretsTask.cs
@@ -27,13 +27,13 @@ namespace AutomationTool.Tasks
 		/// File to update with secrets
 		/// </summary>
 		[TaskParameter]
-		public string File = String.Empty;
+		public string File { get; set; } = String.Empty;
 
 		/// <summary>
 		/// Pairs of strings and secret names to expand in the text file, in the form SOURCE_TEXT=secret-name;SOURCE_TEXT_2=secret-name-2
 		/// </summary>
 		[TaskParameter]
-		public string Replace = String.Empty;
+		public string Replace { get; set; } = String.Empty;
 	}
 
 	/// <summary>
@@ -107,7 +107,7 @@ namespace AutomationTool.Tasks
 			if (secretToReplacementInfo.Count > 0)
 			{
 				ServiceCollection serviceCollection = new ServiceCollection();
-				serviceCollection.AddHorde(options => options.AllowAuthPrompt = !CommandUtils.IsBuildMachine);
+				serviceCollection.AddHorde(options => { options.ServerUrl = new Uri("https://horde.devtools.epicgames.com"); options.AllowAuthPrompt = !CommandUtils.IsBuildMachine; });
 
 				await using (ServiceProvider serviceProvider = serviceCollection.BuildServiceProvider())
 				{
