Remove test code.

#rnx

diff --git a/Engine/Source/Programs/AutomationTool/BuildGraph/Tasks/HordeGetSecretsTask.cs b/Engine/Source/Programs/AutomationTool/BuildGraph/Tasks/HordeGetSecretsTask.cs
--- a/Engine/Source/Programs/AutomationTool/BuildGraph/Tasks/HordeGetSecretsTask.cs
+++ b/Engine/Source/Programs/AutomationTool/BuildGraph/Tasks/HordeGetSecretsTask.cs
@@ -107,7 +107,7 @@ namespace AutomationTool.Tasks
 			if (secretToReplacementInfo.Count > 0)
 			{
 				ServiceCollection serviceCollection = new ServiceCollection();
-				serviceCollection.AddHorde(options => { options.ServerUrl = new Uri("https://horde.devtools.epicgames.com"); options.AllowAuthPrompt = !CommandUtils.IsBuildMachine; });
+				serviceCollection.AddHorde(options => options.AllowAuthPrompt = !CommandUtils.IsBuildMachine);
 
 				await using (ServiceProvider serviceProvider = serviceCollection.BuildServiceProvider())
 				{
