[Gauntlet] TargetDeviceAndroid.IsActivityRunning() will now search for 'ResumedActivity' instead of 'mResumedActivity' when determining if the app is running in the foreground.

#fixVersion 5.4.3

diff --git a/Engine/Source/Programs/AutomationTool/Gauntlet/Platform/Android/Gauntlet.TargetDeviceAndroid.cs b/Engine/Source/Programs/AutomationTool/Gauntlet/Platform/Android/Gauntlet.TargetDeviceAndroid.cs
--- a/Engine/Source/Programs/AutomationTool/Gauntlet/Platform/Android/Gauntlet.TargetDeviceAndroid.cs
+++ b/Engine/Source/Programs/AutomationTool/Gauntlet/Platform/Android/Gauntlet.TargetDeviceAndroid.cs
@@ -1632,7 +1632,7 @@ namespace Gauntlet
 
 			// We have exited if our activity doesn't appear in the activity query or is not the focused activity.
 			bool bActivityPresent = ActivityQuery.Output.Contains(Install.AndroidPackageName);
-			bool bActivityInForeground = ActivityQuery.Output.Contains("mResumedActivity");
+			bool bActivityInForeground = ActivityQuery.Output.Contains("ResumedActivity");
 			bool bHasExited = !bActivityPresent || !bActivityInForeground;
 			if (bHasExited)
 			{
