BuildGraph: Enable warnings for declaring visibile instance fields in BuildGraph.Automation.

diff --git a/Engine/Source/Programs/AutomationTool/BuildGraph/.editorconfig b/Engine/Source/Programs/AutomationTool/BuildGraph/.editorconfig
new file mode 100644
--- /dev/null
+++ b/Engine/Source/Programs/AutomationTool/BuildGraph/.editorconfig
@@ -0,0 +1,9 @@
+# NOTE: Requires **VS2019 16.3** or later
+
+# EditorConfig is awesome:http://EditorConfig.org
+
+# Dotnet code style settings:
+[*.cs]
+
+# CA1051: Do not declare visible instance fields
+dotnet_diagnostic.CA1051.severity = warning
diff --git a/Engine/Source/Programs/AutomationTool/BuildGraph/BgCondition.cs b/Engine/Source/Programs/AutomationTool/BuildGraph/BgCondition.cs
--- a/Engine/Source/Programs/AutomationTool/BuildGraph/BgCondition.cs
+++ b/Engine/Source/Programs/AutomationTool/BuildGraph/BgCondition.cs
@@ -304,7 +304,7 @@ namespace AutomationTool
 				string argument = await EvaluateScalarAsync();
 				result = (argument.Length > 0 && (argument[^1] == Path.DirectorySeparatorChar || argument[^1] == Path.AltDirectorySeparatorChar)) ? "true" : "false";
 			}
-			else if (String.Compare(_tokens[_idx], "Contains", true) == 0 && _tokens[_idx + 1] == "(")
+			else if (String.Equals(_tokens[_idx], "Contains", StringComparison.OrdinalIgnoreCase) && _tokens[_idx + 1] == "(")
 			{
 				// Check a string contains a substring. If a separator is supplied the string is first split
 				_idx++;
@@ -317,7 +317,7 @@ namespace AutomationTool
 
 				result = Contains(arguments.ElementAt(0), arguments.ElementAt(1)) ? "true" : "false";
 			}
-			else if (String.Compare(_tokens[_idx], "ContainsItem", true) == 0 && _tokens[_idx + 1] == "(")
+			else if (String.Equals(_tokens[_idx], "ContainsItem", StringComparison.OrdinalIgnoreCase) && _tokens[_idx + 1] == "(")
 			{
 				// Check a string contains a substring. If a separator is supplied the string is first split
 				_idx++;
diff --git a/Engine/Source/Programs/AutomationTool/BuildGraph/BgSchema.cs b/Engine/Source/Programs/AutomationTool/BuildGraph/BgSchema.cs
--- a/Engine/Source/Programs/AutomationTool/BuildGraph/BgSchema.cs
+++ b/Engine/Source/Programs/AutomationTool/BuildGraph/BgSchema.cs
@@ -195,7 +195,7 @@ namespace AutomationTool
 		/// <summary>
 		/// The inner xml schema
 		/// </summary>
-		public readonly XmlSchema CompiledSchema;
+		public XmlSchema CompiledSchema { get; }
 
 		/// <summary>
 		/// Characters which are not permitted in names.
diff --git a/Engine/Source/Programs/AutomationTool/BuildGraph/BuildGraph.cs b/Engine/Source/Programs/AutomationTool/BuildGraph/BuildGraph.cs
--- a/Engine/Source/Programs/AutomationTool/BuildGraph/BuildGraph.cs
+++ b/Engine/Source/Programs/AutomationTool/BuildGraph/BuildGraph.cs
@@ -20,6 +20,7 @@ using UnrealBuildBase;
 using UnrealBuildTool;
 
 #nullable enable
+#pragma warning disable CA1724
 
 namespace AutomationTool
 {
diff --git a/Engine/Source/Programs/AutomationTool/BuildGraph/Task.cs b/Engine/Source/Programs/AutomationTool/BuildGraph/Task.cs
--- a/Engine/Source/Programs/AutomationTool/BuildGraph/Task.cs
+++ b/Engine/Source/Programs/AutomationTool/BuildGraph/Task.cs
@@ -26,7 +26,7 @@ namespace AutomationTool
 	/// Attribute to mark parameters to a task, which should be read as XML attributes from the script file.
 	/// </summary>
 	[AttributeUsage(AttributeTargets.Field | AttributeTargets.Property)]
-	public class TaskParameterAttribute : Attribute
+	public sealed class TaskParameterAttribute : Attribute
 	{
 		/// <summary>
 		/// Whether the parameter can be omitted
@@ -56,12 +56,12 @@ namespace AutomationTool
 		/// <summary>
 		/// Name of the XML element that can be used to denote this class
 		/// </summary>
-		public string Name;
+		public string Name { get; set; }
 
 		/// <summary>
 		/// Type to be constructed from the deserialized element
 		/// </summary>
-		public Type ParametersType;
+		public Type ParametersType { get; set; }
 
 		/// <summary>
 		/// Constructor
diff --git a/Engine/Source/Programs/AutomationTool/BuildGraph/Tasks/AgeStoreTask.cs b/Engine/Source/Programs/AutomationTool/BuildGraph/Tasks/AgeStoreTask.cs
--- a/Engine/Source/Programs/AutomationTool/BuildGraph/Tasks/AgeStoreTask.cs
+++ b/Engine/Source/Programs/AutomationTool/BuildGraph/Tasks/AgeStoreTask.cs
@@ -24,25 +24,25 @@ namespace AutomationTool.Tasks
 		/// The target platform to age symbols for.
 		/// </summary>
 		[TaskParameter]
-		public UnrealTargetPlatform Platform;
+		public UnrealTargetPlatform Platform { get; set; }
 
 		/// <summary>
 		/// The symbol server directory.
 		/// </summary>
 		[TaskParameter]
-		public string StoreDir;
+		public string StoreDir { get; set; }
 
 		/// <summary>
 		/// Number of days worth of symbols to keep.
 		/// </summary>
 		[TaskParameter]
-		public int Days;
+		public int Days { get; set; }
 
 		/// <summary>
 		/// The root of the build directory to check for existing buildversion named directories.
 		/// </summary>
 		[TaskParameter(Optional = true)]
-		public string BuildDir;
+		public string BuildDir { get; set; }
 
 		/// <summary>
 		/// A substring to match in directory file names before deleting symbols. This allows the "age store" task
@@ -50,7 +50,7 @@ namespace AutomationTool.Tasks
 		/// Specific use of the filter value is determined by the symbol server structure defined by the platform toolchain.
 		/// </summary>
 		[TaskParameter(Optional = true)]
-		public string Filter;
+		public string Filter { get; set; }
 	}
 
 	/// <summary>
diff --git a/Engine/Source/Programs/AutomationTool/BuildGraph/Tasks/AwsAssumeRoleTask.cs b/Engine/Source/Programs/AutomationTool/BuildGraph/Tasks/AwsAssumeRoleTask.cs
--- a/Engine/Source/Programs/AutomationTool/BuildGraph/Tasks/AwsAssumeRoleTask.cs
+++ b/Engine/Source/Programs/AutomationTool/BuildGraph/Tasks/AwsAssumeRoleTask.cs
@@ -22,37 +22,37 @@ namespace AutomationTool.Tasks
 		/// Role to assume
 		/// </summary>
 		[TaskParameter]
-		public string Arn;
+		public string Arn { get; set; }
 
 		/// <summary>
 		/// Name of this session
 		/// </summary>
 		[TaskParameter]
-		public string Session;
+		public string Session { get; set; }
 
 		/// <summary>
 		/// Duration of the token in seconds
 		/// </summary>
 		[TaskParameter(Optional = true)]
-		public int Duration = 1000;
+		public int Duration { get; set; } = 1000;
 
 		/// <summary>
 		/// Environment variables
 		/// </summary>
 		[TaskParameter(Optional = true)]
-		public string Environment;
+		public string Environment { get; set; }
 
 		/// <summary>
 		/// File to read environment variables from
 		/// </summary>
 		[TaskParameter(Optional = true)]
-		public string EnvironmentFile;
+		public string EnvironmentFile { get; set; }
 
 		/// <summary>
 		/// Output file for the new environment
 		/// </summary>
 		[TaskParameter]
-		public string OutputFile;
+		public string OutputFile { get; set; }
 	}
 
 	/// <summary>
diff --git a/Engine/Source/Programs/AutomationTool/BuildGraph/Tasks/AwsEcsDeployTask.cs b/Engine/Source/Programs/AutomationTool/BuildGraph/Tasks/AwsEcsDeployTask.cs
--- a/Engine/Source/Programs/AutomationTool/BuildGraph/Tasks/AwsEcsDeployTask.cs
+++ b/Engine/Source/Programs/AutomationTool/BuildGraph/Tasks/AwsEcsDeployTask.cs
@@ -23,49 +23,49 @@ namespace AutomationTool.Tasks
 		/// Task definition file to use
 		/// </summary>
 		[TaskParameter(Optional = false)]
-		public string TaskDefinitionFile;
+		public string TaskDefinitionFile { get; set; }
 
 		/// <summary>
 		/// Docker image to set in new task definition (will replace %%DOCKER_PATTERN%% with this value)
 		/// </summary>
 		[TaskParameter(Optional = false)]
-		public string DockerImage;
+		public string DockerImage { get; set; }
 
 		/// <summary>
 		/// App version to set in new task definition (will replace %%VERSION%% with this value)
 		/// </summary>
 		[TaskParameter(Optional = true)]
-		public string Version;
+		public string Version { get; set; }
 
 		/// <summary>
 		/// Cluster ARN representing AWS ECS cluster to operate on
 		/// </summary>
 		[TaskParameter(Optional = false)]
-		public string Cluster;
+		public string Cluster { get; set; }
 
 		/// <summary>
 		/// Service name to update and deploy to
 		/// </summary>
 		[TaskParameter(Optional = false)]
-		public string Service;
+		public string Service { get; set; }
 
 		/// <summary>
 		/// Environment variables
 		/// </summary>
 		[TaskParameter(Optional = true)]
-		public string Environment;
+		public string Environment { get; set; }
 
 		/// <summary>
 		/// File to read environment from
 		/// </summary>
 		[TaskParameter(Optional = true)]
-		public string EnvironmentFile;
+		public string EnvironmentFile { get; set; }
 
 		/// <summary>
 		/// Write output to the log
 		/// </summary>
 		[TaskParameter(Optional = true)]
-		public bool LogOutput = false;
+		public bool LogOutput { get; set; } = false;
 	}
 
 	/// <summary>
diff --git a/Engine/Source/Programs/AutomationTool/BuildGraph/Tasks/AwsTask.cs b/Engine/Source/Programs/AutomationTool/BuildGraph/Tasks/AwsTask.cs
--- a/Engine/Source/Programs/AutomationTool/BuildGraph/Tasks/AwsTask.cs
+++ b/Engine/Source/Programs/AutomationTool/BuildGraph/Tasks/AwsTask.cs
@@ -21,25 +21,25 @@ namespace AutomationTool.Tasks
 		/// Arguments for the newly created process.
 		/// </summary>
 		[TaskParameter(Optional = true)]
-		public string Arguments;
+		public string Arguments { get; set; }
 
 		/// <summary>
 		/// Environment variables
 		/// </summary>
 		[TaskParameter(Optional = true)]
-		public string Environment;
+		public string Environment { get; set; }
 
 		/// <summary>
 		/// File to read environment from
 		/// </summary>
 		[TaskParameter(Optional = true)]
-		public string EnvironmentFile;
+		public string EnvironmentFile { get; set; }
 
 		/// <summary>
 		/// Write output to the log
 		/// </summary>
 		[TaskParameter(Optional = true)]
-		public bool LogOutput = false;
+		public bool LogOutput { get; set; } = false;
 	}
 
 	/// <summary>
diff --git a/Engine/Source/Programs/AutomationTool/BuildGraph/Tasks/CheckMarkdownTask.cs b/Engine/Source/Programs/AutomationTool/BuildGraph/Tasks/CheckMarkdownTask.cs
--- a/Engine/Source/Programs/AutomationTool/BuildGraph/Tasks/CheckMarkdownTask.cs
+++ b/Engine/Source/Programs/AutomationTool/BuildGraph/Tasks/CheckMarkdownTask.cs
@@ -22,7 +22,7 @@ namespace AutomationTool.Tasks
 		/// Optional filter to be applied to the list of input files.
 		/// </summary>
 		[TaskParameter(Optional = true, ValidationType = TaskParameterValidationType.FileSpec)]
-		public string Files;
+		public string Files { get; set; }
 	}
 
 	/// <summary>
diff --git a/Engine/Source/Programs/AutomationTool/BuildGraph/Tasks/CommandTask.cs b/Engine/Source/Programs/AutomationTool/BuildGraph/Tasks/CommandTask.cs
--- a/Engine/Source/Programs/AutomationTool/BuildGraph/Tasks/CommandTask.cs
+++ b/Engine/Source/Programs/AutomationTool/BuildGraph/Tasks/CommandTask.cs
@@ -36,19 +36,19 @@ namespace AutomationTool.Tasks
 		/// The command name to execute.
 		/// </summary>
 		[TaskParameter]
-		public string Name;
+		public string Name { get; set; }
 
 		/// <summary>
 		/// Arguments to be passed to the command.
 		/// </summary>
 		[TaskParameter(Optional = true)]
-		public string Arguments;
+		public string Arguments { get; set; }
 
 		/// <summary>
 		/// If non-null, instructs telemetry from the command to be merged into the telemetry for this UAT instance with the given prefix. May be an empty (non-null) string.
 		/// </summary>
 		[TaskParameter(Optional = true)]
-		public string MergeTelemetryWithPrefix;
+		public string MergeTelemetryWithPrefix { get; set; }
 	}
 
 	/// <summary>
diff --git a/Engine/Source/Programs/AutomationTool/BuildGraph/Tasks/CommandletTask.cs b/Engine/Source/Programs/AutomationTool/BuildGraph/Tasks/CommandletTask.cs
--- a/Engine/Source/Programs/AutomationTool/BuildGraph/Tasks/CommandletTask.cs
+++ b/Engine/Source/Programs/AutomationTool/BuildGraph/Tasks/CommandletTask.cs
@@ -23,31 +23,31 @@ namespace AutomationTool.Tasks
 		/// The commandlet name to execute.
 		/// </summary>
 		[TaskParameter]
-		public string Name;
+		public string Name { get; set; }
 
 		/// <summary>
 		/// The project to run the editor with.
 		/// </summary>
 		[TaskParameter(Optional = true, ValidationType = TaskParameterValidationType.FileSpec)]
-		public string Project;
+		public string Project { get; set; }
 
 		/// <summary>
 		/// Arguments to be passed to the commandlet.
 		/// </summary>
 		[TaskParameter(Optional = true)]
-		public string Arguments;
+		public string Arguments { get; set; }
 
 		/// <summary>
 		/// The editor executable to use. Defaults to the development UnrealEditor executable for the current platform.
 		/// </summary>
 		[TaskParameter(Optional = true)]
-		public FileReference EditorExe;
+		public FileReference EditorExe { get; set; }
 
 		/// <summary>
 		/// The minimum exit code, which is treated as an error.
 		/// </summary>
 		[TaskParameter(Optional = true)]
-		public int ErrorLevel = 1;
+		public int ErrorLevel { get; set; } = 1;
 	}
 
 	/// <summary>
diff --git a/Engine/Source/Programs/AutomationTool/BuildGraph/Tasks/CompileTask.cs b/Engine/Source/Programs/AutomationTool/BuildGraph/Tasks/CompileTask.cs
--- a/Engine/Source/Programs/AutomationTool/BuildGraph/Tasks/CompileTask.cs
+++ b/Engine/Source/Programs/AutomationTool/BuildGraph/Tasks/CompileTask.cs
@@ -23,68 +23,68 @@ namespace AutomationTool.Tasks
 		/// The target to compile.
 		/// </summary>
 		[TaskParameter(Optional=true)]
-		public string Target;
+		public string Target { get; set; }
 
 		/// <summary>
 		/// The configuration to compile.
 		/// </summary>
 		[TaskParameter]
-		public UnrealTargetConfiguration Configuration;
+		public UnrealTargetConfiguration Configuration { get; set; }
 
 		/// <summary>
 		/// The platform to compile for.
 		/// </summary>
 		[TaskParameter]
-		public UnrealTargetPlatform Platform;
+		public UnrealTargetPlatform Platform { get; set; }
 
 		/// <summary>
 		/// The project to compile with.
 		/// </summary>
 		[TaskParameter(Optional = true, ValidationType = TaskParameterValidationType.FileSpec)]
-		public string Project;
+		public string Project { get; set; }
 
 		/// <summary>
 		/// Additional arguments for UnrealBuildTool.
 		/// </summary>
 		[TaskParameter(Optional = true)]
-		public string Arguments;
+		public string Arguments { get; set; }
 
 		/// <summary>
 		/// Whether to allow using XGE for compilation.
 		/// </summary>
 		[TaskParameter(Optional = true)]
-		public bool AllowXGE = true;
+		public bool AllowXGE { get; set; } = true;
 
 		/// <summary>
 		/// No longer necessary as UnrealBuildTool is run to compile targets.
 		/// </summary>
 		[TaskParameter(Optional = true)]
 		[Obsolete]
-		public bool AllowParallelExecutor = true;
+		public bool AllowParallelExecutor { get; set; } = true;
 
 		/// <summary>
 		/// Whether to allow UBT to use all available cores, when AllowXGE is disabled.
 		/// </summary>
 		[TaskParameter(Optional = true)]
-		public bool AllowAllCores = false;
+		public bool AllowAllCores { get; set; } = false;
 
 		/// <summary>
 		/// Whether to allow cleaning this target. If unspecified, targets are cleaned if the -Clean argument is passed on the command line.
 		/// </summary>
 		[TaskParameter(Optional = true)]
-		public bool? Clean = null;
+		public bool? Clean { get; set; } = null;
 
 		/// <summary>
 		/// Global flag passed to UBT that can be used to generate target files without fully compiling.
 		/// </summary>
 		[TaskParameter(Optional = true)]
-		public bool SkipBuild = false;
+		public bool SkipBuild { get; set; } = false;
 
 		/// <summary>
 		/// Tag to be applied to build products of this task.
 		/// </summary>
 		[TaskParameter(Optional = true, ValidationType = TaskParameterValidationType.TagList)]
-		public string Tag;
+		public string Tag { get; set; }
 	}
 
 	/// <summary>
@@ -245,12 +245,12 @@ namespace AutomationTool.Tasks
 		/// <summary>
 		/// Parameters for this task
 		/// </summary>
-		public CompileTaskParameters Parameters;
+		public CompileTaskParameters Parameters { get; set; }
 
 		/// <summary>
 		/// Resolved path to Project file
 		/// </summary>
-		public FileReference ProjectFile = null;
+		public FileReference ProjectFile { get; set; } = null;
 
 		/// <summary>
 		/// Construct a compile task
diff --git a/Engine/Source/Programs/AutomationTool/BuildGraph/Tasks/CookTask.cs b/Engine/Source/Programs/AutomationTool/BuildGraph/Tasks/CookTask.cs
--- a/Engine/Source/Programs/AutomationTool/BuildGraph/Tasks/CookTask.cs
+++ b/Engine/Source/Programs/AutomationTool/BuildGraph/Tasks/CookTask.cs
@@ -1,18 +1,14 @@
 // Copyright Epic Games, Inc. All Rights Reserved.
 
-using AutomationTool;
-using EpicGames.BuildGraph;
 using System;
 using System.Collections.Generic;
 using System.IO;
 using System.Linq;
-using System.Text;
 using System.Threading.Tasks;
 using System.Xml;
 using EpicGames.Core;
 using OpenTracing;
 using OpenTracing.Util;
-using UnrealBuildTool;
 
 namespace AutomationTool.Tasks
 {
@@ -25,50 +21,49 @@ namespace AutomationTool.Tasks
 		/// Project file to be cooked.
 		/// </summary>
 		[TaskParameter]
-		public string Project;
+		public string Project { get; set; }
 
 		/// <summary>
 		/// The cook platform to target (for example, Windows).
 		/// </summary>
 		[TaskParameter]
-		public string Platform;
+		public string Platform { get; set; }
 
 		/// <summary>
 		/// List of maps to be cooked, separated by '+' characters.
 		/// </summary>
 		[TaskParameter(Optional = true)]
-		public string Maps;
+		public string Maps { get; set; }
 
 		/// <summary>
 		/// Additional arguments to be passed to the cooker.
 		/// </summary>
 		[TaskParameter(Optional = true)]
-		public bool Versioned = false;
-	
+		public bool Versioned { get; set; } = false;
+
 		/// <summary>
 		/// Additional arguments to be passed to the cooker.
 		/// </summary>
 		[TaskParameter(Optional = true)]
-		public string Arguments = "";
-
+		public string Arguments { get; set; } = "";
 
 		/// <summary>
 		/// Optional path to what editor executable to run for cooking.
 		/// </summary>
 		[TaskParameter(Optional = true)]
-		public string EditorExe = "";
+		public string EditorExe { get; set; } = "";
 
 		/// <summary>
 		/// Whether to tag the output from the cook. Since cooks produce a lot of files, it can be detrimental to spend time tagging them if we don't need them in a dependent node.
 		/// </summary>
 		[TaskParameter(Optional = true)]
-		public bool TagOutput = true;
+		public bool TagOutput { get; set; } = true;
 
 		/// <summary>
 		/// Tag to be applied to build products of this task.
 		/// </summary>
 		[TaskParameter(Optional = true, ValidationType = TaskParameterValidationType.TagList)]
-		public string Tag;
+		public string Tag { get; set; }
 	}
 
 	/// <summary>
@@ -185,4 +180,4 @@ namespace AutomationTool.Tasks
 			return FindTagNamesFromList(Parameters.Tag);
 		}
 	}
-}
+}
\ No newline at end of file
diff --git a/Engine/Source/Programs/AutomationTool/BuildGraph/Tasks/CopyTask.cs b/Engine/Source/Programs/AutomationTool/BuildGraph/Tasks/CopyTask.cs
--- a/Engine/Source/Programs/AutomationTool/BuildGraph/Tasks/CopyTask.cs
+++ b/Engine/Source/Programs/AutomationTool/BuildGraph/Tasks/CopyTask.cs
@@ -27,37 +27,37 @@ namespace AutomationTool.Tasks
 		/// Optional filter to be applied to the list of input files.
 		/// </summary>
 		[TaskParameter(Optional = true, ValidationType = TaskParameterValidationType.FileSpec)]
-		public string Files;
+		public string Files { get; set; }
 
 		/// <summary>
 		/// The pattern(s) to copy from (for example, Engine/*.txt).
 		/// </summary>
 		[TaskParameter(ValidationType = TaskParameterValidationType.FileSpec)]
-		public string From;
+		public string From { get; set; }
 
 		/// <summary>
 		/// The directory to copy to.
 		/// </summary>
 		[TaskParameter(ValidationType = TaskParameterValidationType.FileSpec)]
-		public string To;
+		public string To { get; set; }
 
 		/// <summary>
 		/// Whether or not to overwrite existing files.
 		/// </summary>
 		[TaskParameter(Optional = true)]
-		public bool Overwrite = true;
+		public bool Overwrite { get; set; } = true;
 
 		/// <summary>
 		/// Tag to be applied to build products of this task.
 		/// </summary>
 		[TaskParameter(Optional = true, ValidationType = TaskParameterValidationType.TagList)]
-		public string Tag;
+		public string Tag { get; set; }
 
 		/// <summary>
 		/// Whether or not to throw an error if no files were found to copy
 		/// </summary>
 		[TaskParameter(Optional = true)]
-		public bool ErrorIfNotFound = false;
+		public bool ErrorIfNotFound { get; set; } = false;
 	}
 
 	/// <summary>
diff --git a/Engine/Source/Programs/AutomationTool/BuildGraph/Tasks/CreateArtifactTask.cs b/Engine/Source/Programs/AutomationTool/BuildGraph/Tasks/CreateArtifactTask.cs
--- a/Engine/Source/Programs/AutomationTool/BuildGraph/Tasks/CreateArtifactTask.cs
+++ b/Engine/Source/Programs/AutomationTool/BuildGraph/Tasks/CreateArtifactTask.cs
@@ -39,37 +39,37 @@ namespace AutomationTool.Tasks
 		/// Name for the artifact
 		/// </summary>
 		[TaskParameter]
-		public string Name = String.Empty;
+		public string Name { get; set; } = String.Empty;
 
 		/// <summary>
 		/// Type of the artifact
 		/// </summary>
 		[TaskParameter(Optional = true)]
-		public string Type = "unknown";
+		public string Type { get; set; } = "unknown";
 
 		/// <summary>
 		/// Description for the artifact
 		/// </summary>
 		[TaskParameter(Optional = true)]
-		public string? Description;
+		public string? Description { get; set; }
 
 		/// <summary>
 		/// Base directory to resolve relative paths for input files.
 		/// </summary>
 		[TaskParameter]
-		public string? BaseDir;
+		public string? BaseDir { get; set; }
 
 		/// <summary>
 		/// Changelist number for this artifact
 		/// </summary>
 		[TaskParameter(Optional = true)]
-		public int? Change;
+		public int? Change { get; set; }
 
 		/// <summary>
 		/// Files to be uploaded.
 		/// </summary>
 		[TaskParameter(Optional = true, ValidationType = TaskParameterValidationType.FileSpec)]
-		public string Files = "...";
+		public string Files { get; set; } = "...";
 	}
 
 	/// <summary>
diff --git a/Engine/Source/Programs/AutomationTool/BuildGraph/Tasks/CsCompileTask.cs b/Engine/Source/Programs/AutomationTool/BuildGraph/Tasks/CsCompileTask.cs
--- a/Engine/Source/Programs/AutomationTool/BuildGraph/Tasks/CsCompileTask.cs
+++ b/Engine/Source/Programs/AutomationTool/BuildGraph/Tasks/CsCompileTask.cs
@@ -23,62 +23,62 @@ namespace AutomationTool.Tasks
 		/// The C# project file to compile. Using semicolons, more than one project file can be specified.
 		/// </summary>
 		[TaskParameter]
-		public string Project;
+		public string Project { get; set; }
 
 		/// <summary>
 		/// The configuration to compile.
 		/// </summary>
 		[TaskParameter(Optional = true)]
-		public string Configuration;
+		public string Configuration { get; set; }
 
 		/// <summary>
 		/// The platform to compile.
 		/// </summary>
 		[TaskParameter(Optional = true)]
-		public string Platform;
+		public string Platform { get; set; }
 
 		/// <summary>
 		/// The target to build.
 		/// </summary>
 		[TaskParameter(Optional = true)]
-		public string Target;
+		public string Target { get; set; }
 
 		/// <summary>
 		/// Properties for the command
 		/// </summary>
 		[TaskParameter(Optional = true)]
-		public string Properties;
+		public string Properties { get; set; }
 
 		/// <summary>
 		/// Additional options to pass to the compiler.
 		/// </summary>
 		[TaskParameter(Optional = true)]
-		public string Arguments;
+		public string Arguments { get; set; }
 
 		/// <summary>
 		/// Only enumerate build products -- do not actually compile the projects.
 		/// </summary>
 		[TaskParameter(Optional = true)]
-		public bool EnumerateOnly;
+		public bool EnumerateOnly { get; set; }
 
 		/// <summary>
 		/// Tag to be applied to build products of this task.
 		/// </summary>
 		[TaskParameter(Optional = true, ValidationType = TaskParameterValidationType.TagList)]
-		public string Tag;
+		public string Tag { get; set; }
 
 		/// <summary>
 		/// Tag to be applied to any non-private references the projects have.
 		/// (for example, those that are external and not copied into the output directory).
 		/// </summary>
 		[TaskParameter(Optional = true, ValidationType = TaskParameterValidationType.TagList)]
-		public string TagReferences;
+		public string TagReferences { get; set; }
 
 		/// <summary>
 		/// Whether to use the system toolchain rather than the bundled UE SDK
 		/// </summary>
 		[TaskParameter(Optional = true)]
-		public bool UseSystemCompiler;
+		public bool UseSystemCompiler { get; set; }
 	}
 
 	/// <summary>
diff --git a/Engine/Source/Programs/AutomationTool/BuildGraph/Tasks/DeleteTask.cs b/Engine/Source/Programs/AutomationTool/BuildGraph/Tasks/DeleteTask.cs
--- a/Engine/Source/Programs/AutomationTool/BuildGraph/Tasks/DeleteTask.cs
+++ b/Engine/Source/Programs/AutomationTool/BuildGraph/Tasks/DeleteTask.cs
@@ -27,25 +27,25 @@ namespace AutomationTool.Tasks
 		/// List of file specifications separated by semicolons (for example, *.cpp;Engine/.../*.bat), or the name of a tag set
 		/// </summary>
 		[TaskParameter(Optional = true, ValidationType = TaskParameterValidationType.FileSpec)]
-		public string Files;
+		public string Files { get; set; }
 
 		/// <summary>
 		/// List of directory names
 		/// </summary>
 		[TaskParameter(Optional = true)]
-		public string Directories;
+		public string Directories { get; set; }
 
 		/// <summary>
 		/// Whether to delete empty directories after deleting the files. Defaults to true.
 		/// </summary>
 		[TaskParameter(Optional = true)]
-		public bool DeleteEmptyDirectories = true;
+		public bool DeleteEmptyDirectories { get; set; } = true;
 
 		/// <summary>
 		/// Whether or not to use verbose logging.
 		/// </summary>
 		[TaskParameter(Optional = true)]
-		public bool Verbose = false;
+		public bool Verbose { get; set; } = false;
 	}
 
 	/// <summary>
diff --git a/Engine/Source/Programs/AutomationTool/BuildGraph/Tasks/DeployToolTask.cs b/Engine/Source/Programs/AutomationTool/BuildGraph/Tasks/DeployToolTask.cs
--- a/Engine/Source/Programs/AutomationTool/BuildGraph/Tasks/DeployToolTask.cs
+++ b/Engine/Source/Programs/AutomationTool/BuildGraph/Tasks/DeployToolTask.cs
@@ -38,43 +38,43 @@ namespace AutomationTool.Tasks
 		/// Identifier for the tool
 		/// </summary>
 		[TaskParameter]
-		public string Id = String.Empty;
+		public string Id { get; set; } = String.Empty;
 
 		/// <summary>
 		/// Settings file to use for the deployment. Should be a JSON file containing server name and access token.
 		/// </summary>
 		[TaskParameter(Optional = true)]
-		public string Settings = String.Empty;
+		public string Settings { get; set; } = String.Empty;
 
 		/// <summary>
 		/// Version number for the new tool
 		/// </summary>
 		[TaskParameter]
-		public string Version = String.Empty;
+		public string Version { get; set; } = String.Empty;
 
 		/// <summary>
 		/// Duration over which to roll out the tool, in minutes.
 		/// </summary>
 		[TaskParameter(Optional = true)]
-		public int Duration = 0;
+		public int Duration { get; set; } = 0;
 
 		/// <summary>
 		/// Whether to create the deployment as paused
 		/// </summary>
 		[TaskParameter(Optional = true)]
-		public bool Paused = false;
+		public bool Paused { get; set; } = false;
 
 		/// <summary>
 		/// Zip file containing files to upload
 		/// </summary>
 		[TaskParameter(Optional = true)]
-		public string? File = null!;
+		public string? File { get; set; } = null!;
 
 		/// <summary>
 		/// Directory to upload for the tool
 		/// </summary>
 		[TaskParameter(Optional = true)]
-		public string? Directory = null!;
+		public string? Directory { get; set; } = null!;
 	}
 
 	/// <summary>
diff --git a/Engine/Source/Programs/AutomationTool/BuildGraph/Tasks/DockerBuildTask.cs b/Engine/Source/Programs/AutomationTool/BuildGraph/Tasks/DockerBuildTask.cs
--- a/Engine/Source/Programs/AutomationTool/BuildGraph/Tasks/DockerBuildTask.cs
+++ b/Engine/Source/Programs/AutomationTool/BuildGraph/Tasks/DockerBuildTask.cs
@@ -24,79 +24,79 @@ namespace AutomationTool.Tasks
 		/// Base directory for the build
 		/// </summary>
 		[TaskParameter]
-		public string BaseDir;
+		public string BaseDir { get; set; }
 
 		/// <summary>
 		/// Files to be staged before building the image
 		/// </summary>
 		[TaskParameter]
-		public string Files;
+		public string Files { get; set; }
 
 		/// <summary>
 		/// Path to the Dockerfile. Uses the root of basedir if not specified.
 		/// </summary>
 		[TaskParameter(Optional = true)]
-		public string DockerFile;
+		public string DockerFile { get; set; }
 
 		/// <summary>
 		/// Path to a .dockerignore. Will be copied to basedir if specified.
 		/// </summary>
 		[TaskParameter(Optional = true)]
-		public string DockerIgnoreFile;
+		public string DockerIgnoreFile { get; set; }
 
 		/// <summary>
 		/// Use BuildKit in Docker
 		/// </summary>
 		[TaskParameter(Optional = true)]
-		public bool UseBuildKit;
+		public bool UseBuildKit { get; set; }
 
 		/// <summary>
 		/// Type of progress output (--progress)
 		/// </summary>
 		[TaskParameter(Optional = true)]
-		public string ProgressOutput;
+		public string ProgressOutput { get; set; }
 
 		/// <summary>
 		/// Tag for the image
 		/// </summary>
 		[TaskParameter(Optional = true)]
-		public string Tag;
+		public string Tag { get; set; }
 
 		/// <summary>
 		/// Set the target build stage to build (--target)
 		/// </summary>
 		[TaskParameter(Optional = true)]
-		public string Target;
+		public string Target { get; set; }
 
 		/// <summary>
 		/// Custom output exporter. Requires BuildKit (--output)
 		/// </summary>
 		[TaskParameter(Optional = true)]
-		public string Output;
+		public string Output { get; set; }
 
 		/// <summary>
 		/// Optional arguments
 		/// </summary>
 		[TaskParameter(Optional = true)]
-		public string Arguments;
+		public string Arguments { get; set; }
 
 		/// <summary>
 		/// List of additional directories to overlay into the staged input files. Allows credentials to be staged, etc...
 		/// </summary>
 		[TaskParameter(Optional = true)]
-		public string OverlayDirs;
+		public string OverlayDirs { get; set; }
 
 		/// <summary>
 		/// Environment variables to set
 		/// </summary>
 		[TaskParameter(Optional = true)]
-		public string Environment;
+		public string Environment { get; set; }
 
 		/// <summary>
 		/// File to read environment variables from
 		/// </summary>
 		[TaskParameter(Optional = true)]
-		public string EnvironmentFile;
+		public string EnvironmentFile { get; set; }
 	}
 
 	/// <summary>
diff --git a/Engine/Source/Programs/AutomationTool/BuildGraph/Tasks/DockerComposeDownTask.cs b/Engine/Source/Programs/AutomationTool/BuildGraph/Tasks/DockerComposeDownTask.cs
--- a/Engine/Source/Programs/AutomationTool/BuildGraph/Tasks/DockerComposeDownTask.cs
+++ b/Engine/Source/Programs/AutomationTool/BuildGraph/Tasks/DockerComposeDownTask.cs
@@ -23,13 +23,13 @@ namespace AutomationTool.Tasks
 		/// Path to the docker-compose file
 		/// </summary>
 		[TaskParameter]
-		public string File;
+		public string File { get; set; }
 
 		/// <summary>
 		/// Arguments for the command
 		/// </summary>
 		[TaskParameter(Optional = true)]
-		public string Arguments;
+		public string Arguments { get; set; }
 	}
 
 	/// <summary>
diff --git a/Engine/Source/Programs/AutomationTool/BuildGraph/Tasks/DockerComposeUpTask.cs b/Engine/Source/Programs/AutomationTool/BuildGraph/Tasks/DockerComposeUpTask.cs
--- a/Engine/Source/Programs/AutomationTool/BuildGraph/Tasks/DockerComposeUpTask.cs
+++ b/Engine/Source/Programs/AutomationTool/BuildGraph/Tasks/DockerComposeUpTask.cs
@@ -23,13 +23,13 @@ namespace AutomationTool.Tasks
 		/// Path to the docker-compose file
 		/// </summary>
 		[TaskParameter]
-		public string File;
+		public string File { get; set; }
 
 		/// <summary>
 		/// Arguments for the command
 		/// </summary>
 		[TaskParameter(Optional = true)]
-		public string Arguments;
+		public string Arguments { get; set; }
 	}
 
 	/// <summary>
diff --git a/Engine/Source/Programs/AutomationTool/BuildGraph/Tasks/DockerPushTask.cs b/Engine/Source/Programs/AutomationTool/BuildGraph/Tasks/DockerPushTask.cs
--- a/Engine/Source/Programs/AutomationTool/BuildGraph/Tasks/DockerPushTask.cs
+++ b/Engine/Source/Programs/AutomationTool/BuildGraph/Tasks/DockerPushTask.cs
@@ -23,43 +23,43 @@ namespace AutomationTool.Tasks
 		/// Repository
 		/// </summary>
 		[TaskParameter]
-		public string Repository;
+		public string Repository { get; set; }
 
 		/// <summary>
 		/// Source image to push
 		/// </summary>
 		[TaskParameter]
-		public string Image;
+		public string Image { get; set; }
 
 		/// <summary>
 		/// Name of the target image
 		/// </summary>
 		[TaskParameter(Optional = true)]
-		public string TargetImage;
+		public string TargetImage { get; set; }
 
 		/// <summary>
 		/// Additional environment variables
 		/// </summary>
 		[TaskParameter(Optional = true)]
-		public string Environment;
+		public string Environment { get; set; }
 
 		/// <summary>
 		/// File to read environment from
 		/// </summary>
 		[TaskParameter(Optional = true)]
-		public string EnvironmentFile;
+		public string EnvironmentFile { get; set; }
 
 		/// <summary>
 		/// Whether to login to AWS ECR
 		/// </summary>
 		[TaskParameter(Optional = true)]
-		public bool AwsEcr;
+		public bool AwsEcr { get; set; }
 
 		/// <summary>
 		/// Path to a json file for authentication to the repository for pushing.
 		/// </summary>
 		[TaskParameter(Optional = true)]
-		public string RepositoryAuthFile;
+		public string RepositoryAuthFile { get; set; }
 	}
 
 	/// <summary>
diff --git a/Engine/Source/Programs/AutomationTool/BuildGraph/Tasks/DockerTask.cs b/Engine/Source/Programs/AutomationTool/BuildGraph/Tasks/DockerTask.cs
--- a/Engine/Source/Programs/AutomationTool/BuildGraph/Tasks/DockerTask.cs
+++ b/Engine/Source/Programs/AutomationTool/BuildGraph/Tasks/DockerTask.cs
@@ -20,25 +20,25 @@ namespace AutomationTool.Tasks
 		/// Docker command line arguments
 		/// </summary>
 		[TaskParameter]
-		public string Arguments;
+		public string Arguments { get; set; }
 
 		/// <summary>
 		/// Environment variables to set
 		/// </summary>
 		[TaskParameter(Optional = true)]
-		public string Environment;
+		public string Environment { get; set; }
 
 		/// <summary>
 		/// File to read environment variables from
 		/// </summary>
 		[TaskParameter(Optional = true)]
-		public string EnvironmentFile;
+		public string EnvironmentFile { get; set; }
 
 		/// <summary>
 		/// Base directory for running the command
 		/// </summary>
 		[TaskParameter(Optional = true)]
-		public string WorkingDir;
+		public string WorkingDir { get; set; }
 	}
 
 	/// <summary>
diff --git a/Engine/Source/Programs/AutomationTool/BuildGraph/Tasks/DotNetTask.cs b/Engine/Source/Programs/AutomationTool/BuildGraph/Tasks/DotNetTask.cs
--- a/Engine/Source/Programs/AutomationTool/BuildGraph/Tasks/DotNetTask.cs
+++ b/Engine/Source/Programs/AutomationTool/BuildGraph/Tasks/DotNetTask.cs
@@ -23,37 +23,37 @@ namespace AutomationTool.Tasks
 		/// Docker command line arguments
 		/// </summary>
 		[TaskParameter(Optional = true)]
-		public string Arguments;
+		public string Arguments { get; set; }
 
 		/// <summary>
 		/// Base directory for running the command
 		/// </summary>
 		[TaskParameter(Optional = true)]
-		public string BaseDir;
+		public string BaseDir { get; set; }
 
 		/// <summary>
 		/// Environment variables to set
 		/// </summary>
 		[TaskParameter(Optional = true)]
-		public string Environment;
+		public string Environment { get; set; }
 
 		/// <summary>
 		/// File to read environment variables from
 		/// </summary>
 		[TaskParameter(Optional = true)]
-		public string EnvironmentFile;
+		public string EnvironmentFile { get; set; }
 
 		/// <summary>
 		/// The minimum exit code, which is treated as an error.
 		/// </summary>
 		[TaskParameter(Optional = true)]
-		public int ErrorLevel = 1;
+		public int ErrorLevel { get; set; } = 1;
 
 		/// <summary>
 		/// Override path to dotnet executable
 		/// </summary>
 		[TaskParameter(Optional = true)]
-		public FileReference DotNetPath;
+		public FileReference DotNetPath { get; set; }
 	}
 
 	/// <summary>
diff --git a/Engine/Source/Programs/AutomationTool/BuildGraph/Tasks/FindModifiedFiles.cs b/Engine/Source/Programs/AutomationTool/BuildGraph/Tasks/FindModifiedFiles.cs
--- a/Engine/Source/Programs/AutomationTool/BuildGraph/Tasks/FindModifiedFiles.cs
+++ b/Engine/Source/Programs/AutomationTool/BuildGraph/Tasks/FindModifiedFiles.cs
@@ -25,31 +25,31 @@ namespace AutomationTool.Tasks
 		///  List of file specifications separated by semicolon (default is ...)
 		/// </summary>
 		[TaskParameter(Optional=true, ValidationType = TaskParameterValidationType.FileSpec)]
-		public string Path = "...";
+		public string Path { get; set; } = "...";
 
 		/// <summary>
 		/// The configuration to compile.
 		/// </summary>
 		[TaskParameter(Optional = true)]
-		public int Change;
+		public int Change { get; set; }
 
 		/// <summary>
 		/// The configuration to compile.
 		/// </summary>
 		[TaskParameter(Optional = true)]
-		public int MinChange;
+		public int MinChange { get; set; }
 
 		/// <summary>
 		/// The configuration to compile.
 		/// </summary>
 		[TaskParameter(Optional = true)]
-		public int MaxChange;
+		public int MaxChange { get; set; }
 
 		/// <summary>
 		/// The file to write to
 		/// </summary>
 		[TaskParameter(Optional = true)]
-		public FileReference Output;
+		public FileReference Output { get; set; }
 	}
 
 	/// <summary>
diff --git a/Engine/Source/Programs/AutomationTool/BuildGraph/Tasks/GatherBuildProductsFromFileTask.cs b/Engine/Source/Programs/AutomationTool/BuildGraph/Tasks/GatherBuildProductsFromFileTask.cs
--- a/Engine/Source/Programs/AutomationTool/BuildGraph/Tasks/GatherBuildProductsFromFileTask.cs
+++ b/Engine/Source/Programs/AutomationTool/BuildGraph/Tasks/GatherBuildProductsFromFileTask.cs
@@ -24,7 +24,7 @@ namespace AutomationTool.Tasks
 		/// 
 		/// </summary>
 		[TaskParameter]
-		public string BuildProductsFile;
+		public string BuildProductsFile { get; set; }
 	}
 
 	[TaskElement("GatherBuildProductsFromFile", typeof(GatherBuildProductsFromFileTaskParameters))]
diff --git a/Engine/Source/Programs/AutomationTool/BuildGraph/Tasks/GitCloneTask.cs b/Engine/Source/Programs/AutomationTool/BuildGraph/Tasks/GitCloneTask.cs
--- a/Engine/Source/Programs/AutomationTool/BuildGraph/Tasks/GitCloneTask.cs
+++ b/Engine/Source/Programs/AutomationTool/BuildGraph/Tasks/GitCloneTask.cs
@@ -20,25 +20,25 @@ namespace AutomationTool.Tasks
 		/// Directory for the repository
 		/// </summary>
 		[TaskParameter]
-		public string Dir;
+		public string Dir { get; set; }
 
 		/// <summary>
 		/// The remote to add
 		/// </summary>
 		[TaskParameter(Optional = true)]
-		public string Remote;
+		public string Remote { get; set; }
 
 		/// <summary>
 		/// The branch to check out on the remote
 		/// </summary>
 		[TaskParameter]
-		public string Branch;
+		public string Branch { get; set; }
 
 		/// <summary>
 		/// Configuration file for the repo. This can be used to set up a remote to be fetched and/or provide credentials.
 		/// </summary>
 		[TaskParameter(Optional = true)]
-		public string ConfigFile;
+		public string ConfigFile { get; set; }
 	}
 
 	/// <summary>
diff --git a/Engine/Source/Programs/AutomationTool/BuildGraph/Tasks/GitTask.cs b/Engine/Source/Programs/AutomationTool/BuildGraph/Tasks/GitTask.cs
--- a/Engine/Source/Programs/AutomationTool/BuildGraph/Tasks/GitTask.cs
+++ b/Engine/Source/Programs/AutomationTool/BuildGraph/Tasks/GitTask.cs
@@ -18,19 +18,19 @@ namespace AutomationTool.Tasks
 		/// Git command line arguments
 		/// </summary>
 		[TaskParameter(Optional = true)]
-		public string Arguments;
+		public string Arguments { get; set; }
 
 		/// <summary>
 		/// Base directory for running the command
 		/// </summary>
 		[TaskParameter(Optional = true)]
-		public string BaseDir;
+		public string BaseDir { get; set; }
 
 		/// <summary>
 		/// The minimum exit code, which is treated as an error.
 		/// </summary>
 		[TaskParameter(Optional = true)]
-		public int ErrorLevel = 1;
+		public int ErrorLevel { get; set; } = 1;
 	}
 
 	/// <summary>
diff --git a/Engine/Source/Programs/AutomationTool/BuildGraph/Tasks/HelmTask.cs b/Engine/Source/Programs/AutomationTool/BuildGraph/Tasks/HelmTask.cs
--- a/Engine/Source/Programs/AutomationTool/BuildGraph/Tasks/HelmTask.cs
+++ b/Engine/Source/Programs/AutomationTool/BuildGraph/Tasks/HelmTask.cs
@@ -21,67 +21,67 @@ namespace AutomationTool.Tasks
 		/// Helm command line arguments
 		/// </summary>
 		[TaskParameter]
-		public string Chart;
+		public string Chart { get; set; }
 
 		/// <summary>
 		/// Name of the release
 		/// </summary>
 		[TaskParameter]
-		public string Deployment;
+		public string Deployment { get; set; }
 
 		/// <summary>
 		/// The Kubernetes namespace
 		/// </summary>
 		[TaskParameter(Optional = true)]
-		public string Namespace;
+		public string Namespace { get; set; }
 
 		/// <summary>
 		/// The kubectl context
 		/// </summary>
 		[TaskParameter(Optional = true)]
-		public string KubeContext;
+		public string KubeContext { get; set; }
 
 		/// <summary>
 		/// The kubectl config file to use
 		/// </summary>
 		[TaskParameter(Optional = true)]
-		public string KubeConfig;
+		public string KubeConfig { get; set; }
 
 		/// <summary>
 		/// Values to set for running the chart
 		/// </summary>
 		[TaskParameter(Optional = true)]
-		public string Values;
+		public string Values { get; set; }
 
 		/// <summary>
 		/// Values to set for running the chart
 		/// </summary>
 		[TaskParameter(Optional = true)]
-		public string ValuesFile;
+		public string ValuesFile { get; set; }
 
 		/// <summary>
 		/// Environment variables to set
 		/// </summary>
 		[TaskParameter(Optional = true)]
-		public string Environment;
+		public string Environment { get; set; }
 
 		/// <summary>
 		/// File to parse environment variables from
 		/// </summary>
 		[TaskParameter(Optional = true)]
-		public string EnvironmentFile;
+		public string EnvironmentFile { get; set; }
 
 		/// <summary>
 		/// Additional arguments
 		/// </summary>
 		[TaskParameter(Optional = true)]
-		public string Arguments;
+		public string Arguments { get; set; }
 
 		/// <summary>
 		/// Base directory for running the command
 		/// </summary>
 		[TaskParameter(Optional = true)]
-		public string WorkingDir;
+		public string WorkingDir { get; set; }
 	}
 
 	/// <summary>
diff --git a/Engine/Source/Programs/AutomationTool/BuildGraph/Tasks/HordeCreateReportTask.cs b/Engine/Source/Programs/AutomationTool/BuildGraph/Tasks/HordeCreateReportTask.cs
--- a/Engine/Source/Programs/AutomationTool/BuildGraph/Tasks/HordeCreateReportTask.cs
+++ b/Engine/Source/Programs/AutomationTool/BuildGraph/Tasks/HordeCreateReportTask.cs
@@ -22,25 +22,25 @@ namespace AutomationTool.Tasks
 		/// Name for the report
 		/// </summary>
 		[TaskParameter]
-		public string Name;
+		public string Name { get; set; }
 
 		/// <summary>
 		/// Where to display the report
 		/// </summary>
 		[TaskParameter]
-		public string Scope;
+		public string Scope { get; set; }
 
 		/// <summary>
 		/// Where to show the report
 		/// </summary>
 		[TaskParameter]
-		public string Placement;
+		public string Placement { get; set; }
 
 		/// <summary>
 		/// Text to be displayed
 		/// </summary>
 		[TaskParameter]
-		public string Text;
+		public string Text { get; set; }
 	}
 
 	/// <summary>
diff --git a/Engine/Source/Programs/AutomationTool/BuildGraph/Tasks/KubectlTask.cs b/Engine/Source/Programs/AutomationTool/BuildGraph/Tasks/KubectlTask.cs
--- a/Engine/Source/Programs/AutomationTool/BuildGraph/Tasks/KubectlTask.cs
+++ b/Engine/Source/Programs/AutomationTool/BuildGraph/Tasks/KubectlTask.cs
@@ -19,13 +19,13 @@ namespace AutomationTool.Tasks
 		/// Command line arguments
 		/// </summary>
 		[TaskParameter]
-		public string Arguments;
+		public string Arguments { get; set; }
 
 		/// <summary>
 		/// Base directory for running the command
 		/// </summary>
 		[TaskParameter(Optional = true)]
-		public string BaseDir;
+		public string BaseDir { get; set; }
 	}
 
 	/// <summary>
diff --git a/Engine/Source/Programs/AutomationTool/BuildGraph/Tasks/LogTask.cs b/Engine/Source/Programs/AutomationTool/BuildGraph/Tasks/LogTask.cs
--- a/Engine/Source/Programs/AutomationTool/BuildGraph/Tasks/LogTask.cs
+++ b/Engine/Source/Programs/AutomationTool/BuildGraph/Tasks/LogTask.cs
@@ -27,19 +27,19 @@ namespace AutomationTool.Tasks
 		/// Message to print out.
 		/// </summary>
 		[TaskParameter(Optional = true)]
-		public string Message;
+		public string Message { get; set; }
 
 		/// <summary>
 		/// If specified, causes the given list of files to be printed after the given message.
 		/// </summary>
 		[TaskParameter(Optional = true, ValidationType = TaskParameterValidationType.FileSpec)]
-		public string Files;
+		public string Files { get; set; }
 
 		/// <summary>
 		/// If specified, causes the contents of the given files to be printed out.
 		/// </summary>
 		[TaskParameter(Optional = true)]
-		public bool IncludeContents;
+		public bool IncludeContents { get; set; }
 	}
 
 	/// <summary>
diff --git a/Engine/Source/Programs/AutomationTool/BuildGraph/Tasks/ModifyConfigTask.cs b/Engine/Source/Programs/AutomationTool/BuildGraph/Tasks/ModifyConfigTask.cs
--- a/Engine/Source/Programs/AutomationTool/BuildGraph/Tasks/ModifyConfigTask.cs
+++ b/Engine/Source/Programs/AutomationTool/BuildGraph/Tasks/ModifyConfigTask.cs
@@ -22,31 +22,31 @@ namespace AutomationTool.Tasks
 		/// Path to the config file
 		/// </summary>
 		[TaskParameter(ValidationType = TaskParameterValidationType.FileSpec)]
-		public string File;
+		public string File { get; set; }
 
 		/// <summary>
 		/// The section name to modify
 		/// </summary>
 		[TaskParameter]
-		public string Section;
+		public string Section { get; set; }
 
 		/// <summary>
 		/// The property name to set
 		/// </summary>
 		[TaskParameter]
-		public string Key;
+		public string Key { get; set; }
 
 		/// <summary>
 		/// The property value to set
 		/// </summary>
 		[TaskParameter]
-		public string Value;
+		public string Value { get; set; }
 
 		/// <summary>
 		/// Tag to be applied to the extracted files
 		/// </summary>
 		[TaskParameter(Optional = true, ValidationType = TaskParameterValidationType.TagList)]
-		public string Tag;
+		public string Tag { get; set; }
 	}
 
 	/// <summary>
diff --git a/Engine/Source/Programs/AutomationTool/BuildGraph/Tasks/ModifyJsonValue.cs b/Engine/Source/Programs/AutomationTool/BuildGraph/Tasks/ModifyJsonValue.cs
--- a/Engine/Source/Programs/AutomationTool/BuildGraph/Tasks/ModifyJsonValue.cs
+++ b/Engine/Source/Programs/AutomationTool/BuildGraph/Tasks/ModifyJsonValue.cs
@@ -23,19 +23,19 @@ namespace AutomationTool.Tasks
 		/// json file paths which will be modified
 		/// </summary>
 		[TaskParameter(ValidationType = TaskParameterValidationType.FileSpec)]
-		public string Files;
+		public string Files { get; set; }
 
 		/// <summary>
 		/// json key path to find in each file
 		/// </summary>
 		[TaskParameter(ValidationType = TaskParameterValidationType.Default)]
-		public string KeyPath;
+		public string KeyPath { get; set; }
 
 		/// <summary>
 		/// new value to apply
 		/// </summary>
 		[TaskParameter(ValidationType = TaskParameterValidationType.Default)]
-		public int NewValue;
+		public int NewValue { get; set; }
 	}
 
 	/// <summary>
diff --git a/Engine/Source/Programs/AutomationTool/BuildGraph/Tasks/MoveTask.cs b/Engine/Source/Programs/AutomationTool/BuildGraph/Tasks/MoveTask.cs
--- a/Engine/Source/Programs/AutomationTool/BuildGraph/Tasks/MoveTask.cs
+++ b/Engine/Source/Programs/AutomationTool/BuildGraph/Tasks/MoveTask.cs
@@ -26,37 +26,37 @@ namespace AutomationTool.Tasks
 		/// Optional filter to be applied to the list of input files.
 		/// </summary>
 		[TaskParameter(Optional = true, ValidationType = TaskParameterValidationType.FileSpec)]
-		public string Files;
+		public string Files { get; set; }
 
 		/// <summary>
 		/// The pattern(s) to copy from (for example, Engine/*.txt).
 		/// </summary>
 		[TaskParameter(ValidationType = TaskParameterValidationType.FileSpec)]
-		public string From;
+		public string From { get; set; }
 
 		/// <summary>
 		/// The directory to copy to.
 		/// </summary>
 		[TaskParameter(ValidationType = TaskParameterValidationType.FileSpec)]
-		public string To;
+		public string To { get; set; }
 
 		/// <summary>
 		/// Optionally if files should be overwritten, defaults to false.
 		/// </summary>
 		[TaskParameter(Optional = true)]
-		public bool Overwrite = false;
+		public bool Overwrite { get; set; } = false;
 
 		/// <summary>
 		/// Tag to be applied to build products of this task.
 		/// </summary>
 		[TaskParameter(Optional = true, ValidationType = TaskParameterValidationType.TagList)]
-		public string Tag;
+		public string Tag { get; set; }
 
 		/// <summary>
 		/// Whether or not to throw an error if no files were found to copy
 		/// </summary>
 		[TaskParameter(Optional = true)]
-		public bool ErrorIfNotFound = false;
+		public bool ErrorIfNotFound { get; set; } = false;
 	}
 
 	/// <summary>
diff --git a/Engine/Source/Programs/AutomationTool/BuildGraph/Tasks/MsBuildTask.cs b/Engine/Source/Programs/AutomationTool/BuildGraph/Tasks/MsBuildTask.cs
--- a/Engine/Source/Programs/AutomationTool/BuildGraph/Tasks/MsBuildTask.cs
+++ b/Engine/Source/Programs/AutomationTool/BuildGraph/Tasks/MsBuildTask.cs
@@ -20,31 +20,31 @@ namespace AutomationTool.Tasks
         /// The C# project file to compile. Using semicolons, more than one project file can be specified.
         /// </summary>
         [TaskParameter]
-		public string Project;
+		public string Project { get; set; }
 
 		/// <summary>
 		/// The configuration to compile.
 		/// </summary>
 		[TaskParameter(Optional = true)]
-		public string Configuration;
+		public string Configuration { get; set; }
 
 		/// <summary>
 		/// The platform to compile.
 		/// </summary>
 		[TaskParameter(Optional = true)]
-		public string Platform;
+		public string Platform { get; set; }
 
 		/// <summary>
 		/// Additional options to pass to the compiler.
 		/// </summary>
 		[TaskParameter(Optional = true)]
-		public string Arguments;
+		public string Arguments { get; set; }
 
 		/// <summary>
 		/// The MSBuild output verbosity.
 		/// </summary>
 		[TaskParameter(Optional = true)]
-		public string Verbosity = "minimal";
+		public string Verbosity { get; set; } = "minimal";
 	}
 
 	/// <summary>
@@ -137,7 +137,7 @@ namespace AutomationTool.Tasks
 		/// <returns>The tag names which are modified by this task</returns>
 		public override IEnumerable<string> FindProducedTagNames()
 		{
-			return new string[0];
+			return Enumerable.Empty<string>();
 		}
 	}	
 }
diff --git a/Engine/Source/Programs/AutomationTool/BuildGraph/Tasks/NotarizeTask.cs b/Engine/Source/Programs/AutomationTool/BuildGraph/Tasks/NotarizeTask.cs
--- a/Engine/Source/Programs/AutomationTool/BuildGraph/Tasks/NotarizeTask.cs
+++ b/Engine/Source/Programs/AutomationTool/BuildGraph/Tasks/NotarizeTask.cs
@@ -30,31 +30,31 @@ namespace AutomationTool.Tasks
 		/// Path to the dmg to notarize
 		/// </summary>
 		[TaskParameter]
-		public string DmgPath;
+		public string DmgPath { get; set; }
 
 		/// <summary>
 		/// primary bundle ID
 		/// </summary>
 		[TaskParameter]
-		public string BundleID;
+		public string BundleID { get; set; }
 
 		/// <summary>
 		/// Apple ID Username
 		/// </summary>
 		[TaskParameter]
-		public string UserName;
+		public string UserName { get; set; }
 
 		/// <summary>
 		/// The keychain ID
 		/// </summary>
 		[TaskParameter]
-		public string KeyChainID;
+		public string KeyChainID { get; set; }
 
 		/// <summary>
 		/// When true the notarization ticket will be stapled
 		/// </summary>
 		[TaskParameter(Optional = true)]
-		public bool RequireStapling = false;
+		public bool RequireStapling { get; set; } = false;
 	}
 
 	[TaskElement("Notarize", typeof(NotarizeTaskParameters))]
diff --git a/Engine/Source/Programs/AutomationTool/BuildGraph/Tasks/NuGetLicenseCheckTask.cs b/Engine/Source/Programs/AutomationTool/BuildGraph/Tasks/NuGetLicenseCheckTask.cs
--- a/Engine/Source/Programs/AutomationTool/BuildGraph/Tasks/NuGetLicenseCheckTask.cs
+++ b/Engine/Source/Programs/AutomationTool/BuildGraph/Tasks/NuGetLicenseCheckTask.cs
@@ -30,31 +30,31 @@ namespace AutomationTool.Tasks
 		/// Base directory for running the command
 		/// </summary>
 		[TaskParameter]
-		public string BaseDir;
+		public string BaseDir { get; set; }
 
 		/// <summary>
 		/// Specifies a list of packages to ignore for version checks, separated by semicolons. Optional version number may be specified with 'name@version' syntax.
 		/// </summary>
 		[TaskParameter(Optional = true)]
-		public string IgnorePackages;
+		public string IgnorePackages { get; set; }
 
 		/// <summary>
 		/// Directory containing allowed licenses
 		/// </summary>
 		[TaskParameter(Optional = true)]
-		public DirectoryReference LicenseDir;
+		public DirectoryReference LicenseDir { get; set; }
 
 		/// <summary>
 		/// Path to a csv file to write with list of packages and their licenses
 		/// </summary>
 		[TaskParameter(Optional = true)]
-		public FileReference CsvFile;
+		public FileReference CsvFile { get; set; }
 
 		/// <summary>
 		/// Override path to dotnet executable
 		/// </summary>
 		[TaskParameter(Optional = true)]
-		public FileReference DotNetPath;
+		public FileReference DotNetPath { get; set; }
 	}
 
 	/// <summary>
diff --git a/Engine/Source/Programs/AutomationTool/BuildGraph/Tasks/OnExitTask.cs b/Engine/Source/Programs/AutomationTool/BuildGraph/Tasks/OnExitTask.cs
--- a/Engine/Source/Programs/AutomationTool/BuildGraph/Tasks/OnExitTask.cs
+++ b/Engine/Source/Programs/AutomationTool/BuildGraph/Tasks/OnExitTask.cs
@@ -19,13 +19,13 @@ namespace AutomationTool.Tasks
 		/// Executable to spawn.
 		/// </summary>
 		[TaskParameter]
-		public string Command = String.Empty;
+		public string Command { get; set; } = String.Empty;
 
 		/// <summary>
 		/// Whether to execute on lease termination
 		/// </summary>
 		[TaskParameter(Optional = true)]
-		public bool Lease = false;
+		public bool Lease { get; set; } = false;
 	}
 
 	/// <summary>
diff --git a/Engine/Source/Programs/AutomationTool/BuildGraph/Tasks/PakFileTask.cs b/Engine/Source/Programs/AutomationTool/BuildGraph/Tasks/PakFileTask.cs
--- a/Engine/Source/Programs/AutomationTool/BuildGraph/Tasks/PakFileTask.cs
+++ b/Engine/Source/Programs/AutomationTool/BuildGraph/Tasks/PakFileTask.cs
@@ -27,55 +27,55 @@ namespace AutomationTool.Tasks
 		/// List of files, wildcards, and tag sets to add to the pak file, separated by ';' characters.
 		/// </summary>
 		[TaskParameter(ValidationType = TaskParameterValidationType.FileSpec)]
-		public string Files;
+		public string Files { get; set; }
 
 		/// <summary>
 		/// PAK file to output.
 		/// </summary>
 		[TaskParameter]
-		public FileReference Output;
+		public FileReference Output { get; set; }
 
 		/// <summary>
 		/// Path to a Response File that contains a list of files to add to the pak file -- instead of specifying them individually.
 		/// </summary>
 		[TaskParameter(Optional = true)]
-		public FileReference ResponseFile;
+		public FileReference ResponseFile { get; set; }
 
 		/// <summary>
 		/// Directories to rebase the files relative to. If specified, the shortest path under a listed directory will be used for each file.
 		/// </summary>
 		[TaskParameter(Optional = true)]
-		public HashSet<DirectoryReference> RebaseDir;
+		public HashSet<DirectoryReference> RebaseDir { get; set; }
 
 		/// <summary>
 		/// Script that gives the order of files.
 		/// </summary>
 		[TaskParameter(Optional = true)]
-		public FileReference Order;
+		public FileReference Order { get; set; }
 
 		/// <summary>
 		/// Encryption keys for this pak file.
 		/// </summary>
 		[TaskParameter(Optional = true)]
-		public FileReference Sign;
+		public FileReference Sign { get; set; }
 
 		/// <summary>
 		/// Whether to compress files.
 		/// </summary>
 		[TaskParameter(Optional = true)]
-		public bool Compress = true;
+		public bool Compress { get; set; } = true;
 
 		/// <summary>
 		/// Additional arguments to pass to UnrealPak.
 		/// </summary>
 		[TaskParameter(Optional = true)]
-		public string Arguments = "";
+		public string Arguments { get; set; } = "";
 
 		/// <summary>
 		/// Tag to be applied to build products of this task.
 		/// </summary>
 		[TaskParameter(Optional = true, ValidationType = TaskParameterValidationType.TagList)]
-		public string Tag;
+		public string Tag { get; set; }
 	}
 
 	/// <summary>
diff --git a/Engine/Source/Programs/AutomationTool/BuildGraph/Tasks/QueryTargetsTask.cs b/Engine/Source/Programs/AutomationTool/BuildGraph/Tasks/QueryTargetsTask.cs
--- a/Engine/Source/Programs/AutomationTool/BuildGraph/Tasks/QueryTargetsTask.cs
+++ b/Engine/Source/Programs/AutomationTool/BuildGraph/Tasks/QueryTargetsTask.cs
@@ -26,25 +26,25 @@ namespace AutomationTool.Tasks
 		/// Path to the project file to query
 		/// </summary>
 		[TaskParameter(Optional = true)]
-		public FileReference? ProjectFile;
+		public FileReference? ProjectFile { get; set; }
 
 		/// <summary>
 		/// Path to the output file to receive information about the targets
 		/// </summary>
 		[TaskParameter(Optional = true)]
-		public FileReference? OutputFile = null;
+		public FileReference? OutputFile { get; set; } = null;
 
 		/// <summary>
 		/// Write out all targets, even if a default is specified in the BuildSettings section of the Default*.ini files. 
 		/// </summary>
 		[TaskParameter(Optional = true)]
-		public bool IncludeAllTargets = false;
+		public bool IncludeAllTargets { get; set; } = false;
 
 		/// <summary>
 		/// Tag to be applied to build products of this task.
 		/// </summary>
 		[TaskParameter(Optional = true, ValidationType = TaskParameterValidationType.TagList)]
-		public string? Tag;
+		public string? Tag { get; set; }
 	}
 
 	/// <summary>
diff --git a/Engine/Source/Programs/AutomationTool/BuildGraph/Tasks/RandomDataTask.cs b/Engine/Source/Programs/AutomationTool/BuildGraph/Tasks/RandomDataTask.cs
--- a/Engine/Source/Programs/AutomationTool/BuildGraph/Tasks/RandomDataTask.cs
+++ b/Engine/Source/Programs/AutomationTool/BuildGraph/Tasks/RandomDataTask.cs
@@ -26,31 +26,31 @@ namespace AutomationTool.Tasks
 		/// The size of each file.
 		/// </summary>
 		[TaskParameter(Optional = true)]
-		public int Size = 1024;
+		public int Size { get; set; } = 1024;
 
 		/// <summary>
 		/// Number of files to write.
 		/// </summary>
 		[TaskParameter(Optional = true)]
-		public int Count = 50;
+		public int Count { get; set; } = 50;
 
 		/// <summary>
 		/// Whether to generate different data for each output file.
 		/// </summary>
 		[TaskParameter(Optional = true)]
-		public bool Different = true;
+		public bool Different { get; set; } = true;
 
 		/// <summary>
 		/// Output directory
 		/// </summary>
 		[TaskParameter(Optional = true)]
-		public string OutputDir;
+		public string OutputDir { get; set; }
 
 		/// <summary>
 		/// Optional filter to be applied to the list of input files.
 		/// </summary>
 		[TaskParameter(Optional = true)]
-		public string Tag;
+		public string Tag { get; set; }
 	}
 
 	/// <summary>
diff --git a/Engine/Source/Programs/AutomationTool/BuildGraph/Tasks/RenameTask.cs b/Engine/Source/Programs/AutomationTool/BuildGraph/Tasks/RenameTask.cs
--- a/Engine/Source/Programs/AutomationTool/BuildGraph/Tasks/RenameTask.cs
+++ b/Engine/Source/Programs/AutomationTool/BuildGraph/Tasks/RenameTask.cs
@@ -25,25 +25,25 @@ namespace AutomationTool.Tasks
 		/// The file or files to rename.
 		/// </summary>
 		[TaskParameter(ValidationType = TaskParameterValidationType.FileSpec)]
-		public string Files;
+		public string Files { get; set; }
 
 		/// <summary>
 		/// The current file name, or pattern to match (for example, *.txt). Should not include any path separators.
 		/// </summary>
 		[TaskParameter(Optional = true)]
-		public string From;
+		public string From { get; set; }
 
 		/// <summary>
 		/// The new name for the file(s). Should not include any path separators.
 		/// </summary>
 		[TaskParameter]
-		public string To;
+		public string To { get; set; }
 
 		/// <summary>
 		/// Tag to be applied to the renamed files.
 		/// </summary>
 		[TaskParameter(Optional = true, ValidationType = TaskParameterValidationType.TagList)]
-		public string Tag;
+		public string Tag { get; set; }
 	}
 
 	/// <summary>
diff --git a/Engine/Source/Programs/AutomationTool/BuildGraph/Tasks/SanitizeReceiptTask.cs b/Engine/Source/Programs/AutomationTool/BuildGraph/Tasks/SanitizeReceiptTask.cs
--- a/Engine/Source/Programs/AutomationTool/BuildGraph/Tasks/SanitizeReceiptTask.cs
+++ b/Engine/Source/Programs/AutomationTool/BuildGraph/Tasks/SanitizeReceiptTask.cs
@@ -25,13 +25,13 @@ namespace AutomationTool.Tasks
 		/// Set of receipt files (*.target) to read, including wildcards and tag names, separated by semicolons.
 		/// </summary>
 		[TaskParameter(ValidationType = TaskParameterValidationType.FileSpec)]
-		public string Files;
+		public string Files { get; set; }
 
 		/// <summary>
 		/// Path to the Engine folder, used to expand $(EngineDir) properties in receipt files. Defaults to the Engine directory for the current workspace.
 		/// </summary>
 		[TaskParameter(Optional = true)]
-		public DirectoryReference EngineDir;
+		public DirectoryReference EngineDir { get; set; }
 	}
 
 	/// <summary>
diff --git a/Engine/Source/Programs/AutomationTool/BuildGraph/Tasks/SetVersionTask.cs b/Engine/Source/Programs/AutomationTool/BuildGraph/Tasks/SetVersionTask.cs
--- a/Engine/Source/Programs/AutomationTool/BuildGraph/Tasks/SetVersionTask.cs
+++ b/Engine/Source/Programs/AutomationTool/BuildGraph/Tasks/SetVersionTask.cs
@@ -22,55 +22,55 @@ namespace AutomationTool.Tasks
 		/// The changelist to set in the version files.
 		/// </summary>
 		[TaskParameter]
-		public int Change;
+		public int Change { get; set; }
 
 		/// <summary>
 		/// The engine compatible changelist to set in the version files.
 		/// </summary>
 		[TaskParameter(Optional = true)]
-		public int CompatibleChange;
+		public int CompatibleChange { get; set; }
 
 		/// <summary>
 		/// The branch string.
 		/// </summary>
 		[TaskParameter]
-		public string Branch;
+		public string Branch { get; set; }
 
 		/// <summary>
 		/// The build version string.
 		/// </summary>
 		[TaskParameter(Optional = true)]
-		public string Build;
+		public string Build { get; set; }
 
 		/// <summary>
 		/// The URL for a running continuous integration job.
 		/// </summary>
 		[TaskParameter(Optional = true)]
-		public string BuildURL;
+		public string BuildURL { get; set; }
 
 		/// <summary>
 		/// Whether to set the IS_LICENSEE_VERSION flag to true.
 		/// </summary>
 		[TaskParameter(Optional = true)]
-		public bool Licensee;
+		public bool Licensee { get; set; }
 
 		/// <summary>
 		/// Whether to set the ENGINE_IS_PROMOTED_BUILD flag to true.
 		/// </summary>
 		[TaskParameter(Optional = true)]
-		public bool Promoted = true;
+		public bool Promoted { get; set; } = true;
 
 		/// <summary>
 		/// If set, do not write to the files -- just return the version files that would be updated. Useful for local builds.
 		/// </summary>
 		[TaskParameter(Optional = true)]
-		public bool SkipWrite;
+		public bool SkipWrite { get; set; }
 
 		/// <summary>
 		/// Tag to be applied to build products of this task.
 		/// </summary>
 		[TaskParameter(Optional = true, ValidationType = TaskParameterValidationType.TagList)]
-		public string Tag;
+		public string Tag { get; set; }
 	}
 
 	/// <summary>
diff --git a/Engine/Source/Programs/AutomationTool/BuildGraph/Tasks/SignTask.cs b/Engine/Source/Programs/AutomationTool/BuildGraph/Tasks/SignTask.cs
--- a/Engine/Source/Programs/AutomationTool/BuildGraph/Tasks/SignTask.cs
+++ b/Engine/Source/Programs/AutomationTool/BuildGraph/Tasks/SignTask.cs
@@ -18,25 +18,25 @@ namespace AutomationTool.Tasks
 		/// List of file specifications separated by semicolons (for example, *.cpp;Engine/.../*.bat), or the name of a tag set.
 		/// </summary>
 		[TaskParameter(ValidationType = TaskParameterValidationType.FileSpec)]
-		public string Files;
+		public string Files { get; set; }
 
 		/// <summary>
 		/// Optional description for the signed content
 		/// </summary>
 		[TaskParameter(Optional = true)]
-		public string Description;
+		public string Description { get; set; }
 
 		/// <summary>
 		/// Tag to be applied to build products of this task.
 		/// </summary>
 		[TaskParameter(Optional = true, ValidationType = TaskParameterValidationType.TagList)]
-		public string Tag;
+		public string Tag { get; set; }
 
 		/// <summary>
 		/// If true, the calls to the signing tool will be performed in parallel.
 		/// </summary>
 		[TaskParameter(Optional = true)]
-		public bool Parallel;
+		public bool Parallel { get; set; }
 	}
 
 	/// <summary>
diff --git a/Engine/Source/Programs/AutomationTool/BuildGraph/Tasks/SpawnTask.cs b/Engine/Source/Programs/AutomationTool/BuildGraph/Tasks/SpawnTask.cs
--- a/Engine/Source/Programs/AutomationTool/BuildGraph/Tasks/SpawnTask.cs
+++ b/Engine/Source/Programs/AutomationTool/BuildGraph/Tasks/SpawnTask.cs
@@ -24,43 +24,43 @@ namespace AutomationTool.Tasks
 		/// Executable to spawn.
 		/// </summary>
 		[TaskParameter]
-		public string Exe;
+		public string Exe { get; set; }
 
 		/// <summary>
 		/// Arguments for the newly created process.
 		/// </summary>
 		[TaskParameter(Optional = true)]
-		public string Arguments;
+		public string Arguments { get; set; }
 
 		/// <summary>
 		/// Working directory for spawning the new task
 		/// </summary>
 		[TaskParameter(Optional = true)]
-		public string WorkingDir;
+		public string WorkingDir { get; set; }
 
 		/// <summary>
 		/// Environment variables to set
 		/// </summary>
 		[TaskParameter(Optional = true)]
-		public string Environment;
+		public string Environment { get; set; }
 
 		/// <summary>
 		/// File to read environment from
 		/// </summary>
 		[TaskParameter(Optional = true)]
-		public string EnvironmentFile;
+		public string EnvironmentFile { get; set; }
 
 		/// <summary>
 		/// Write output to the log
 		/// </summary>
 		[TaskParameter(Optional = true)]
-		public bool LogOutput = true;
+		public bool LogOutput { get; set; } = true;
 
 		/// <summary>
 		/// The minimum exit code, which is treated as an error.
 		/// </summary>
 		[TaskParameter(Optional = true)]
-		public int ErrorLevel = 1;
+		public int ErrorLevel { get; set; } = 1;
 	}
 
 	/// <summary>
diff --git a/Engine/Source/Programs/AutomationTool/BuildGraph/Tasks/StageTask.cs b/Engine/Source/Programs/AutomationTool/BuildGraph/Tasks/StageTask.cs
--- a/Engine/Source/Programs/AutomationTool/BuildGraph/Tasks/StageTask.cs
+++ b/Engine/Source/Programs/AutomationTool/BuildGraph/Tasks/StageTask.cs
@@ -23,49 +23,49 @@ namespace AutomationTool.Tasks
 		/// The project that this target belongs to.
 		/// </summary>
 		[TaskParameter(Optional = true)]
-		public FileReference Project;
+		public FileReference Project { get; set; }
 
 		/// <summary>
 		/// Name of the target to stage.
 		/// </summary>
 		[TaskParameter]
-		public string Target;
+		public string Target { get; set; }
 
 		/// <summary>
 		/// Platform to stage.
 		/// </summary>
 		[TaskParameter]
-		public UnrealTargetPlatform Platform;
+		public UnrealTargetPlatform Platform { get; set; }
 
 		/// <summary>
 		/// Configuration to be staged.
 		/// </summary>
 		[TaskParameter]
-		public UnrealTargetConfiguration Configuration;
+		public UnrealTargetConfiguration Configuration { get; set; }
 
 		/// <summary>
 		/// Architecture to be staged.
 		/// </summary>
 		[TaskParameter(Optional = true)]
-		public string Architecture;
+		public string Architecture { get; set; }
 
 		/// <summary>
 		/// Directory that the receipt files should be staged to.
 		/// </summary>
 		[TaskParameter]
-		public DirectoryReference ToDir;
+		public DirectoryReference ToDir { get; set; }
 
 		/// <summary>
 		/// Whether to overwrite existing files.
 		/// </summary>
 		[TaskParameter(Optional = true)]
-		public bool Overwrite;
+		public bool Overwrite { get; set; }
 
 		/// <summary>
 		/// Tag to be applied to build products of this task.
 		/// </summary>
 		[TaskParameter(Optional = true, ValidationType = TaskParameterValidationType.TagList)]
-		public string Tag;
+		public string Tag { get; set; }
 	}
 
 	/// <summary>
diff --git a/Engine/Source/Programs/AutomationTool/BuildGraph/Tasks/StripTask.cs b/Engine/Source/Programs/AutomationTool/BuildGraph/Tasks/StripTask.cs
--- a/Engine/Source/Programs/AutomationTool/BuildGraph/Tasks/StripTask.cs
+++ b/Engine/Source/Programs/AutomationTool/BuildGraph/Tasks/StripTask.cs
@@ -26,31 +26,31 @@ namespace AutomationTool.Tasks
 		/// The platform toolchain to strip binaries.
 		/// </summary>
 		[TaskParameter]
-		public UnrealTargetPlatform Platform;
+		public UnrealTargetPlatform Platform { get; set; }
 
 		/// <summary>
 		/// The directory to find files in.
 		/// </summary>
 		[TaskParameter(Optional = true)]
-		public DirectoryReference BaseDir;
+		public DirectoryReference BaseDir { get; set; }
 
 		/// <summary>
 		/// List of file specifications separated by semicolons (for example, Engine/.../*.pdb), or the name of a tag set.
 		/// </summary>
 		[TaskParameter(ValidationType = TaskParameterValidationType.FileSpec)]
-		public string Files;
+		public string Files { get; set; }
 
 		/// <summary>
 		/// Output directory for the stripped files. Defaults to the input path, overwriting the input files.
 		/// </summary>
 		[TaskParameter(Optional = true)]
-		public DirectoryReference OutputDir;
+		public DirectoryReference OutputDir { get; set; }
 
 		/// <summary>
 		/// Tag to be applied to build products of this task.
 		/// </summary>
 		[TaskParameter(Optional = true, ValidationType = TaskParameterValidationType.TagList)]
-		public string Tag;
+		public string Tag { get; set; }
 	}
 
 	/// <summary>
diff --git a/Engine/Source/Programs/AutomationTool/BuildGraph/Tasks/SubmitTask.cs b/Engine/Source/Programs/AutomationTool/BuildGraph/Tasks/SubmitTask.cs
--- a/Engine/Source/Programs/AutomationTool/BuildGraph/Tasks/SubmitTask.cs
+++ b/Engine/Source/Programs/AutomationTool/BuildGraph/Tasks/SubmitTask.cs
@@ -26,61 +26,61 @@ namespace AutomationTool.Tasks
 		/// The description for the submitted changelist.
 		/// </summary>
 		[TaskParameter]
-		public string Description;
+		public string Description { get; set; }
 
 		/// <summary>
 		/// The files to submit.
 		/// </summary>
 		[TaskParameter(ValidationType = TaskParameterValidationType.FileSpec)]
-		public string Files;
+		public string Files { get; set; }
 
 		/// <summary>
 		/// The Perforce file type for the submitted files (for example, binary+FS32).
 		/// </summary>
 		[TaskParameter(Optional = true)]
-		public string FileType;
+		public string FileType { get; set; }
 
 		/// <summary>
 		/// The workspace name. If specified, a new workspace will be created using the given stream and root directory to submit the files. If not, the current workspace will be used.
 		/// </summary>
 		[TaskParameter(Optional=true)]
-		public string Workspace;
+		public string Workspace { get; set; }
 
 		/// <summary>
 		/// The stream for the workspace -- defaults to the current stream. Ignored unless the Workspace attribute is also specified.
 		/// </summary>
 		[TaskParameter(Optional = true)]
-		public string Stream;
+		public string Stream { get; set; }
 
 		/// <summary>
 		/// Branch for the workspace (legacy P4 depot path). May not be used in conjunction with Stream.
 		/// </summary>
 		[TaskParameter(Optional = true)]
-		public string Branch;
+		public string Branch { get; set; }
 
 		/// <summary>
 		/// Root directory for the stream. If not specified, defaults to the current root directory.
 		/// </summary>
 		[TaskParameter(Optional = true)]
-		public DirectoryReference RootDir;
+		public DirectoryReference RootDir { get; set; }
 
 		/// <summary>
 		/// Whether to revert unchanged files before attempting to submit.
 		/// </summary>
 		[TaskParameter(Optional = true)]
-		public bool RevertUnchanged;
+		public bool RevertUnchanged { get; set; }
 
 		/// <summary>
 		/// Force the submit to happen -- even if a resolve is needed (always accept current version).
 		/// </summary>
 		[TaskParameter(Optional = true)]
-		public bool Force;
+		public bool Force { get; set; }
 
 		/// <summary>
 		/// Allow verbose P4 output (spew).
 		/// </summary>
 		[TaskParameter(Optional = true)]
-		public bool P4Verbose;
+		public bool P4Verbose { get; set; }
 	}
 
 	/// <summary>
diff --git a/Engine/Source/Programs/AutomationTool/BuildGraph/Tasks/SymStoreTask.cs b/Engine/Source/Programs/AutomationTool/BuildGraph/Tasks/SymStoreTask.cs
--- a/Engine/Source/Programs/AutomationTool/BuildGraph/Tasks/SymStoreTask.cs
+++ b/Engine/Source/Programs/AutomationTool/BuildGraph/Tasks/SymStoreTask.cs
@@ -26,45 +26,45 @@ namespace AutomationTool.Tasks
         /// The platform toolchain required to handle symbol files.
         /// </summary>
         [TaskParameter]
-        public UnrealTargetPlatform Platform;
+        public UnrealTargetPlatform Platform { get; set; }
 
-        /// <summary>
-        /// List of output files. PDBs will be extracted from this list.
-        /// </summary>
-        [TaskParameter]
-        public string Files;
+		/// <summary>
+		/// List of output files. PDBs will be extracted from this list.
+		/// </summary>
+		[TaskParameter]
+        public string Files { get; set; }
 
-        /// <summary>
-        /// Output directory for the compressed symbols.
-        /// </summary>
-        [TaskParameter]
-        public string StoreDir;
+		/// <summary>
+		/// Output directory for the compressed symbols.
+		/// </summary>
+		[TaskParameter]
+        public string StoreDir { get; set; }
 
-        /// <summary>
-        /// Name of the product for the symbol store records.
-        /// </summary>
-        [TaskParameter]
-        public string Product;
+		/// <summary>
+		/// Name of the product for the symbol store records.
+		/// </summary>
+		[TaskParameter]
+        public string Product { get; set; }
 
 		/// <summary>
 		/// Name of the Branch to base all the depot source files from.
 		/// Used when IndexSources is true (may be used only on some platforms).
 		/// </summary>
 		[TaskParameter(Optional = true)]
-		public string Branch;
+		public string Branch { get; set; }
 
 		/// <summary>
 		/// Changelist to which all the depot source files have been synced to.
 		/// Used when IndexSources is true (may be used only on some platforms).
 		/// </summary>
 		[TaskParameter(Optional = true)]
-		public int Change;
+		public int Change { get; set; }
 
 		/// <summary>
 		/// BuildVersion associated with these symbols. Used for clean-up in AgeStore by matching this version against a directory name in a build share.
 		/// </summary>
 		[TaskParameter(Optional = true)]
-		public string BuildVersion;
+		public string BuildVersion { get; set; }
 
 		/// <summary>
 		/// Whether to include the source code index in the uploaded symbols.
@@ -72,7 +72,7 @@ namespace AutomationTool.Tasks
 		/// The source server allows debuggers to automatically fetch the matching source code when debbugging builds or analyzing dumps.
 		/// </summary>
 		[TaskParameter(Optional = true)]
-		public bool IndexSources = false;
+		public bool IndexSources { get; set; } = false;
 
 		/// <summary>
 		/// Filter for the depot source files that are to be indexed.
@@ -81,7 +81,7 @@ namespace AutomationTool.Tasks
 		/// Used when IndexSources is true (may be used only on some platforms).
 		/// </summary>
 		[TaskParameter(Optional = true)]
-		public string SourceFiles;
+		public string SourceFiles { get; set; }
 	}
 
     /// <summary>
diff --git a/Engine/Source/Programs/AutomationTool/BuildGraph/Tasks/TagReceiptTask.cs b/Engine/Source/Programs/AutomationTool/BuildGraph/Tasks/TagReceiptTask.cs
--- a/Engine/Source/Programs/AutomationTool/BuildGraph/Tasks/TagReceiptTask.cs
+++ b/Engine/Source/Programs/AutomationTool/BuildGraph/Tasks/TagReceiptTask.cs
@@ -26,49 +26,49 @@ namespace AutomationTool.Tasks
 		/// Set of receipt files (*.target) to read, including wildcards and tag names, separated by semicolons.
 		/// </summary>
 		[TaskParameter(ValidationType = TaskParameterValidationType.FileSpec)]
-		public string Files;
+		public string Files { get; set; }
 
 		/// <summary>
 		/// Path to the Engine folder, used to expand $(EngineDir) properties in receipt files. Defaults to the Engine directory for the current workspace.
 		/// </summary>
 		[TaskParameter(Optional = true)]
-		public DirectoryReference EngineDir;
+		public DirectoryReference EngineDir { get; set; }
 
 		/// <summary>
 		/// Path to the project folder, used to expand $(ProjectDir) properties in receipt files. Defaults to the Engine directory for the current workspace -- DEPRECATED.
 		/// </summary>
 		[TaskParameter(Optional = true)]
-		public DirectoryReference ProjectDir;
+		public DirectoryReference ProjectDir { get; set; }
 
 		/// <summary>
 		/// Whether to tag the Build Products listed in receipts.
 		/// </summary>
 		[TaskParameter(Optional = true)]
-		public bool BuildProducts;
+		public bool BuildProducts { get; set; }
 
 		/// <summary>
 		/// Which type of Build Products to tag (see TargetReceipt.cs - UnrealBuildTool.BuildProductType for valid values).
 		/// </summary>
 		[TaskParameter(Optional = true)]
-		public string BuildProductType;
+		public string BuildProductType { get; set; }
 
 		/// <summary>
 		/// Whether to tag the Runtime Dependencies listed in receipts.
 		/// </summary>
 		[TaskParameter(Optional = true)]
-		public bool RuntimeDependencies;
+		public bool RuntimeDependencies { get; set; }
 
 		/// <summary>
 		/// Which type of Runtime Dependencies to tag (see TargetReceipt.cs - UnrealBuildTool.StagedFileType for valid values).
 		/// </summary>
 		[TaskParameter(Optional = true)]
-		public string StagedFileType;
+		public string StagedFileType { get; set; }
 
 		/// <summary>
 		/// Name of the tag to apply.
 		/// </summary>
 		[TaskParameter(ValidationType = TaskParameterValidationType.TagList)]
-		public string With;
+		public string With { get; set; }
 	}
 
 	/// <summary>
diff --git a/Engine/Source/Programs/AutomationTool/BuildGraph/Tasks/TagTask.cs b/Engine/Source/Programs/AutomationTool/BuildGraph/Tasks/TagTask.cs
--- a/Engine/Source/Programs/AutomationTool/BuildGraph/Tasks/TagTask.cs
+++ b/Engine/Source/Programs/AutomationTool/BuildGraph/Tasks/TagTask.cs
@@ -23,37 +23,37 @@ namespace AutomationTool.Tasks
 		/// Set the base directory to resolve relative paths and patterns against. If set, any absolute patterns (for example, /Engine/Build/...) are taken to be relative to this path. If not, they are taken to be truly absolute.
 		/// </summary>
 		[TaskParameter(Optional = true)]
-		public DirectoryReference BaseDir;
+		public DirectoryReference BaseDir { get; set; }
 
 		/// <summary>
 		/// Set of files to work from, including wildcards and tag names, separated by semicolons. If set, resolved relative to BaseDir, otherwise resolved to the branch root directory.
 		/// </summary>
 		[TaskParameter(Optional = true, ValidationType = TaskParameterValidationType.FileSpec)]
-		public string Files;
+		public string Files { get; set; }
 
 		/// <summary>
 		/// Set of text files to add additional files from. Each file list should have one file per line.
 		/// </summary>
 		[TaskParameter(Optional = true, ValidationType = TaskParameterValidationType.FileSpec)]
-		public string FileLists;
+		public string FileLists { get; set; }
 
 		/// <summary>
 		/// Patterns to filter the list of files by, including tag names or wildcards. If set, may include patterns that apply to the base directory. If not specified, defaults to all files.
 		/// </summary>
 		[TaskParameter(Optional = true, ValidationType = TaskParameterValidationType.FileSpec)]
-		public string Filter;
+		public string Filter { get; set; }
 
 		/// <summary>
 		/// Set of patterns to exclude from the matched list. May include tag names of patterns that apply to the base directory.
 		/// </summary>
 		[TaskParameter(Optional = true, ValidationType = TaskParameterValidationType.FileSpec)]
-		public string Except;
+		public string Except { get; set; }
 
 		/// <summary>
 		/// Name of the tag to apply.
 		/// </summary>
 		[TaskParameter(ValidationType = TaskParameterValidationType.TagList)]
-		public string With;
+		public string With { get; set; }
 	}
 
 	/// <summary>
diff --git a/Engine/Source/Programs/AutomationTool/BuildGraph/Tasks/UnzipTask.cs b/Engine/Source/Programs/AutomationTool/BuildGraph/Tasks/UnzipTask.cs
--- a/Engine/Source/Programs/AutomationTool/BuildGraph/Tasks/UnzipTask.cs
+++ b/Engine/Source/Programs/AutomationTool/BuildGraph/Tasks/UnzipTask.cs
@@ -23,31 +23,31 @@ namespace AutomationTool.Tasks
 		/// Path to the zip file to extract.
 		/// </summary>
 		[TaskParameter(ValidationType = TaskParameterValidationType.FileSpec)]
-		public string ZipFile;
+		public string ZipFile { get; set; }
 
 		/// <summary>
 		/// Output directory for the extracted files.
 		/// </summary>
 		[TaskParameter]
-		public DirectoryReference ToDir;
+		public DirectoryReference ToDir { get; set; }
 
 		/// <summary>
 		/// Whether or not to use the legacy unzip code.
 		/// </summary>
 		[TaskParameter(Optional = true)]
-		public bool UseLegacyUnzip = false;
+		public bool UseLegacyUnzip { get; set; } = false;
 
 		/// <summary>
 		/// Whether or not to overwrite files during unzip.
 		/// </summary>
 		[TaskParameter(Optional = true)]
-		public bool OverwriteFiles = true;
+		public bool OverwriteFiles { get; set; } = true;
 
 		/// <summary>
 		/// Tag to be applied to the extracted files.
 		/// </summary>
 		[TaskParameter(Optional = true, ValidationType = TaskParameterValidationType.TagList)]
-		public string Tag;
+		public string Tag { get; set; }
 	}
 
 	/// <summary>
diff --git a/Engine/Source/Programs/AutomationTool/BuildGraph/Tasks/UploadTask.cs b/Engine/Source/Programs/AutomationTool/BuildGraph/Tasks/UploadTask.cs
--- a/Engine/Source/Programs/AutomationTool/BuildGraph/Tasks/UploadTask.cs
+++ b/Engine/Source/Programs/AutomationTool/BuildGraph/Tasks/UploadTask.cs
@@ -26,67 +26,67 @@ namespace AutomationTool.Tasks
 		/// The directory to read compressed files from.
 		/// </summary>
 		[TaskParameter]
-		public DirectoryReference FromDir;
+		public DirectoryReference FromDir { get; set; }
 
 		/// <summary>
 		/// List of file specifications separated by semicolons (for example, *.cpp;Engine/.../*.bat), or the name of a tag set. Relative paths are taken from FromDir.
 		/// </summary>
 		[TaskParameter(Optional = true, ValidationType = TaskParameterValidationType.FileSpec)]
-		public string Files;
+		public string Files { get; set; }
 		
 		/// <summary>
 		/// The jupiter namespace used to upload the build. Used to control who has access to the build.
 		/// </summary>
 		[TaskParameter]
-		public string JupiterNamespace;
+		public string JupiterNamespace { get; set; }
 
 		/// <summary>
 		/// The key of the build as will be used to download the build again. This has to be globally unique for this particular upload.
 		/// </summary>
 		[TaskParameter]
-		public string JupiterKey;
+		public string JupiterKey { get; set; }
 
 		/// <summary>
 		/// The type of archive these files are from, will be added to the metadata
 		/// </summary>
 		[TaskParameter]
-		public string ArchiveType;
+		public string ArchiveType { get; set; }
 
 		/// <summary>
 		/// The name of the project this set of files are associated with, will be added to the metadata
 		/// </summary>
 		[TaskParameter]
-		public string ProjectName;
+		public string ProjectName { get; set; }
 
 		/// <summary>
 		/// The source control branch these files were generated from, will be added to the metadata
 		/// </summary>
 		[TaskParameter]
-		public string Branch;
+		public string Branch { get; set; }
 
 		/// <summary>
 		/// The source control revision these files were generated from, will be added to the metadata
 		/// </summary>
 		[TaskParameter]
-		public string Changelist;
+		public string Changelist { get; set; }
 
 		/// <summary>
 		/// Specify the url to the Jupiter instance to upload to
 		/// </summary>
 		[TaskParameter]
-		public string JupiterUrl;
+		public string JupiterUrl { get; set; }
 
 		/// <summary>
 		/// Semi-colon separated list of '=' separated key value mappings to add to the metadata. E.g. Foo=bar;spam=eggs
 		/// </summary>
 		[TaskParameter(Optional = true)]
-		public string AdditionalMetadata;
+		public string AdditionalMetadata { get; set; }
 
 		/// <summary>
 		/// If enabled file content is not kept in memory, results in lower memory usage but increased io as file contents needs to be read multiple times (for hashing as well as during upload)
 		/// </summary>
 		[TaskParameter(Optional = true)]
-		public bool LimitMemoryUsage = true;
+		public bool LimitMemoryUsage { get; set; } = true;
 	}
 
 	/// <summary>
diff --git a/Engine/Source/Programs/AutomationTool/BuildGraph/Tasks/WaitTask.cs b/Engine/Source/Programs/AutomationTool/BuildGraph/Tasks/WaitTask.cs
--- a/Engine/Source/Programs/AutomationTool/BuildGraph/Tasks/WaitTask.cs
+++ b/Engine/Source/Programs/AutomationTool/BuildGraph/Tasks/WaitTask.cs
@@ -20,7 +20,7 @@ namespace AutomationTool.Tasks
 		/// Number of seconds to wait.
 		/// </summary>
 		[TaskParameter]
-		public int Seconds;
+		public int Seconds { get; set; }
 	}
 
 	/// <summary>
diff --git a/Engine/Source/Programs/AutomationTool/BuildGraph/Tasks/WriteJsonValueTask.cs b/Engine/Source/Programs/AutomationTool/BuildGraph/Tasks/WriteJsonValueTask.cs
--- a/Engine/Source/Programs/AutomationTool/BuildGraph/Tasks/WriteJsonValueTask.cs
+++ b/Engine/Source/Programs/AutomationTool/BuildGraph/Tasks/WriteJsonValueTask.cs
@@ -28,7 +28,7 @@ namespace AutomationTool.Tasks
 		/// Json file(s) which will be modified
 		/// </summary>
 		[TaskParameter(ValidationType = TaskParameterValidationType.FileSpec)]
-		public string File = null!;
+		public string File { get; set; } = null!;
 
 		/// <summary>
 		/// Json element to set in each file. Syntax for this string is a limited subset of JsonPath notation, and may support object properties and
@@ -36,13 +36,13 @@ namespace AutomationTool.Tasks
 		/// an element to the 'bar' array in the 'foo' object).
 		/// </summary>
 		[TaskParameter(ValidationType = TaskParameterValidationType.Default)]
-		public string Key = null!;
+		public string Key { get; set; } = null!;
 
 		/// <summary>
 		/// New value to set. May be any value JSON value (string, array, object, number, boolean or null).
 		/// </summary>
 		[TaskParameter(ValidationType = TaskParameterValidationType.Default)]
-		public string Value = null!;
+		public string Value { get; set; } = null!;
 	}
 
 	/// <summary>
diff --git a/Engine/Source/Programs/AutomationTool/BuildGraph/Tasks/WriteTextFileTask.cs b/Engine/Source/Programs/AutomationTool/BuildGraph/Tasks/WriteTextFileTask.cs
--- a/Engine/Source/Programs/AutomationTool/BuildGraph/Tasks/WriteTextFileTask.cs
+++ b/Engine/Source/Programs/AutomationTool/BuildGraph/Tasks/WriteTextFileTask.cs
@@ -24,31 +24,31 @@ namespace AutomationTool.Tasks
 		/// Path to the file to write.
 		/// </summary>
 		[TaskParameter]
-		public FileReference File;
+		public FileReference File { get; set; }
 
 		/// <summary>
 		/// Optional, whether or not to append to the file rather than overwrite.
 		/// </summary>
 		[TaskParameter(Optional = true)]
-		public bool Append;
+		public bool Append { get; set; }
 
 		/// <summary>
 		/// The text to write to the file.
 		/// </summary>
 		[TaskParameter(Optional = true)]
-		public string Text;
+		public string Text { get; set; }
 
 		/// <summary>
 		/// If specified, causes the given list of files to be printed after the given message.
 		/// </summary>
 		[TaskParameter(Optional = true, ValidationType = TaskParameterValidationType.FileSpec)]
-		public string Files;
+		public string Files { get; set; }
 
 		/// <summary>
 		/// Tag to be applied to build products of this task.
 		/// </summary>
 		[TaskParameter(Optional = true, ValidationType = TaskParameterValidationType.TagList)]
-		public string Tag;
+		public string Tag { get; set; }
 	}
 
 	/// <summary>
diff --git a/Engine/Source/Programs/AutomationTool/BuildGraph/Tasks/ZenExportSnapshotTask.cs b/Engine/Source/Programs/AutomationTool/BuildGraph/Tasks/ZenExportSnapshotTask.cs
--- a/Engine/Source/Programs/AutomationTool/BuildGraph/Tasks/ZenExportSnapshotTask.cs
+++ b/Engine/Source/Programs/AutomationTool/BuildGraph/Tasks/ZenExportSnapshotTask.cs
@@ -51,105 +51,104 @@ namespace AutomationTool.Tasks
 		/// The project from which to export the snapshot
 		/// </summary>
 		[TaskParameter(Optional = true)]
-		public FileReference Project;
+		public FileReference Project { get; set; }
 
 		/// <summary>
 		/// The target platform(s) to export the snapshot for
 		/// </summary>
 		[TaskParameter(Optional = true)]
-		public string Platform;
+		public string Platform { get; set; }
 
 		/// <summary>
 		/// A file to read with information about the snapshot that should be used as a base when exporting this new snapshot
 		/// </summary>
 		[TaskParameter(Optional = true)]
-		public FileReference SnapshotBaseDescriptorFile;
+		public FileReference SnapshotBaseDescriptorFile { get; set; }
 
 		/// <summary>
 		/// A file to create with information about the snapshot that was exported
 		/// </summary>
 		[TaskParameter(Optional = true)]
-		public FileReference SnapshotDescriptorFile;
+		public FileReference SnapshotDescriptorFile { get; set; }
 
 		/// <summary>
 		/// The type of destination to export the snapshot to (cloud, ...)
 		/// </summary>
 		[TaskParameter]
-		public string DestinationStorageType;
+		public string DestinationStorageType { get; set; }
 
 		/// <summary>
 		/// The identifier to use when exporting to a destination
 		/// </summary>
 		[TaskParameter(Optional = true)]
-		public string DestinationIdentifier;
+		public string DestinationIdentifier { get; set; }
 
 		/// <summary>
 		/// The host name to use when exporting to a cloud destination
 		/// </summary>
 		[TaskParameter(Optional = true)]
-		public string DestinationCloudHost;
+		public string DestinationCloudHost { get; set; }
 
 		/// <summary>
 		/// The host name to use when writing a snapshot descriptor for a cloud destination
 		/// </summary>
 		[TaskParameter(Optional = true)]
-		public string SnapshotDescriptorCloudHost;
+		public string SnapshotDescriptorCloudHost { get; set; }
 
 		/// <summary>
 		/// The http version to use when exporting to a cloud destination
 		/// </summary>
 		[TaskParameter(Optional = true)]
-		public string DestinationCloudHttpVersion;
+		public string DestinationCloudHttpVersion { get; set; }
 
 		/// <summary>
 		/// The http version to use when writing a snapshot descriptor for a cloud destination
 		/// </summary>
 		[TaskParameter(Optional = true)]
-		public string SnapshotDescriptorCloudHttpVersion;
+		public string SnapshotDescriptorCloudHttpVersion { get; set; }
 
 		/// <summary>
 		/// The namespace to use when exporting to a cloud destination
 		/// </summary>
 		[TaskParameter(Optional = true)]
-		public string DestinationCloudNamespace;
+		public string DestinationCloudNamespace { get; set; }
 
 		/// <summary>
 		/// A custom bucket name to use when exporting to a cloud destination
 		/// </summary>
 		[TaskParameter(Optional = true)]
-		public string DestinationCloudBucket;
+		public string DestinationCloudBucket { get; set; }
 
 		/// <summary>
 		/// The host name to use when exporting to a zen destination
 		/// </summary>
 		[TaskParameter(Optional = true)]
-		public string DestinationZenHost;
+		public string DestinationZenHost { get; set; }
 
 		/// <summary>
 		/// The directory to use when exporting to a file destination
 		/// </summary>
 		[TaskParameter(Optional = true)]
-		public DirectoryReference DestinationFileDir;
+		public DirectoryReference DestinationFileDir { get; set; }
 
 		/// <summary>
 		/// The filename to use when exporting to a file destination
 		/// </summary>
 		[TaskParameter(Optional = true)]
-		public string DestinationFileName;
+		public string DestinationFileName { get; set; }
 
 		/// <summary>
 		/// Optional. Where to look for the ue.projectstore
 		/// The pattern {Platform} can be used for exporting multiple platforms at once.
 		/// </summary>
 		[TaskParameter(Optional = true)]
-		public string OverridePlatformCookedDir;
+		public string OverridePlatformCookedDir { get; set; }
 
 		/// <summary>
 		/// Optional. Whether to force export of data even if the destination claims to have them.
 		/// </summary>
 		[TaskParameter(Optional = true)]
-		public bool Force = false;
-
+		public bool Force { get; set; } = false;
 	}
 
 	/// <summary>
@@ -384,7 +383,7 @@ namespace AutomationTool.Tasks
 			Writer.WriteObjectEnd();
 		}
 
-		private bool TryRunAndLogWithoutSpew(string App, string CommandLine, bool IgnoreFailure)
+		private static bool TryRunAndLogWithoutSpew(string App, string CommandLine, bool IgnoreFailure)
 		{
 			ProcessResult.SpewFilterCallbackType SilentOutputFilter = new ProcessResult.SpewFilterCallbackType(Line =>
 				{
@@ -422,7 +421,7 @@ namespace AutomationTool.Tasks
 			return false;
 		}
 
-		private string SanitizeBucketName(string InString)
+		private static string SanitizeBucketName(string InString)
 		{
 			return StringId.Sanitize(InString).ToString();
 		}
diff --git a/Engine/Source/Programs/AutomationTool/BuildGraph/Tasks/ZenImportOplogTask.cs b/Engine/Source/Programs/AutomationTool/BuildGraph/Tasks/ZenImportOplogTask.cs
--- a/Engine/Source/Programs/AutomationTool/BuildGraph/Tasks/ZenImportOplogTask.cs
+++ b/Engine/Source/Programs/AutomationTool/BuildGraph/Tasks/ZenImportOplogTask.cs
@@ -24,82 +24,81 @@ namespace AutomationTool.Tasks
 		/// The type of destination to import from to (cloud, file...)
 		/// </summary>
 		[TaskParameter]
-		public string ImportType;
+		public string ImportType { get; set; }
 
 		/// <summary>
 		/// comma separated full path to the oplog dir to import into the local zen server
 		/// Files="Path1,Path2"
 		/// </summary>
 		[TaskParameter(Optional = true)]
-		public string Files;
+		public string Files { get; set; }
 
 		/// <summary>
 		/// The project from which to import for
 		/// </summary>
 		[TaskParameter(Optional = true)]
-		public FileReference Project;
+		public FileReference Project { get; set; }
 
 		/// <summary>
 		/// The name of the newly created Zen Project we will be importing into
 		/// </summary>
 		[TaskParameter(Optional = true)]
-		public string ProjectName;
-
+		public string ProjectName { get; set; }
 
 		/// <summary>
 		/// The target platform to import the snapshot for
 		/// </summary>
 		[TaskParameter(Optional = true)]
-		public string Platform;
+		public string Platform { get; set; }
 
 		/// <summary>
 		/// Root dir for the UE project. Used to derive the Enging folder and the Project folder
 		/// </summary>
 		[TaskParameter(Optional = true)]
-		public string RootDir;
+		public string RootDir { get; set; }
 
 		/// <summary>
 		/// The name of the imported oplog
 		/// </summary>
 		[TaskParameter(Optional = true)]
-		public string OplogName;
+		public string OplogName { get; set; }
 
 		/// <summary>
 		/// The host URL for the zen server we are importing from
 		/// </summary>
 		[TaskParameter(Optional = true)]
-		public string HostName = "localhost";
+		public string HostName { get; set; } = "localhost";
 
 		/// <summary>
 		/// The host port for the zen server we are importing from
 		/// </summary>
 		[TaskParameter(Optional = true)]
-		public string HostPort = "8558";
+		public string HostPort { get; set; } = "8558";
 		
 
 		/// <summary>
 		/// The cloud URL to import from
 		/// </summary>
 		[TaskParameter(Optional = true)]
-		public string CloudURL;
+		public string CloudURL { get; set; }
 
 		/// <summary>
 		/// what namespace to use when importing from cloud
 		/// </summary>
 		[TaskParameter(Optional = true)]
-		public string Namespace;
+		public string Namespace { get; set; }
 
 		/// <summary>
 		/// what bucket to use when importing from cloud
 		/// </summary>
 		[TaskParameter(Optional = true)]
-		public string Bucket;
+		public string Bucket { get; set; }
 
 		/// <summary>
 		/// What key to use when importing from cloud
 		/// </summary>
 		[TaskParameter(Optional = true)]
-		public string Key;
+		public string Key { get; set; }
 	}
 
 	/// <summary>
diff --git a/Engine/Source/Programs/AutomationTool/BuildGraph/Tasks/ZipTask.cs b/Engine/Source/Programs/AutomationTool/BuildGraph/Tasks/ZipTask.cs
--- a/Engine/Source/Programs/AutomationTool/BuildGraph/Tasks/ZipTask.cs
+++ b/Engine/Source/Programs/AutomationTool/BuildGraph/Tasks/ZipTask.cs
@@ -25,31 +25,31 @@ namespace AutomationTool.Tasks
 		/// The directory to read compressed files from.
 		/// </summary>
 		[TaskParameter]
-		public DirectoryReference FromDir;
+		public DirectoryReference FromDir { get; set; }
 
 		/// <summary>
 		/// List of file specifications separated by semicolons (for example, *.cpp;Engine/.../*.bat), or the name of a tag set. Relative paths are taken from FromDir.
 		/// </summary>
 		[TaskParameter(Optional = true, ValidationType = TaskParameterValidationType.FileSpec)]
-		public string Files;
+		public string Files { get; set; }
 
 		/// <summary>
 		/// List of files that should have an executable bit set.
 		/// </summary>
 		[TaskParameter(Optional = true, ValidationType = TaskParameterValidationType.FileSpec)]
-		public string ExecutableFiles;
+		public string ExecutableFiles { get; set; }
 
 		/// <summary>
 		/// The zip file to create.
 		/// </summary>
 		[TaskParameter]
-		public FileReference ZipFile;
+		public FileReference ZipFile { get; set; }
 
 		/// <summary>
 		/// Tag to be applied to the created zip file.
 		/// </summary>
 		[TaskParameter(Optional = true, ValidationType = TaskParameterValidationType.TagList)]
-		public string Tag;
+		public string Tag { get; set; }
 	}
 
 	/// <summary>
diff --git a/Engine/Source/Programs/AutomationTool/BuildGraph/TempStorage.cs b/Engine/Source/Programs/AutomationTool/BuildGraph/TempStorage.cs
--- a/Engine/Source/Programs/AutomationTool/BuildGraph/TempStorage.cs
+++ b/Engine/Source/Programs/AutomationTool/BuildGraph/TempStorage.cs
@@ -35,13 +35,13 @@ namespace AutomationTool
 		/// Name of the node
 		/// </summary>
 		[XmlAttribute]
-		public string NodeName;
+		public string NodeName { get; set; }
 
 		/// <summary>
 		/// Name of the output from this node
 		/// </summary>
 		[XmlAttribute]
-		public string OutputName;
+		public string OutputName { get; set; }
 
 		/// <summary>
 		/// Default constructor, for XML serialization.
@@ -103,25 +103,25 @@ namespace AutomationTool
 		/// The path of the file, relative to the engine root. Stored using forward slashes.
 		/// </summary>
 		[XmlAttribute]
-		public string RelativePath;
+		public string RelativePath { get; set; }
 
 		/// <summary>
 		/// The last modified time of the file, in UTC ticks since the Epoch.
 		/// </summary>
 		[XmlAttribute]
-		public long LastWriteTimeUtcTicks;
+		public long LastWriteTimeUtcTicks { get; set; }
 
 		/// <summary>
 		/// Length of the file
 		/// </summary>
 		[XmlAttribute]
-		public long Length;
+		public long Length { get; set; }
 
 		/// <summary>
 		/// Digest for the file. Not all files are hashed.
 		/// </summary>
 		[XmlAttribute]
-		public string? Digest;
+		public string? Digest { get; set; }
 
 		/// <summary>
 		/// Default constructor, for XML serialization.
@@ -305,13 +305,13 @@ namespace AutomationTool
 		/// Name of this file, including extension
 		/// </summary>
 		[XmlAttribute]
-		public string Name;
+		public string Name { get; set; }
 
 		/// <summary>
 		/// Length of the file in bytes
 		/// </summary>
 		[XmlAttribute]
-		public long Length;
+		public long Length { get; set; }
 
 		/// <summary>
 		/// Default constructor, for XML serialization
@@ -342,14 +342,14 @@ namespace AutomationTool
 		/// </summary>
 		[XmlArray]
 		[XmlArrayItem("File")]
-		public TempStorageFile[] Files;
+		public TempStorageFile[] Files { get; set; }
 
 		/// <summary>
 		/// List of compressed archives containing the given files
 		/// </summary>
 		[XmlArray]
 		[XmlArrayItem("ZipFile")]
-		public TempStorageZipFile[] ZipFiles;
+		public TempStorageZipFile[] ZipFiles { get; set; }
 
 		/// <summary>
 		/// Construct a static Xml serializer to avoid throwing an exception searching for the reflection info at runtime
@@ -430,28 +430,28 @@ namespace AutomationTool
 		/// </summary>
 		[XmlArray]
 		[XmlArrayItem("LocalFile")]
-		public string[] LocalFiles;
+		public string[] LocalFiles { get; set; }
 
 		/// <summary>
 		/// List of files that are in this tag set, but not relative to the root directory
 		/// </summary>
 		[XmlArray]
 		[XmlArrayItem("LocalFile")]
-		public string[] ExternalFiles;
+		public string[] ExternalFiles { get; set; }
 
 		/// <summary>
 		/// List of referenced storage blocks
 		/// </summary>
 		[XmlArray]
 		[XmlArrayItem("Block")]
-		public TempStorageBlock[] Blocks;
+		public TempStorageBlock[] Blocks { get; set; }
 
 		/// <summary>
 		/// List of keys for published artifacts
 		/// </summary>
 		[XmlArray]
 		[XmlArrayItem("ArtifactKey")]
-		public string[] ArtifactKeys;
+		public string[] ArtifactKeys { get; set; }
 
 		/// <summary>
 		/// Construct a static Xml serializer to avoid throwing an exception searching for the reflection info at runtime
