BuildGraph: Enable warnings for string comparisons in BuildGraph.Automation.

#rnx

diff --git a/Engine/Source/Programs/AutomationTool/BuildGraph/.editorconfig b/Engine/Source/Programs/AutomationTool/BuildGraph/.editorconfig
--- a/Engine/Source/Programs/AutomationTool/BuildGraph/.editorconfig
+++ b/Engine/Source/Programs/AutomationTool/BuildGraph/.editorconfig
@@ -7,3 +7,5 @@
 
 # CA1051: Do not declare visible instance fields
 dotnet_diagnostic.CA1051.severity = warning
+
+dotnet_diagnostic.CA1307.severity = warning
\ No newline at end of file
diff --git a/Engine/Source/Programs/AutomationTool/BuildGraph/BgScriptReader.cs b/Engine/Source/Programs/AutomationTool/BuildGraph/BgScriptReader.cs
--- a/Engine/Source/Programs/AutomationTool/BuildGraph/BgScriptReader.cs
+++ b/Engine/Source/Programs/AutomationTool/BuildGraph/BgScriptReader.cs
@@ -611,7 +611,7 @@ namespace AutomationTool
 		protected bool TryGetPropertyValue(string name, out string? value)
 		{
 			int valueLength = 0;
-			if (name.Contains(":"))
+			if (name.Contains(":", StringComparison.Ordinal))
 			{
 				string[] tokens = name.Split(':');
 				name = tokens[0];
@@ -923,7 +923,7 @@ namespace AutomationTool
 						{
 							throw new AutomationException($"String operation 'Replace' requires exactly 2 arguments.");
 						}
-						operationResult = input.Replace(arguments[0], arguments[1]);
+						operationResult = input.Replace(arguments[0], arguments[1], StringComparison.Ordinal);
 						break;
 					case "SplitFirst":
 						if (arguments.Length != 1)
@@ -960,7 +960,7 @@ namespace AutomationTool
 					if (!RuntimePlatform.IsWindows)
 					{
 						// Non-windows platforms don't allow dashes in variable names. The engine platform layer substitutes underscores for them.
-						envVarName = envVarName.Replace("-", "_");
+						envVarName = envVarName.Replace("-", "_", StringComparison.Ordinal);
 					}
 
 					string value = Environment.GetEnvironmentVariable(envVarName) ?? "";
@@ -1175,7 +1175,7 @@ namespace AutomationTool
 					BgLabelDef label;
 
 					// Create the label
-					int slashIdx = labelCategoryName.IndexOf('/');
+					int slashIdx = labelCategoryName.IndexOf('/', StringComparison.Ordinal);
 					if (slashIdx != -1)
 					{
 						label = new BgLabelDef(labelCategoryName.Substring(slashIdx + 1), labelCategoryName.Substring(0, slashIdx), null, null, BgLabelChange.Current);
@@ -1399,7 +1399,7 @@ namespace AutomationTool
 					{
 						LogError(element, "Output tag '{Tag}' is already generated by node '{Name}'", producesName, existingOutput.ProducingNode.Name);
 					}
-					else if (!producesName.StartsWith("#"))
+					else if (!producesName.StartsWith("#", StringComparison.Ordinal))
 					{
 						LogError(element, "Output tag names must begin with a '#' character ('{Name}')", producesName);
 					}
@@ -1977,7 +1977,7 @@ namespace AutomationTool
 					LogError(element, "Consecutive spaces in object name '{Name}'", name);
 					return false;
 				}
-				if (Char.IsControl(name[idx]) || BgScriptSchema.IllegalNameCharacters.IndexOf(name[idx]) != -1)
+				if (Char.IsControl(name[idx]) || BgScriptSchema.IllegalNameCharacters.IndexOf(name[idx], StringComparison.Ordinal) != -1)
 				{
 					LogError(element, "Invalid character in object name '{Name}': '{Character}'", name, name[idx]);
 					return false;
@@ -2060,7 +2060,7 @@ namespace AutomationTool
 			{
 				if (!String.IsNullOrWhiteSpace(pair))
 				{
-					int idx = pair.IndexOf('=');
+					int idx = pair.IndexOf('=', StringComparison.Ordinal);
 					if (idx < 0)
 					{
 						LogError(element, "Invalid annotation '{Pair}'", pair);
diff --git a/Engine/Source/Programs/AutomationTool/BuildGraph/BuildGraph.cs b/Engine/Source/Programs/AutomationTool/BuildGraph/BuildGraph.cs
--- a/Engine/Source/Programs/AutomationTool/BuildGraph/BuildGraph.cs
+++ b/Engine/Source/Programs/AutomationTool/BuildGraph/BuildGraph.cs
@@ -379,7 +379,7 @@ namespace AutomationTool
 				const string SetPrefix = "set:";
 				if (Param.StartsWith(SetPrefix, StringComparison.InvariantCultureIgnoreCase))
 				{
-					int EqualsIdx = Param.IndexOf('=');
+					int EqualsIdx = Param.IndexOf('=', StringComparison.Ordinal);
 					if (EqualsIdx >= 0)
 					{
 						Arguments[Param.Substring(SetPrefix.Length, EqualsIdx - SetPrefix.Length)] = Param.Substring(EqualsIdx + 1);
@@ -393,7 +393,7 @@ namespace AutomationTool
 				const string AppendPrefix = "append:";
 				if (Param.StartsWith(AppendPrefix, StringComparison.InvariantCultureIgnoreCase))
 				{
-					int EqualsIdx = Param.IndexOf('=');
+					int EqualsIdx = Param.IndexOf('=', StringComparison.Ordinal);
 					if (EqualsIdx >= 0)
 					{
 						string Property = Param.Substring(AppendPrefix.Length, EqualsIdx - AppendPrefix.Length);
diff --git a/Engine/Source/Programs/AutomationTool/BuildGraph/Task.cs b/Engine/Source/Programs/AutomationTool/BuildGraph/Task.cs
--- a/Engine/Source/Programs/AutomationTool/BuildGraph/Task.cs
+++ b/Engine/Source/Programs/AutomationTool/BuildGraph/Task.cs
@@ -51,17 +51,17 @@ namespace AutomationTool
 	/// Attribute used to associate an XML element name with a parameter block that can be used to construct tasks
 	/// </summary>
 	[AttributeUsage(AttributeTargets.Class)]
-	public class TaskElementAttribute : Attribute
+	public sealed class TaskElementAttribute : Attribute
 	{
 		/// <summary>
 		/// Name of the XML element that can be used to denote this class
 		/// </summary>
-		public string Name { get; set; }
+		public string Name { get; }
 
 		/// <summary>
 		/// Type to be constructed from the deserialized element
 		/// </summary>
-		public Type ParametersType { get; set; }
+		public Type ParametersType { get; }
 
 		/// <summary>
 		/// Constructor
@@ -328,7 +328,7 @@ namespace AutomationTool
 			}
 
 			// Any spaces should be later than the second char - most likely to be a typo if directly after the # character
-			if (TagName.IndexOf(' ') == 1)
+			if (TagName.IndexOf(' ', StringComparison.Ordinal) == 1)
 			{
 				throw new AutomationException("Tag name '{0}' is not valid - spaces should only be used to separate words", TagName);
 			}
diff --git a/Engine/Source/Programs/AutomationTool/BuildGraph/Tasks/AwsEcsDeployTask.cs b/Engine/Source/Programs/AutomationTool/BuildGraph/Tasks/AwsEcsDeployTask.cs
--- a/Engine/Source/Programs/AutomationTool/BuildGraph/Tasks/AwsEcsDeployTask.cs
+++ b/Engine/Source/Programs/AutomationTool/BuildGraph/Tasks/AwsEcsDeployTask.cs
@@ -97,10 +97,10 @@ namespace AutomationTool.Tasks
 		public override async Task ExecuteAsync(JobContext Job, HashSet<FileReference> BuildProducts, Dictionary<string, HashSet<FileReference>> TagNameToFileSet)
 		{
 			string TaskDefTemplate = File.ReadAllText(ResolveFile(Parameters.TaskDefinitionFile).FullName);
-			string TaskDefRendered = TaskDefTemplate.Replace("%%DOCKER_IMAGE%%", Parameters.DockerImage);
+			string TaskDefRendered = TaskDefTemplate.Replace("%%DOCKER_IMAGE%%", Parameters.DockerImage, StringComparison.Ordinal);
 			if (Parameters.Version != null)
 			{
-				TaskDefRendered = TaskDefRendered.Replace("%%VERSION%%", Parameters.Version);
+				TaskDefRendered = TaskDefRendered.Replace("%%VERSION%%", Parameters.Version, StringComparison.Ordinal);
 			}
 
 			FileReference TempTaskDefFile = FileReference.Combine(Unreal.RootDirectory, "Engine", "Intermediate", "Build", "AwsEcsDeployTaskTemp.json");
diff --git a/Engine/Source/Programs/AutomationTool/BuildGraph/Tasks/CheckMarkdownTask.cs b/Engine/Source/Programs/AutomationTool/BuildGraph/Tasks/CheckMarkdownTask.cs
--- a/Engine/Source/Programs/AutomationTool/BuildGraph/Tasks/CheckMarkdownTask.cs
+++ b/Engine/Source/Programs/AutomationTool/BuildGraph/Tasks/CheckMarkdownTask.cs
@@ -96,9 +96,9 @@ namespace AutomationTool.Tasks
 						string link = match.Groups[2].Value;
 						if (!Regex.IsMatch(link, "^(?:[a-z]+:/)?/"))
 						{
-							if (!link.Contains('\\'))
+							if (!link.Contains('\\', StringComparison.Ordinal))
 							{
-								int hashIdx = link.IndexOf('#');
+								int hashIdx = link.IndexOf('#', StringComparison.Ordinal);
 								if (hashIdx == -1)
 								{
 									link = FileReference.Combine(file.Directory, link).FullName;
diff --git a/Engine/Source/Programs/AutomationTool/BuildGraph/Tasks/CsCompileTask.cs b/Engine/Source/Programs/AutomationTool/BuildGraph/Tasks/CsCompileTask.cs
--- a/Engine/Source/Programs/AutomationTool/BuildGraph/Tasks/CsCompileTask.cs
+++ b/Engine/Source/Programs/AutomationTool/BuildGraph/Tasks/CsCompileTask.cs
@@ -139,7 +139,7 @@ namespace AutomationTool.Tasks
 				{
 					if (!String.IsNullOrWhiteSpace(Property))
 					{
-						int EqualsIdx = Property.IndexOf('=');
+						int EqualsIdx = Property.IndexOf('=', StringComparison.Ordinal);
 						if (EqualsIdx == -1)
 						{
 							Logger.LogWarning("Missing '=' in property assignment");
diff --git a/Engine/Source/Programs/AutomationTool/BuildGraph/Tasks/NuGetLicenseCheckTask.cs b/Engine/Source/Programs/AutomationTool/BuildGraph/Tasks/NuGetLicenseCheckTask.cs
--- a/Engine/Source/Programs/AutomationTool/BuildGraph/Tasks/NuGetLicenseCheckTask.cs
+++ b/Engine/Source/Programs/AutomationTool/BuildGraph/Tasks/NuGetLicenseCheckTask.cs
@@ -157,7 +157,7 @@ namespace AutomationTool.Tasks
 			List<DirectoryReference> NuGetPackageDirs = new List<DirectoryReference>();
 			foreach (string Line in NuGetOutput.Output.Split('\n'))
 			{
-				int ColonIdx = Line.IndexOf(':');
+				int ColonIdx = Line.IndexOf(':', StringComparison.Ordinal);
 				if (ColonIdx != -1)
 				{
 					DirectoryReference NuGetPackageDir = new DirectoryReference(Line.Substring(ColonIdx + 1).Trim());
diff --git a/Engine/Source/Programs/AutomationTool/BuildGraph/Tasks/RenameTask.cs b/Engine/Source/Programs/AutomationTool/BuildGraph/Tasks/RenameTask.cs
--- a/Engine/Source/Programs/AutomationTool/BuildGraph/Tasks/RenameTask.cs
+++ b/Engine/Source/Programs/AutomationTool/BuildGraph/Tasks/RenameTask.cs
@@ -99,8 +99,8 @@ namespace AutomationTool.Tasks
 
 			// Convert the source pattern into a regex
 			string EscapedFromPattern = "^" + Regex.Escape(FromPattern) + "$";
-			EscapedFromPattern = EscapedFromPattern.Replace("\\*", "(.*)");
-			EscapedFromPattern = EscapedFromPattern.Replace("\\?", "(.)");
+			EscapedFromPattern = EscapedFromPattern.Replace("\\*", "(.*)", StringComparison.Ordinal);
+			EscapedFromPattern = EscapedFromPattern.Replace("\\?", "(.)", StringComparison.Ordinal);
 			Regex FromRegex = new Regex(EscapedFromPattern, RegexOptions.IgnoreCase | RegexOptions.CultureInvariant);
 
 			// Split the output pattern into fragments that we can insert captures between
diff --git a/Engine/Source/Programs/AutomationTool/BuildGraph/Tasks/SubmitTask.cs b/Engine/Source/Programs/AutomationTool/BuildGraph/Tasks/SubmitTask.cs
--- a/Engine/Source/Programs/AutomationTool/BuildGraph/Tasks/SubmitTask.cs
+++ b/Engine/Source/Programs/AutomationTool/BuildGraph/Tasks/SubmitTask.cs
@@ -151,7 +151,7 @@ namespace AutomationTool.Tasks
 					}
 
 					// Get the latest version of it
-					int NewCL = SubmitP4.CreateChange(Description: Parameters.Description.Replace("\\n", "\n"));
+					int NewCL = SubmitP4.CreateChange(Description: Parameters.Description.Replace("\\n", "\n", StringComparison.Ordinal));
 					foreach(FileReference File in Files)
 					{
 						SubmitP4.Revert(String.Format("-k \"{0}\"", File.FullName), AllowSpew: Parameters.P4Verbose);
diff --git a/Engine/Source/Programs/AutomationTool/BuildGraph/Tasks/TagTask.cs b/Engine/Source/Programs/AutomationTool/BuildGraph/Tasks/TagTask.cs
--- a/Engine/Source/Programs/AutomationTool/BuildGraph/Tasks/TagTask.cs
+++ b/Engine/Source/Programs/AutomationTool/BuildGraph/Tasks/TagTask.cs
@@ -165,15 +165,15 @@ namespace AutomationTool.Tasks
 				else
 				{
 					// Parse a wildcard filter
-					if(Pattern.StartsWith("..."))
+					if(Pattern.StartsWith("...", StringComparison.Ordinal))
 					{
 						Rules.Add(Pattern);
 					}
-					else if(!Pattern.Contains("/"))
+					else if(!Pattern.Contains("/", StringComparison.Ordinal))
 					{
 						Rules.Add(".../" + Pattern);
 					}
-					else if(!Pattern.StartsWith("/"))
+					else if(!Pattern.StartsWith("/", StringComparison.Ordinal))
 					{
 						Rules.Add(BaseDir.FullName + "/" + Pattern);
 					}
diff --git a/Engine/Source/Programs/AutomationTool/BuildGraph/Tasks/UploadTask.cs b/Engine/Source/Programs/AutomationTool/BuildGraph/Tasks/UploadTask.cs
--- a/Engine/Source/Programs/AutomationTool/BuildGraph/Tasks/UploadTask.cs
+++ b/Engine/Source/Programs/AutomationTool/BuildGraph/Tasks/UploadTask.cs
@@ -150,7 +150,7 @@ namespace AutomationTool.Tasks
 				string[] KV = Parameters.AdditionalMetadata.Split(';');
 				foreach (string Option in KV)
 				{
-					int SeparatorIndex = Option.IndexOf('=');
+					int SeparatorIndex = Option.IndexOf('=', StringComparison.Ordinal);
 					if (SeparatorIndex == -1)
 						continue;
 
diff --git a/Engine/Source/Programs/AutomationTool/BuildGraph/Tasks/ZenExportSnapshotTask.cs b/Engine/Source/Programs/AutomationTool/BuildGraph/Tasks/ZenExportSnapshotTask.cs
--- a/Engine/Source/Programs/AutomationTool/BuildGraph/Tasks/ZenExportSnapshotTask.cs
+++ b/Engine/Source/Programs/AutomationTool/BuildGraph/Tasks/ZenExportSnapshotTask.cs
@@ -397,7 +397,7 @@ namespace AutomationTool.Tasks
 			{
 				if (!IgnoreFailure)
 				{
-					Logger.LogWarning(e.ToString());
+					Logger.LogWarning("{Text}", e.ToString());
 				}
 				return false;
 			}
@@ -686,7 +686,7 @@ namespace AutomationTool.Tasks
 
 			if ((Parameters.SnapshotDescriptorFile != null) && SuccessfullyExportedSources.Any())
 			{
-				if (Parameters.SnapshotDescriptorFile.FullName.Contains("{Platform}"))
+				if (Parameters.SnapshotDescriptorFile.FullName.Contains("{Platform}", StringComparison.OrdinalIgnoreCase))
 				{
 					// Separate descriptor file per platform
 					ExportIndex = 0;
diff --git a/Engine/Source/Programs/AutomationTool/BuildGraph/TempStorage.cs b/Engine/Source/Programs/AutomationTool/BuildGraph/TempStorage.cs
--- a/Engine/Source/Programs/AutomationTool/BuildGraph/TempStorage.cs
+++ b/Engine/Source/Programs/AutomationTool/BuildGraph/TempStorage.cs
@@ -80,7 +80,7 @@ namespace AutomationTool
 		/// <returns>Hash code for the block</returns>
 		public override int GetHashCode()
 		{
-			return ToString().GetHashCode();
+			return ToString().GetHashCode(StringComparison.OrdinalIgnoreCase);
 		}
 
 		/// <summary>
@@ -1339,7 +1339,7 @@ namespace AutomationTool
 			{
 				return true;
 			}
-			if ((FileName.Equals("info.plist", StringComparison.OrdinalIgnoreCase) || FileName.Equals("coderesources", StringComparison.OrdinalIgnoreCase)) && LocalFile.FullName.Contains(".app/"))
+			if ((FileName.Equals("info.plist", StringComparison.OrdinalIgnoreCase) || FileName.Equals("coderesources", StringComparison.OrdinalIgnoreCase)) && LocalFile.FullName.Contains(".app/", StringComparison.OrdinalIgnoreCase))
 			{
 				// xcode can generate plist files and coderesources differently in different stages of compile/cook/stage/package/etc. only allow ones inside a .app bundle
 				return true;
