BuildGraph: Enable warnings for using non-portable string comparisons.

#rnx

diff --git a/Engine/Source/Programs/AutomationTool/BuildGraph/.editorconfig b/Engine/Source/Programs/AutomationTool/BuildGraph/.editorconfig
--- a/Engine/Source/Programs/AutomationTool/BuildGraph/.editorconfig
+++ b/Engine/Source/Programs/AutomationTool/BuildGraph/.editorconfig
@@ -8,4 +8,11 @@
 # CA1051: Do not declare visible instance fields
 dotnet_diagnostic.CA1051.severity = warning
 
-dotnet_diagnostic.CA1307.severity = warning
\ No newline at end of file
+dotnet_diagnostic.CA1307.severity = warning
+
+# CA1309: Use ordinal string comparison
+dotnet_diagnostic.CA1309.severity = warning
+
+# CA1310: Specify StringComparison for correctness
+dotnet_diagnostic.CA1310.severity = warning
+
diff --git a/Engine/Source/Programs/AutomationTool/BuildGraph/BgCondition.cs b/Engine/Source/Programs/AutomationTool/BuildGraph/BgCondition.cs
--- a/Engine/Source/Programs/AutomationTool/BuildGraph/BgCondition.cs
+++ b/Engine/Source/Programs/AutomationTool/BuildGraph/BgCondition.cs
@@ -121,7 +121,7 @@ namespace AutomationTool
 		{
 			// <Condition> Or <Condition> Or...
 			string result = await EvaluateAndAsync();
-			while (String.Compare(_tokens[_idx], "Or", true) == 0)
+			while (String.Equals(_tokens[_idx], "Or", StringComparison.OrdinalIgnoreCase))
 			{
 				// Evaluate this condition. We use a binary OR here, because we want to parse everything rather than short-circuit it.
 				_idx++;
@@ -140,7 +140,7 @@ namespace AutomationTool
 		{
 			// <Condition> And <Condition> And...
 			string result = await EvaluateComparisonAsync();
-			while (String.Compare(_tokens[_idx], "And", true) == 0)
+			while (String.Equals(_tokens[_idx], "And", StringComparison.OrdinalIgnoreCase))
 			{
 				// Evaluate this condition. We use a binary AND here, because we want to parse everything rather than short-circuit it.
 				_idx++;
@@ -172,7 +172,7 @@ namespace AutomationTool
 				_idx++;
 				string lhs = result;
 				string rhs = await EvaluateScalarAsync();
-				result = (String.Compare(lhs, rhs, true) == 0) ? "true" : "false";
+				result = String.Equals(lhs, rhs, StringComparison.OrdinalIgnoreCase) ? "true" : "false";
 			}
 			else if (_tokens[_idx] == "!=")
 			{
@@ -180,7 +180,7 @@ namespace AutomationTool
 				_idx++;
 				string lhs = result;
 				string rhs = await EvaluateScalarAsync();
-				result = (String.Compare(lhs, rhs, true) != 0) ? "true" : "false";
+				result = String.Equals(lhs, rhs, StringComparison.OrdinalIgnoreCase) ? "true" : "false";
 			}
 			else if (_tokens[_idx] == "<")
 			{
@@ -290,14 +290,14 @@ namespace AutomationTool
 				string rhs = await EvaluateScalarAsync();
 				result = CoerceToBool(rhs) ? "false" : "true";
 			}
-			else if (String.Compare(_tokens[_idx], "Exists", true) == 0 && _tokens[_idx + 1] == "(")
+			else if (String.Equals(_tokens[_idx], "Exists", StringComparison.OrdinalIgnoreCase) && _tokens[_idx + 1] == "(")
 			{
 				// Check whether file or directory exists. Evaluate the argument as a subexpression.
 				_idx++;
 				string argument = await EvaluateScalarAsync();
 				result = Exists(argument) ? "true" : "false";
 			}
-			else if (String.Compare(_tokens[_idx], "HasTrailingSlash", true) == 0 && _tokens[_idx + 1] == "(")
+			else if (String.Equals(_tokens[_idx], "HasTrailingSlash", StringComparison.OrdinalIgnoreCase) && _tokens[_idx + 1] == "(")
 			{
 				// Check whether the given string ends with a slash
 				_idx++;
@@ -415,11 +415,11 @@ namespace AutomationTool
 		static bool CoerceToBool(string scalar)
 		{
 			bool result;
-			if (String.Compare(scalar, "true", true) == 0)
+			if (String.Equals(scalar, "true", StringComparison.OrdinalIgnoreCase))
 			{
 				result = true;
 			}
-			else if (String.Compare(scalar, "false", true) == 0)
+			else if (String.Equals(scalar, "false", StringComparison.OrdinalIgnoreCase))
 			{
 				result = false;
 			}
diff --git a/Engine/Source/Programs/AutomationTool/BuildGraph/BgScriptReader.cs b/Engine/Source/Programs/AutomationTool/BuildGraph/BgScriptReader.cs
--- a/Engine/Source/Programs/AutomationTool/BuildGraph/BgScriptReader.cs
+++ b/Engine/Source/Programs/AutomationTool/BuildGraph/BgScriptReader.cs
@@ -1914,7 +1914,7 @@ namespace AutomationTool
 				{
 					nodes.UnionWith(otherNodes);
 				}
-				else if (!referenceName.StartsWith("#") && _graph.TagNameToNodeOutput.ContainsKey("#" + referenceName))
+				else if (!referenceName.StartsWith("#", StringComparison.Ordinal) && _graph.TagNameToNodeOutput.ContainsKey("#" + referenceName))
 				{
 					LogError(element, "Reference to '{Name}' cannot be resolved; did you mean '{PossibleName}'?", referenceName, $"#{referenceName}");
 				}
@@ -1942,7 +1942,7 @@ namespace AutomationTool
 				{
 					inputs.UnionWith(referenceInputs);
 				}
-				else if (!referenceName.StartsWith("#") && _graph.TagNameToNodeOutput.ContainsKey("#" + referenceName))
+				else if (!referenceName.StartsWith("#", StringComparison.Ordinal) && _graph.TagNameToNodeOutput.ContainsKey("#" + referenceName))
 				{
 					LogError(element, "Reference to '{Name}' cannot be resolved; did you mean '{PossibleName}'?", referenceName, $"#{referenceName}");
 				}
@@ -2098,11 +2098,11 @@ namespace AutomationTool
 			if (element.HasAttribute(name))
 			{
 				string value = ReadAttribute(element, name).Trim();
-				if (value.Equals("true", StringComparison.InvariantCultureIgnoreCase))
+				if (value.Equals("true", StringComparison.OrdinalIgnoreCase))
 				{
 					bResult = true;
 				}
-				else if (value.Equals("false", StringComparison.InvariantCultureIgnoreCase))
+				else if (value.Equals("false", StringComparison.OrdinalIgnoreCase))
 				{
 					bResult = false;
 				}
@@ -2229,7 +2229,7 @@ namespace AutomationTool
 		{
 			string result = text;
 			// Iterate in reverse order to handle cases where there are nested expansions like $(Outer$(Inner))
-			for (int idx = result.LastIndexOf("$("); idx != -1; idx = result.LastIndexOf("$(", idx, idx + 1))
+			for (int idx = result.LastIndexOf("$(", StringComparison.Ordinal); idx != -1; idx = result.LastIndexOf("$(", idx, idx + 1, StringComparison.Ordinal))
 			{
 				// Find the end of the variable name
 				int endIdx = result.IndexOf(')', idx + 2);
diff --git a/Engine/Source/Programs/AutomationTool/BuildGraph/BuildGraph.cs b/Engine/Source/Programs/AutomationTool/BuildGraph/BuildGraph.cs
--- a/Engine/Source/Programs/AutomationTool/BuildGraph/BuildGraph.cs
+++ b/Engine/Source/Programs/AutomationTool/BuildGraph/BuildGraph.cs
@@ -943,7 +943,7 @@ namespace AutomationTool
 			if (File.IsUnderDirectory(EngineDirectory))
 			{
 				string[] PathFragments = File.MakeRelativeTo(EngineDirectory).Split(Path.DirectorySeparatorChar, Path.AltDirectorySeparatorChar);
-				if (PathFragments.All(x => !x.Equals("NotForLicensees", StringComparison.InvariantCultureIgnoreCase) && !x.Equals("NoRedist", StringComparison.InvariantCultureIgnoreCase)))
+				if (PathFragments.All(x => !x.Equals("NotForLicensees", StringComparison.InvariantCultureIgnoreCase) && !x.Equals("NoRedist", StringComparison.OrdinalIgnoreCase)))
 				{
 					return true;
 				}
diff --git a/Engine/Source/Programs/AutomationTool/BuildGraph/Task.cs b/Engine/Source/Programs/AutomationTool/BuildGraph/Task.cs
--- a/Engine/Source/Programs/AutomationTool/BuildGraph/Task.cs
+++ b/Engine/Source/Programs/AutomationTool/BuildGraph/Task.cs
@@ -251,7 +251,7 @@ namespace AutomationTool
 			{
 				foreach(string Pattern in SplitDelimitedList(Filespec))
 				{
-					if(Pattern.StartsWith("#"))
+					if(Pattern.StartsWith("#", StringComparison.Ordinal))
 					{
 						yield return Pattern;
 					}
@@ -399,7 +399,7 @@ namespace AutomationTool
 			foreach(string Pattern in FilePatterns)
 			{
 				// Check if it's a tag name
-				if(Pattern.StartsWith("#"))
+				if(Pattern.StartsWith("#", StringComparison.Ordinal))
 				{
 					Files.UnionWith(FindOrAddTagSet(TagNameToFileSet, Pattern));
 					continue;
diff --git a/Engine/Source/Programs/AutomationTool/BuildGraph/Tasks/RenameTask.cs b/Engine/Source/Programs/AutomationTool/BuildGraph/Tasks/RenameTask.cs
--- a/Engine/Source/Programs/AutomationTool/BuildGraph/Tasks/RenameTask.cs
+++ b/Engine/Source/Programs/AutomationTool/BuildGraph/Tasks/RenameTask.cs
@@ -79,7 +79,7 @@ namespace AutomationTool.Tasks
 			if (FromPattern == null)
 			{
 				List<string> Patterns = SplitDelimitedList(Parameters.Files);
-				if (Patterns.Count != 1 || Patterns[0].StartsWith("#"))
+				if (Patterns.Count != 1 || Patterns[0].StartsWith("#", StringComparison.Ordinal))
 				{
 					throw new AutomationException("Missing 'From' attribute specifying pattern to match source files against");
 				}
@@ -91,7 +91,7 @@ namespace AutomationTool.Tasks
 				{
 					FromPattern = FromPattern.Substring(SlashIdx + 1);
 				}
-				if (FromPattern.StartsWith("..."))
+				if (FromPattern.StartsWith("...", StringComparison.Ordinal))
 				{
 					FromPattern = "*" + FromPattern.Substring(3);
 				}
diff --git a/Engine/Source/Programs/AutomationTool/BuildGraph/Tasks/TagTask.cs b/Engine/Source/Programs/AutomationTool/BuildGraph/Tasks/TagTask.cs
--- a/Engine/Source/Programs/AutomationTool/BuildGraph/Tasks/TagTask.cs
+++ b/Engine/Source/Programs/AutomationTool/BuildGraph/Tasks/TagTask.cs
@@ -153,7 +153,7 @@ namespace AutomationTool.Tasks
 			List<string> Rules = new List<string>();
 			foreach(string Pattern in Patterns)
 			{
-				if(Pattern.StartsWith("#"))
+				if(Pattern.StartsWith("#", StringComparison.Ordinal))
 				{
 					// Add the files in a specific set to the filter
 					HashSet<FileReference> Files = FindOrAddTagSet(TagNameToFileSet, Pattern);
diff --git a/Engine/Source/Programs/AutomationTool/BuildGraph/Tasks/ZenExportSnapshotTask.cs b/Engine/Source/Programs/AutomationTool/BuildGraph/Tasks/ZenExportSnapshotTask.cs
--- a/Engine/Source/Programs/AutomationTool/BuildGraph/Tasks/ZenExportSnapshotTask.cs
+++ b/Engine/Source/Programs/AutomationTool/BuildGraph/Tasks/ZenExportSnapshotTask.cs
@@ -354,7 +354,7 @@ namespace AutomationTool.Tasks
 					Writer.WriteValue("type", "cloud");
 					Writer.WriteValue("targetplatform", ExportSource.TargetPlatform);
 					Writer.WriteValue("host", HostName);
-					if (!string.IsNullOrEmpty(HttpVersion) && !HttpVersion.Equals("None", StringComparison.InvariantCultureIgnoreCase))
+					if (!string.IsNullOrEmpty(HttpVersion) && !HttpVersion.Equals("None", StringComparison.OrdinalIgnoreCase))
 					{
 						Writer.WriteValue("httpversion", HttpVersion);
 					}
@@ -542,7 +542,7 @@ namespace AutomationTool.Tasks
 
 					if (!string.IsNullOrEmpty(Parameters.DestinationCloudHttpVersion))
 					{
-						if (Parameters.DestinationCloudHttpVersion.Equals("http2-only", StringComparison.InvariantCultureIgnoreCase))
+						if (Parameters.DestinationCloudHttpVersion.Equals("http2-only", StringComparison.OrdinalIgnoreCase))
 						{
 							OplogExportCommandline.Append(" --assume-http2");
 						}
diff --git a/Engine/Source/Programs/AutomationTool/BuildGraph/TempStorage.cs b/Engine/Source/Programs/AutomationTool/BuildGraph/TempStorage.cs
--- a/Engine/Source/Programs/AutomationTool/BuildGraph/TempStorage.cs
+++ b/Engine/Source/Programs/AutomationTool/BuildGraph/TempStorage.cs
@@ -1271,7 +1271,7 @@ namespace AutomationTool
 		/// <param name="TagName">Name of the tag to get the manifest for</param>
 		static FileReference GetTaggedFileListLocation(DirectoryReference BaseDir, string NodeName, string TagName)
 		{
-			Debug.Assert(TagName.StartsWith("#"));
+			Debug.Assert(TagName.StartsWith("#", StringComparison.Ordinal));
 			return FileReference.Combine(BaseDir, NodeName, String.Format("Tag-{0}.xml", TagName.Substring(1)));
 		}
 
