Missing files.

diff --git a/Engine/Source/Programs/AutomationTool/BuildGraph/BgScriptReader.cs b/Engine/Source/Programs/AutomationTool/BuildGraph/BgScriptReader.cs
--- a/Engine/Source/Programs/AutomationTool/BuildGraph/BgScriptReader.cs
+++ b/Engine/Source/Programs/AutomationTool/BuildGraph/BgScriptReader.cs
@@ -1672,7 +1672,7 @@ namespace AutomationTool
 				BgTask info = new BgTask(element.Location, element.Name);
 				foreach (XmlAttribute? attribute in element.Attributes)
 				{
-					if (String.Compare(attribute!.Name, "If", StringComparison.InvariantCultureIgnoreCase) != 0)
+					if (String.Compare(attribute!.Name, "If", StringComparison.OrdinalIgnoreCase) != 0)
 					{
 						string expandedValue = ExpandProperties(element, attribute.Value);
 						info.Arguments.Add(attribute.Name, expandedValue);
diff --git a/Engine/Source/Programs/AutomationTool/BuildGraph/BuildGraph.cs b/Engine/Source/Programs/AutomationTool/BuildGraph/BuildGraph.cs
--- a/Engine/Source/Programs/AutomationTool/BuildGraph/BuildGraph.cs
+++ b/Engine/Source/Programs/AutomationTool/BuildGraph/BuildGraph.cs
@@ -943,7 +943,7 @@ namespace AutomationTool
 			if (File.IsUnderDirectory(EngineDirectory))
 			{
 				string[] PathFragments = File.MakeRelativeTo(EngineDirectory).Split(Path.DirectorySeparatorChar, Path.AltDirectorySeparatorChar);
-				if (PathFragments.All(x => !x.Equals("NotForLicensees", StringComparison.InvariantCultureIgnoreCase) && !x.Equals("NoRedist", StringComparison.OrdinalIgnoreCase)))
+				if (PathFragments.All(x => !x.Equals("NotForLicensees", StringComparison.OrdinalIgnoreCase) && !x.Equals("NoRedist", StringComparison.OrdinalIgnoreCase)))
 				{
 					return true;
 				}
