BuildGraph: Enable warnings for methods that should be made static.

#rnx

diff --git a/Engine/Source/Programs/AutomationTool/BuildGraph/.editorconfig b/Engine/Source/Programs/AutomationTool/BuildGraph/.editorconfig
--- a/Engine/Source/Programs/AutomationTool/BuildGraph/.editorconfig
+++ b/Engine/Source/Programs/AutomationTool/BuildGraph/.editorconfig
@@ -16,3 +16,8 @@ dotnet_diagnostic.CA1309.severity = warning
 # CA1310: Specify StringComparison for correctness
 dotnet_diagnostic.CA1310.severity = warning
 
+# CA1822: Make member static
+dotnet_diagnostic.CA1822.severity = warning
+
+# IDE0005: Using directive is unnecessary
+dotnet_diagnostic.IDE0005.severity = warning
diff --git a/Engine/Source/Programs/AutomationTool/BuildGraph/BgCondition.cs b/Engine/Source/Programs/AutomationTool/BuildGraph/BgCondition.cs
--- a/Engine/Source/Programs/AutomationTool/BuildGraph/BgCondition.cs
+++ b/Engine/Source/Programs/AutomationTool/BuildGraph/BgCondition.cs
@@ -15,7 +15,7 @@ namespace AutomationTool
 	/// <summary>
 	/// Exception class thrown due to type and syntax errors in condition expressions
 	/// </summary>
-	class BgConditionException : Exception
+	public class BgConditionException : Exception
 	{
 		/// <summary>
 		/// Constructor; formats the exception message with the given String.Format() style parameters.
diff --git a/Engine/Source/Programs/AutomationTool/BuildGraph/BgNodeExecutor.cs b/Engine/Source/Programs/AutomationTool/BuildGraph/BgNodeExecutor.cs
--- a/Engine/Source/Programs/AutomationTool/BuildGraph/BgNodeExecutor.cs
+++ b/Engine/Source/Programs/AutomationTool/BuildGraph/BgNodeExecutor.cs
@@ -72,7 +72,7 @@ namespace AutomationTool
 			Node = node;
 		}
 
-		public bool Bind(ILogger logger)
+		public static bool Bind(ILogger logger)
 		{
 			return true;
 		}
diff --git a/Engine/Source/Programs/AutomationTool/BuildGraph/BgScriptReader.cs b/Engine/Source/Programs/AutomationTool/BuildGraph/BgScriptReader.cs
--- a/Engine/Source/Programs/AutomationTool/BuildGraph/BgScriptReader.cs
+++ b/Engine/Source/Programs/AutomationTool/BuildGraph/BgScriptReader.cs
@@ -723,7 +723,7 @@ namespace AutomationTool
 		/// </summary>
 		/// <param name="Pattern"></param>
 		/// <returns></returns>
-		IEnumerable<FileReference> FindMatchingFiles(string Pattern)
+		static IEnumerable<FileReference> FindMatchingFiles(string Pattern)
 		{
 			FileFilter Filter = new FileFilter();
 			Filter.AddRule(Pattern, FileFilterType.Include);
@@ -2260,7 +2260,7 @@ namespace AutomationTool
 		}
 
 		/// <inheritdoc/>
-		public object GetNativePath(string Path)
+		public static object GetNativePath(string Path)
 		{
 			return FileReference.Combine(Unreal.RootDirectory, Path).FullName;
 		}
diff --git a/Engine/Source/Programs/AutomationTool/BuildGraph/BuildGraph.cs b/Engine/Source/Programs/AutomationTool/BuildGraph/BuildGraph.cs
--- a/Engine/Source/Programs/AutomationTool/BuildGraph/BuildGraph.cs
+++ b/Engine/Source/Programs/AutomationTool/BuildGraph/BuildGraph.cs
@@ -757,7 +757,7 @@ namespace AutomationTool
 			return ExitCode.Success;
 		}
 
-		bool BindNodes(BgGraphDef Graph, Dictionary<string, ScriptTaskBinding> NameToTask, Dictionary<BgNodeDef, BgNodeExecutor> NodeToExecutor)
+		static bool BindNodes(BgGraphDef Graph, Dictionary<string, ScriptTaskBinding> NameToTask, Dictionary<BgNodeDef, BgNodeExecutor> NodeToExecutor)
 		{
 			bool bResult = true;
 			foreach (BgAgentDef Agent in Graph.Agents)
@@ -770,7 +770,7 @@ namespace AutomationTool
 			return bResult;
 		}
 
-		bool BindNodes(BgNodeDef Node, Dictionary<string, ScriptTaskBinding> NameToTask, Dictionary<string, BgNodeOutput> TagNameToNodeOutput, Dictionary<BgNodeDef, BgNodeExecutor> NodeToExecutor)
+		static bool BindNodes(BgNodeDef Node, Dictionary<string, ScriptTaskBinding> NameToTask, Dictionary<string, BgNodeOutput> TagNameToNodeOutput, Dictionary<BgNodeDef, BgNodeExecutor> NodeToExecutor)
 		{
 			if (Node is BgScriptNode ScriptNode)
 			{
@@ -782,7 +782,7 @@ namespace AutomationTool
 			{
 				BgBytecodeNodeExecutor executor = new BgBytecodeNodeExecutor(BytecodeNode);
 				NodeToExecutor[Node] = executor;
-				return executor.Bind(Logger);
+				return BgBytecodeNodeExecutor.Bind(Logger);
 			}
 			else
 			{
@@ -866,7 +866,7 @@ namespace AutomationTool
 		/// Reads the contents of the given token
 		/// </summary>
 		/// <returns>Contents of the token, or null if it does not exist</returns>
-		public string? ReadTokenFile(FileReference Location)
+		public static string? ReadTokenFile(FileReference Location)
 		{
 			return FileReference.Exists(Location) ? File.ReadAllText(Location.FullName) : null;
 		}
@@ -875,7 +875,7 @@ namespace AutomationTool
 		/// Attempts to write an owner to a token file transactionally
 		/// </summary>
 		/// <returns>True if the lock was acquired, false otherwise</returns>
-		public bool WriteTokenFile(FileReference Location, string Signature)
+		public static bool WriteTokenFile(FileReference Location, string Signature)
 		{
 			// Check it doesn't already exist
 			if (FileReference.Exists(Location))
@@ -956,7 +956,7 @@ namespace AutomationTool
 		/// </summary>
 		/// <param name="Graph">The graph instance</param>
 		/// <param name="Storage">The temp storage backend which stores the shared state</param>
-		HashSet<BgNodeDef> FindCompletedNodes(BgGraphDef Graph, TempStorage Storage)
+		static HashSet<BgNodeDef> FindCompletedNodes(BgGraphDef Graph, TempStorage Storage)
 		{
 			HashSet<BgNodeDef> CompletedNodes = new HashSet<BgNodeDef>();
 			foreach (BgNodeDef Node in Graph.Agents.SelectMany(x => x.Nodes))
diff --git a/Engine/Source/Programs/AutomationTool/BuildGraph/Task.cs b/Engine/Source/Programs/AutomationTool/BuildGraph/Task.cs
--- a/Engine/Source/Programs/AutomationTool/BuildGraph/Task.cs
+++ b/Engine/Source/Programs/AutomationTool/BuildGraph/Task.cs
@@ -245,7 +245,7 @@ namespace AutomationTool
 		/// </summary>
 		/// <param name="Filespec">A filespec, as can be passed to ResolveFilespec</param>
 		/// <returns>Tag names from this filespec</returns>
-		protected IEnumerable<string> FindTagNamesFromFilespec(string Filespec)
+		protected static IEnumerable<string> FindTagNamesFromFilespec(string Filespec)
 		{
 			if(!String.IsNullOrEmpty(Filespec))
 			{
diff --git a/Engine/Source/Programs/AutomationTool/BuildGraph/Tasks/AgeStoreTask.cs b/Engine/Source/Programs/AutomationTool/BuildGraph/Tasks/AgeStoreTask.cs
--- a/Engine/Source/Programs/AutomationTool/BuildGraph/Tasks/AgeStoreTask.cs
+++ b/Engine/Source/Programs/AutomationTool/BuildGraph/Tasks/AgeStoreTask.cs
@@ -133,7 +133,7 @@ namespace AutomationTool.Tasks
 			return true;
 		}
 
-		private void RecurseDirectory(DateTime ExpireTimeUtc, DirectoryInfo CurrentDirectory, string[] DirectoryStructure, int Level, string Filter, HashSet<string> ExistingBuilds, bool bDeleteIndividualFiles)
+		private static void RecurseDirectory(DateTime ExpireTimeUtc, DirectoryInfo CurrentDirectory, string[] DirectoryStructure, int Level, string Filter, HashSet<string> ExistingBuilds, bool bDeleteIndividualFiles)
 		{
 			// Do a file search at the last level.
 			if (Level == DirectoryStructure.Length)
diff --git a/Engine/Source/Programs/AutomationTool/BuildGraph/Tasks/GitCloneTask.cs b/Engine/Source/Programs/AutomationTool/BuildGraph/Tasks/GitCloneTask.cs
--- a/Engine/Source/Programs/AutomationTool/BuildGraph/Tasks/GitCloneTask.cs
+++ b/Engine/Source/Programs/AutomationTool/BuildGraph/Tasks/GitCloneTask.cs
@@ -107,7 +107,7 @@ namespace AutomationTool.Tasks
 		/// <param name="ToolFile"></param>
 		/// <param name="Arguments"></param>
 		/// <param name="WorkingDir"></param>
-		Task RunGitAsync(FileReference ToolFile, string Arguments, DirectoryReference WorkingDir)
+		static Task RunGitAsync(FileReference ToolFile, string Arguments, DirectoryReference WorkingDir)
 		{
 			IProcessResult Result = CommandUtils.Run(ToolFile.FullName, Arguments, WorkingDir: WorkingDir.FullName);
 			if (Result.ExitCode != 0)
diff --git a/Engine/Source/Programs/AutomationTool/BuildGraph/Tasks/NuGetLicenseCheckTask.cs b/Engine/Source/Programs/AutomationTool/BuildGraph/Tasks/NuGetLicenseCheckTask.cs
--- a/Engine/Source/Programs/AutomationTool/BuildGraph/Tasks/NuGetLicenseCheckTask.cs
+++ b/Engine/Source/Programs/AutomationTool/BuildGraph/Tasks/NuGetLicenseCheckTask.cs
@@ -112,7 +112,7 @@ namespace AutomationTool.Tasks
 			public FileReference File;
 		}
 
-		LicenseInfo FindOrAddLicense(Dictionary<IoHash, LicenseInfo> Licenses, string Text, string Extension)
+		static LicenseInfo FindOrAddLicense(Dictionary<IoHash, LicenseInfo> Licenses, string Text, string Extension)
 		{
 			string NormalizedText = Text;
 			NormalizedText = Regex.Replace(NormalizedText, @"^\s+", "", RegexOptions.Multiline);
diff --git a/Engine/Source/Programs/AutomationTool/BuildGraph/Tasks/TagTask.cs b/Engine/Source/Programs/AutomationTool/BuildGraph/Tasks/TagTask.cs
--- a/Engine/Source/Programs/AutomationTool/BuildGraph/Tasks/TagTask.cs
+++ b/Engine/Source/Programs/AutomationTool/BuildGraph/Tasks/TagTask.cs
@@ -144,7 +144,7 @@ namespace AutomationTool.Tasks
 		/// <param name="DelimitedPatterns">List of patterns to add, separated by semicolons.</param>
 		/// <param name="TagNameToFileSet">Mapping of tag name to a set of files.</param>
 		/// <returns>List of rules, suitable for adding to a FileFilter object</returns>
-		List<string> ParseRules(DirectoryReference BaseDir, string DelimitedPatterns, Dictionary<string, HashSet<FileReference>> TagNameToFileSet)
+		static List<string> ParseRules(DirectoryReference BaseDir, string DelimitedPatterns, Dictionary<string, HashSet<FileReference>> TagNameToFileSet)
 		{
 			// Split up the list of patterns
 			List<string> Patterns = SplitDelimitedList(DelimitedPatterns);
diff --git a/Engine/Source/Programs/AutomationTool/BuildGraph/Tasks/ZenExportSnapshotTask.cs b/Engine/Source/Programs/AutomationTool/BuildGraph/Tasks/ZenExportSnapshotTask.cs
--- a/Engine/Source/Programs/AutomationTool/BuildGraph/Tasks/ZenExportSnapshotTask.cs
+++ b/Engine/Source/Programs/AutomationTool/BuildGraph/Tasks/ZenExportSnapshotTask.cs
@@ -404,7 +404,7 @@ namespace AutomationTool.Tasks
 			return true;
 		}
 
-		private bool TryExportOplogCommand(string App, string CommandLine)
+		private static bool TryExportOplogCommand(string App, string CommandLine)
 		{
 			int AttemptLimit = 2;
 			int Attempt = 0;
