BuildGraph: Fix comparison error after refactor.

#rnx

diff --git a/Engine/Source/Programs/AutomationTool/BuildGraph/BgCondition.cs b/Engine/Source/Programs/AutomationTool/BuildGraph/BgCondition.cs
--- a/Engine/Source/Programs/AutomationTool/BuildGraph/BgCondition.cs
+++ b/Engine/Source/Programs/AutomationTool/BuildGraph/BgCondition.cs
@@ -180,7 +180,7 @@ namespace AutomationTool
 				_idx++;
 				string lhs = result;
 				string rhs = await EvaluateScalarAsync();
-				result = String.Equals(lhs, rhs, StringComparison.OrdinalIgnoreCase) ? "true" : "false";
+				result = String.Equals(lhs, rhs, StringComparison.OrdinalIgnoreCase) ? "false" : "true";
 			}
 			else if (_tokens[_idx] == "<")
 			{
