EpicGames.BuildGraph: Apply standard naming conventions to BuildGraph.Automation.

diff --git a/Engine/Source/Programs/AutomationTool/BuildGraph/BgCondition.cs b/Engine/Source/Programs/AutomationTool/BuildGraph/BgCondition.cs
--- a/Engine/Source/Programs/AutomationTool/BuildGraph/BgCondition.cs
+++ b/Engine/Source/Programs/AutomationTool/BuildGraph/BgCondition.cs
@@ -1,11 +1,11 @@
 // Copyright Epic Games, Inc. All Rights Reserved.
 
-using EpicGames.Core;
 using System;
 using System.Collections.Generic;
 using System.IO;
 using System.Linq;
 using System.Threading.Tasks;
+using EpicGames.Core;
 using UnrealBuildBase;
 
 #nullable enable
@@ -88,7 +88,7 @@ namespace AutomationTool
 		/// </summary>
 		/// <param name="text"></param>
 		/// <returns>The result of evaluating the condition</returns>
-		public static ValueTask<bool> EvaluateAsync(string text)
+		public static ValueTask<bool> Evaluate(string text)
 		{
 			return new BgCondition(text).EvaluateAsync();
 		}
@@ -99,18 +99,20 @@ namespace AutomationTool
 		/// <returns>The result of evaluating the condition</returns>
 		async ValueTask<bool> EvaluateAsync()
 		{
-			bool bResult = true;
+			bool result = true;
 			if (_tokens.Count > 1)
 			{
 				_idx = 0;
-				string result = await EvaluateOrAsync();
+
+				string value = await EvaluateOrAsync();
 				if (_tokens[_idx] != EndToken)
 				{
 					throw new BgConditionException("Garbage after expression: {0}", String.Join("", _tokens.Skip(_idx)));
 				}
-				bResult = CoerceToBool(result);
+				
+				result = CoerceToBool(value);
 			}
-			return bResult;
+			return result;
 		}
 
 		/// <summary>
@@ -355,13 +357,13 @@ namespace AutomationTool
 		/// <summary>
 		/// Determine if a path exists
 		/// </summary>
-		/// <param name="Path"></param>
+		/// <param name="path"></param>
 		/// <returns></returns>
-		static bool Exists(string Path)
+		static bool Exists(string path)
 		{
 			try
 			{
-				return FileReference.Exists(FileReference.Combine(Unreal.RootDirectory, Path)) || DirectoryReference.Exists(DirectoryReference.Combine(Unreal.RootDirectory, Path));
+				return FileReference.Exists(FileReference.Combine(Unreal.RootDirectory, path)) || DirectoryReference.Exists(DirectoryReference.Combine(Unreal.RootDirectory, path));
 			}
 			catch
 			{
@@ -505,7 +507,7 @@ namespace AutomationTool
 		/// <summary>
 		/// Test cases for conditions.
 		/// </summary>
-		public static async Task TestConditions()
+		public static async Task TestConditionsAsync()
 		{
 			await TestConditionAsync("1 == 2", false);
 			await TestConditionAsync("1 == 1", true);
diff --git a/Engine/Source/Programs/AutomationTool/BuildGraph/BgGraphBuilder.cs b/Engine/Source/Programs/AutomationTool/BuildGraph/BgGraphBuilder.cs
--- a/Engine/Source/Programs/AutomationTool/BuildGraph/BgGraphBuilder.cs
+++ b/Engine/Source/Programs/AutomationTool/BuildGraph/BgGraphBuilder.cs
@@ -1,12 +1,8 @@
 // Copyright Epic Games, Inc. All Rights Reserved.
 
-using EpicGames.BuildGraph;
 using EpicGames.BuildGraph.Expressions;
 using EpicGames.Core;
 using Microsoft.Extensions.Logging;
-using System;
-using System.Collections.Generic;
-using System.Text;
 
 #nullable enable
 
diff --git a/Engine/Source/Programs/AutomationTool/BuildGraph/BgNodeExecutor.cs b/Engine/Source/Programs/AutomationTool/BuildGraph/BgNodeExecutor.cs
--- a/Engine/Source/Programs/AutomationTool/BuildGraph/BgNodeExecutor.cs
+++ b/Engine/Source/Programs/AutomationTool/BuildGraph/BgNodeExecutor.cs
@@ -3,17 +3,10 @@
 using System;
 using System.Collections.Generic;
 using System.ComponentModel;
-using System.Diagnostics.CodeAnalysis;
-using System.IO;
 using System.Linq;
 using System.Reflection;
 using System.Runtime.CompilerServices;
-using System.Text;
-using System.Text.RegularExpressions;
 using System.Threading.Tasks;
-using System.Xml;
-using System.Xml.Schema;
-using AutomationTool.Tasks;
 using EpicGames.BuildGraph;
 using EpicGames.BuildGraph.Expressions;
 using EpicGames.Core;
@@ -32,19 +25,17 @@ namespace AutomationTool
 	/// </summary>
 	abstract class BgNodeExecutor
 	{
-		public abstract Task<bool> ExecuteAsync(JobContext Job, Dictionary<string, HashSet<FileReference>> TagNameToFileSet);
+		public abstract Task<bool> Execute(JobContext job, Dictionary<string, HashSet<FileReference>> tagNameToFileSet);
 	}
 
 	class BgBytecodeNodeExecutor : BgNodeExecutor
 	{
 		class BgContextImpl : BgContext
 		{
-			JobContext JobContext;
-
-			public BgContextImpl(JobContext JobContext, Dictionary<string, HashSet<FileReference>> TagNameToFileSet)
-				: base(TagNameToFileSet.ToDictionary(x => x.Key, x => FileSet.FromFiles(Unreal.RootDirectory, x.Value)))
+			public BgContextImpl(JobContext jobContext, Dictionary<string, HashSet<FileReference>> tagNameToFileSet)
+				: base(tagNameToFileSet.ToDictionary(x => x.Key, x => FileSet.FromFiles(Unreal.RootDirectory, x.Value)))
 			{
-				this.JobContext = JobContext;
+				_ = jobContext;
 			}
 
 			public override string Stream => CommandUtils.P4Enabled ? CommandUtils.P4Env.Branch : "";
@@ -57,89 +48,90 @@ namespace AutomationTool
 			{
 				get
 				{
-					ReadOnlyBuildVersion Current = ReadOnlyBuildVersion.Current;
-					return (Current.MajorVersion, Current.MinorVersion, Current.PatchVersion);
+					ReadOnlyBuildVersion current = ReadOnlyBuildVersion.Current;
+					return (current.MajorVersion, current.MinorVersion, current.PatchVersion);
 				}
 			}
 
 			public override bool IsBuildMachine => CommandUtils.IsBuildMachine;
 		}
 
-		readonly BgNodeDef Node;
+		readonly BgNodeDef _node;
 
 		public BgBytecodeNodeExecutor(BgNodeDef node)
 		{
-			Node = node;
+			_node = node;
 		}
 
 		public static bool Bind(ILogger logger)
 		{
+			_ = logger;
 			return true;
 		}
 
 		/// <summary>
-		/// Execute the method given in the 
+		/// ExecuteAsync the method given in the 
 		/// </summary>
-		/// <param name="Job"></param>
-		/// <param name="TagNameToFileSet"></param>
+		/// <param name="job"></param>
+		/// <param name="tagNameToFileSet"></param>
 		/// <returns></returns>
-		public override async Task<bool> ExecuteAsync(JobContext Job, Dictionary<string, HashSet<FileReference>> TagNameToFileSet)
+		public override async Task<bool> Execute(JobContext job, Dictionary<string, HashSet<FileReference>> tagNameToFileSet)
 		{
-			BgThunkDef thunk = Node.Thunk!;
+			BgThunkDef thunk = _node.Thunk!;
 			MethodInfo method = thunk.Method;
 
-			HashSet<FileReference> BuildProducts = TagNameToFileSet[Node.DefaultOutput.TagName];
+			HashSet<FileReference> buildProducts = tagNameToFileSet[_node.DefaultOutput.TagName];
 
-			BgContextImpl Context = new BgContextImpl(Job, TagNameToFileSet);
+			BgContextImpl context = new BgContextImpl(job, tagNameToFileSet);
 
 			ParameterInfo[] parameters = method.GetParameters();
 
 			object?[] arguments = new object[parameters.Length];
-			for (int Idx = 0; Idx < parameters.Length; Idx++)
+			for (int idx = 0; idx < parameters.Length; idx++)
 			{
-				Type ParameterType = parameters[Idx].ParameterType;
-				if (ParameterType == typeof(BgContext))
+				Type parameterType = parameters[idx].ParameterType;
+				if (parameterType == typeof(BgContext))
 				{
-					arguments[Idx] = Context;
+					arguments[idx] = context;
 				}
 				else
 				{
-					arguments[Idx] = thunk.Arguments[Idx];
+					arguments[idx] = thunk.Arguments[idx];
 				}
 			}
 
-			Task Task = (Task)method.Invoke(null, arguments)!;
-			await Task;
+			Task task = (Task)method.Invoke(null, arguments)!;
+			await task;
 
-			if (Node.Outputs.Count > 0)
+			if (_node.Outputs.Count > 0)
 			{
-				object? Result = null;
+				object? result = null;
 				if (method.ReturnType.IsGenericType && method.ReturnType.GetGenericTypeDefinition() == typeof(Task<>))
 				{
-					Type TaskType = Task.GetType();
+					Type taskType = task.GetType();
 #pragma warning disable CA1849 // Task.Result synchronously blocks
-					PropertyInfo Property = TaskType.GetProperty(nameof(Task<int>.Result))!;
+					PropertyInfo property = taskType.GetProperty(nameof(Task<int>.Result))!;
 #pragma warning restore CA1849
-					Result = Property!.GetValue(Task);
+					result = property!.GetValue(task);
 				}
 
-				object?[] OutputValues;
-				if (Result is ITuple Tuple)
+				object?[] outputValues;
+				if (result is ITuple tuple)
 				{
-					OutputValues = Enumerable.Range(0, Tuple.Length).Select(x => Tuple[x]).ToArray();
+					outputValues = Enumerable.Range(0, tuple.Length).Select(x => tuple[x]).ToArray();
 				}
 				else
 				{
-					OutputValues = new[] { Result };
+					outputValues = new[] { result };
 				}
 
-				for (int Idx = 0; Idx < OutputValues.Length; Idx++)
+				for (int idx = 0; idx < outputValues.Length; idx++)
 				{
-					if (OutputValues[Idx] is BgFileSetOutputExpr FileSet)
+					if (outputValues[idx] is BgFileSetOutputExpr fileSet)
 					{
-						string TagName = Node.Outputs[Idx + 1].TagName;
-						TagNameToFileSet[TagName] = new HashSet<FileReference>(FileSet.Value.Flatten().Values);
-						BuildProducts.UnionWith(TagNameToFileSet[TagName]);
+						string tagName = _node.Outputs[idx + 1].TagName;
+						tagNameToFileSet[tagName] = new HashSet<FileReference>(fileSet.Value.Flatten().Values);
+						buildProducts.UnionWith(tagNameToFileSet[tagName]);
 					}
 				}
 			}
@@ -161,7 +153,7 @@ namespace AutomationTool
 		/// <summary>
 		/// List of bound task implementations
 		/// </summary>
-		List<BgTaskImpl> _boundTasks = new List<BgTaskImpl>();
+		readonly List<BgTaskImpl> _boundTasks = new List<BgTaskImpl>();
 
 		/// <summary>
 		/// Constructor
@@ -173,222 +165,222 @@ namespace AutomationTool
 
 		public bool Bind(Dictionary<string, ScriptTaskBinding> nameToTask, Dictionary<string, BgNodeOutput> tagNameToNodeOutput, ILogger logger)
 		{
-			bool bResult = true;
-			foreach (BgTask TaskInfo in Node.Tasks)
+			bool result = true;
+			foreach (BgTask taskInfo in Node.Tasks)
 			{
-				BgTaskImpl? boundTask = BindTask(TaskInfo, nameToTask, tagNameToNodeOutput, logger);
+				BgTaskImpl? boundTask = BindTask(taskInfo, nameToTask, tagNameToNodeOutput, logger);
 				if (boundTask == null)
 				{
-					bResult = false;
+					result = false;
 				}
 				else
 				{
 					_boundTasks.Add(boundTask);
 				}
 			}
-			return bResult;
+			return result;
 		}
 
-		BgTaskImpl? BindTask(BgTask TaskInfo, Dictionary<string, ScriptTaskBinding> NameToTask, IReadOnlyDictionary<string, BgNodeOutput> TagNameToNodeOutput, ILogger Logger)
+		BgTaskImpl? BindTask(BgTask taskInfo, Dictionary<string, ScriptTaskBinding> nameToTask, IReadOnlyDictionary<string, BgNodeOutput> tagNameToNodeOutput, ILogger logger)
 		{
 			// Get the reflection info for this element
-			ScriptTaskBinding? Task;
-			if (!NameToTask.TryGetValue(TaskInfo.Name, out Task))
+			ScriptTaskBinding? task;
+			if (!nameToTask.TryGetValue(taskInfo.Name, out task))
 			{
-				Logger.LogScriptError(TaskInfo.Location, "Unknown task '{TaskName}'", TaskInfo.Name);
+				logger.LogScriptError(taskInfo.Location, "Unknown task '{TaskName}'", taskInfo.Name);
 				return null;
 			}
 
 			// Check all the required parameters are present
-			bool bHasRequiredAttributes = true;
-			foreach (ScriptTaskParameterBinding Parameter in Task.NameToParameter.Values)
+			bool hasRequiredAttributes = true;
+			foreach (ScriptTaskParameterBinding parameter in task.NameToParameter.Values)
 			{
-				if (!Parameter.Optional && !TaskInfo.Arguments.ContainsKey(Parameter.Name))
+				if (!parameter.Optional && !taskInfo.Arguments.ContainsKey(parameter.Name))
 				{
-					Logger.LogScriptError(TaskInfo.Location, "Missing required attribute - {AttrName}", Parameter.Name);
-					bHasRequiredAttributes = false;
+					logger.LogScriptError(taskInfo.Location, "Missing required attribute - {AttrName}", parameter.Name);
+					hasRequiredAttributes = false;
 				}
 			}
 
 			// Read all the attributes into a parameters object for this task
-			object ParametersObject = Activator.CreateInstance(Task.ParametersClass)!;
-			foreach ((string Name, string Value) in TaskInfo.Arguments)
+			object parametersObject = Activator.CreateInstance(task.ParametersClass)!;
+			foreach ((string name, string value) in taskInfo.Arguments)
 			{
 				// Get the field that this attribute should be written to in the parameters object
-				ScriptTaskParameterBinding? Parameter;
-				if (!Task.NameToParameter.TryGetValue(Name, out Parameter))
+				ScriptTaskParameterBinding? parameter;
+				if (!task.NameToParameter.TryGetValue(name, out parameter))
 				{
-					Logger.LogScriptError(TaskInfo.Location, "Unknown attribute '{AttrName}'", Name);
+					logger.LogScriptError(taskInfo.Location, "Unknown attribute '{AttrName}'", name);
 					continue;
 				}
 
 				// If it's a collection type, split it into separate values
 				try
 				{
-					if (Parameter.CollectionType == null)
+					if (parameter.CollectionType == null)
 					{
 						// Parse it and assign it to the parameters object
-						object? FieldValue = ParseValue(Value, Parameter.ValueType);
-						if (FieldValue != null)
+						object? fieldValue = ParseValue(value, parameter.ValueType);
+						if (fieldValue != null)
 						{
-							Parameter.SetValue(ParametersObject, FieldValue);
+							parameter.SetValue(parametersObject, fieldValue);
 						}
-						else if (!Parameter.Optional)
+						else if (!parameter.Optional)
 						{
-							Logger.LogScriptError(TaskInfo.Location, "Empty value for parameter '{AttrName}' is not allowed.", Name);
+							logger.LogScriptError(taskInfo.Location, "Empty value for parameter '{AttrName}' is not allowed.", name);
 						}
 					}
 					else
 					{
 						// Get the collection, or create one if necessary
-						object? CollectionValue = Parameter.GetValue(ParametersObject);
-						if (CollectionValue == null)
+						object? collectionValue = parameter.GetValue(parametersObject);
+						if (collectionValue == null)
 						{
-							CollectionValue = Activator.CreateInstance(Parameter.ParameterType)!;
-							Parameter.SetValue(ParametersObject, CollectionValue);
+							collectionValue = Activator.CreateInstance(parameter.ParameterType)!;
+							parameter.SetValue(parametersObject, collectionValue);
 						}
 
 						// Parse the values and add them to the collection
-						List<string> ValueStrings = BgTaskImpl.SplitDelimitedList(Value);
-						foreach (string ValueString in ValueStrings)
+						List<string> valueStrings = BgTaskImpl.SplitDelimitedList(value);
+						foreach (string valueString in valueStrings)
 						{
-							object? ElementValue = ParseValue(ValueString, Parameter.ValueType);
-							if (ElementValue != null)
+							object? elementValue = ParseValue(valueString, parameter.ValueType);
+							if (elementValue != null)
 							{
-								Parameter.CollectionType.InvokeMember("Add", BindingFlags.InvokeMethod | BindingFlags.Instance | BindingFlags.Public, null, CollectionValue, new object[] { ElementValue });
+								parameter.CollectionType.InvokeMember("Add", BindingFlags.InvokeMethod | BindingFlags.Instance | BindingFlags.Public, null, collectionValue, new object[] { elementValue });
 							}
 						}
 					}
 				}
 				catch (Exception ex)
 				{
-					Logger.LogScriptError(TaskInfo.Location, "Unable to parse argument {Name} from {Value}", Name, Value);
-					Logger.LogDebug(ex, "Exception while parsing argument {Name}", Name);
+					logger.LogScriptError(taskInfo.Location, "Unable to parse argument {Name} from {Value}", name, value);
+					logger.LogDebug(ex, "Exception while parsing argument {Name}", name);
 				}
 			}
 
 			// Construct the task
-			if (!bHasRequiredAttributes)
+			if (!hasRequiredAttributes)
 			{
 				return null;
 			}
 
 			// Add it to the list
-			BgTaskImpl NewTask = (BgTaskImpl)Activator.CreateInstance(Task.TaskClass, ParametersObject)!;
+			BgTaskImpl newTask = (BgTaskImpl)Activator.CreateInstance(task.TaskClass, parametersObject)!;
 
 			// Set up the source location for diagnostics
-			NewTask.SourceLocation = TaskInfo.Location;
+			newTask.SourceLocation = taskInfo.Location;
 
 			// Make sure all the read tags are local or listed as a dependency
-			foreach (string ReadTagName in NewTask.FindConsumedTagNames())
+			foreach (string readTagName in newTask.FindConsumedTagNames())
 			{
-				BgNodeOutput? Output;
-				if (TagNameToNodeOutput.TryGetValue(ReadTagName, out Output))
+				BgNodeOutput? output;
+				if (tagNameToNodeOutput.TryGetValue(readTagName, out output))
 				{
-					if (Output != null && Output.ProducingNode != Node && !Node.Inputs.Contains(Output))
+					if (output != null && output.ProducingNode != Node && !Node.Inputs.Contains(output))
 					{
-						Logger.LogScriptError(TaskInfo.Location, "The tag '{TagName}' is not a dependency of node '{Node}'", ReadTagName, Node.Name);
+						logger.LogScriptError(taskInfo.Location, "The tag '{TagName}' is not a dependency of node '{Node}'", readTagName, Node.Name);
 					}
 				}
 			}
 
 			// Make sure all the written tags are local or listed as an output
-			foreach (string ModifiedTagName in NewTask.FindProducedTagNames())
+			foreach (string modifiedTagName in newTask.FindProducedTagNames())
 			{
-				BgNodeOutput? Output;
-				if (TagNameToNodeOutput.TryGetValue(ModifiedTagName, out Output))
+				BgNodeOutput? output;
+				if (tagNameToNodeOutput.TryGetValue(modifiedTagName, out output))
 				{
-					if (Output != null && !Node.Outputs.Contains(Output))
+					if (output != null && !Node.Outputs.Contains(output))
 					{
-						Logger.LogScriptError(TaskInfo.Location, "The tag '{TagName}' is created by '{Node}', and cannot be modified downstream", Output.TagName, Output.ProducingNode.Name);
+						logger.LogScriptError(taskInfo.Location, "The tag '{TagName}' is created by '{Node}', and cannot be modified downstream", output.TagName, output.ProducingNode.Name);
 					}
 				}
 			}
-			return NewTask;
+			return newTask;
 		}
 
 		/// <summary>
 		/// Parse a value of the given type
 		/// </summary>
-		/// <param name="ValueText">The text to parse</param>
-		/// <param name="ValueType">Type of the value to parse</param>
+		/// <param name="valueText">The text to parse</param>
+		/// <param name="valueType">Type of the value to parse</param>
 		/// <returns>Value that was parsed</returns>
-		static object? ParseValue(string ValueText, Type ValueType)
+		static object? ParseValue(string valueText, Type valueType)
 		{
 			// Parse it and assign it to the parameters object
-			if (ValueType.IsEnum)
+			if (valueType.IsEnum)
 			{
-				return Enum.Parse(ValueType, ValueText);
+				return Enum.Parse(valueType, valueText);
 			}
-			else if (ValueType == typeof(Boolean))
+			else if (valueType == typeof(Boolean))
 			{
-				return BgCondition.EvaluateAsync(ValueText).AsTask().Result;
+				return BgCondition.Evaluate(valueText).AsTask().Result;
 			}
-			else if (ValueType == typeof(FileReference))
+			else if (valueType == typeof(FileReference))
 			{
-				if (String.IsNullOrEmpty(ValueText))
+				if (String.IsNullOrEmpty(valueText))
 				{
 					return null;
 				}
 				else
 				{
-					return BgTaskImpl.ResolveFile(ValueText);
+					return BgTaskImpl.ResolveFile(valueText);
 				}
 			}
-			else if (ValueType == typeof(DirectoryReference))
+			else if (valueType == typeof(DirectoryReference))
 			{
-				if (String.IsNullOrEmpty(ValueText))
+				if (String.IsNullOrEmpty(valueText))
 				{
 					return null;
 				}
 				else
 				{
-					return BgTaskImpl.ResolveDirectory(ValueText);
+					return BgTaskImpl.ResolveDirectory(valueText);
 				}
 			}
 
-			TypeConverter Converter = TypeDescriptor.GetConverter(ValueType);
-			if (Converter.CanConvertFrom(typeof(string)))
+			TypeConverter converter = TypeDescriptor.GetConverter(valueType);
+			if (converter.CanConvertFrom(typeof(string)))
 			{
-				return Converter.ConvertFromString(ValueText);
+				return converter.ConvertFromString(valueText);
 			}
 			else
 			{
-				return Convert.ChangeType(ValueText, ValueType);
+				return Convert.ChangeType(valueText, valueType);
 			}
 		}
 
 		/// <summary>
 		/// Build all the tasks for this node
 		/// </summary>
-		/// <param name="Job">Information about the current job</param>
-		/// <param name="TagNameToFileSet">Mapping from tag names to the set of files they include. Should be set to contain the node inputs on entry.</param>
+		/// <param name="job">Information about the current job</param>
+		/// <param name="tagNameToFileSet">Mapping from tag names to the set of files they include. Should be set to contain the node inputs on entry.</param>
 		/// <returns>Whether the task succeeded or not. Exiting with an exception will be caught and treated as a failure.</returns>
-		public override async Task<bool> ExecuteAsync(JobContext Job, Dictionary<string, HashSet<FileReference>> TagNameToFileSet)
+		public override async Task<bool> Execute(JobContext job, Dictionary<string, HashSet<FileReference>> tagNameToFileSet)
 		{
 			// Run each of the tasks in order
-			HashSet<FileReference> BuildProducts = TagNameToFileSet[Node.DefaultOutput.TagName];
-			for (int Idx = 0; Idx < _boundTasks.Count; Idx++)
+			HashSet<FileReference> buildProducts = tagNameToFileSet[Node.DefaultOutput.TagName];
+			for (int idx = 0; idx < _boundTasks.Count; idx++)
 			{
-				using (IScope Scope = GlobalTracer.Instance.BuildSpan("Task").WithTag("resource", _boundTasks[Idx].GetTraceName()).StartActive())
+				using (IScope scope = GlobalTracer.Instance.BuildSpan("Task").WithTag("resource", _boundTasks[idx].GetTraceName()).StartActive())
 				{
-					ITaskExecutor? Executor = _boundTasks[Idx].GetExecutor();
-					if (Executor == null)
+					ITaskExecutor? executor = _boundTasks[idx].GetExecutor();
+					if (executor == null)
 					{
-						// Execute this task directly
+						// ExecuteAsync this task directly
 						try
 						{
-							_boundTasks[Idx].GetTraceMetadata(Scope.Span, "");
-							await _boundTasks[Idx].ExecuteAsync(Job, BuildProducts, TagNameToFileSet);
+							_boundTasks[idx].GetTraceMetadata(scope.Span, "");
+							await _boundTasks[idx].ExecuteAsync(job, buildProducts, tagNameToFileSet);
 						}
-						catch (Exception Ex)
+						catch (Exception ex)
 						{
-							ExceptionUtils.AddContext(Ex, "while executing task {0}", _boundTasks[Idx].GetTraceString());
+							ExceptionUtils.AddContext(ex, "while executing task {0}", _boundTasks[idx].GetTraceString());
 
-							BgScriptLocation? sourceLocation = _boundTasks[Idx].SourceLocation;
+							BgScriptLocation? sourceLocation = _boundTasks[idx].SourceLocation;
 							if (sourceLocation != null)
 							{
-								ExceptionUtils.AddContext(Ex, "at {0}({1})", sourceLocation.File, sourceLocation.LineNumber);
+								ExceptionUtils.AddContext(ex, "at {0}({1})", sourceLocation.File, sourceLocation.LineNumber);
 							}
 
 							throw;
@@ -396,30 +388,30 @@ namespace AutomationTool
 					}
 					else
 					{
-						_boundTasks[Idx].GetTraceMetadata(Scope.Span, "1.");
+						_boundTasks[idx].GetTraceMetadata(scope.Span, "1.");
 
 						// The task has a custom executor, which may be able to execute several tasks simultaneously. Try to add the following tasks.
-						int FirstIdx = Idx;
-						while (Idx + 1 < Node.Tasks.Count && Executor.Add(_boundTasks[Idx + 1]))
+						int firstIdx = idx;
+						while (idx + 1 < Node.Tasks.Count && executor.Add(_boundTasks[idx + 1]))
 						{
-							Idx++;
-							_boundTasks[Idx].GetTraceMetadata(Scope.Span, string.Format("{0}.", 1 + Idx - FirstIdx));
+							idx++;
+							_boundTasks[idx].GetTraceMetadata(scope.Span, string.Format("{0}.", 1 + idx - firstIdx));
 						}
 						try
 						{
-							await Executor.ExecuteAsync(Job, BuildProducts, TagNameToFileSet);
+							await executor.ExecuteAsync(job, buildProducts, tagNameToFileSet);
 						}
-						catch (Exception Ex)
+						catch (Exception ex)
 						{
-							for (int TaskIdx = FirstIdx; TaskIdx <= Idx; TaskIdx++)
+							for (int taskIdx = firstIdx; taskIdx <= idx; taskIdx++)
 							{
-								ExceptionUtils.AddContext(Ex, "while executing {0}", _boundTasks[TaskIdx].GetTraceString());
+								ExceptionUtils.AddContext(ex, "while executing {0}", _boundTasks[taskIdx].GetTraceString());
 							}
 
-							BgScriptLocation? sourceLocation = _boundTasks[FirstIdx].SourceLocation;
+							BgScriptLocation? sourceLocation = _boundTasks[firstIdx].SourceLocation;
 							if (sourceLocation != null)
 							{
-								ExceptionUtils.AddContext(Ex, "at {0}({1})", sourceLocation.File, sourceLocation.LineNumber);
+								ExceptionUtils.AddContext(ex, "at {0}({1})", sourceLocation.File, sourceLocation.LineNumber);
 							}
 
 							throw;
@@ -429,7 +421,7 @@ namespace AutomationTool
 			}
 
 			// Remove anything that doesn't exist, since these files weren't explicitly tagged
-			BuildProducts.RemoveWhere(x => !FileReference.Exists(x));
+			buildProducts.RemoveWhere(x => !FileReference.Exists(x));
 			return true;
 		}
 	}
diff --git a/Engine/Source/Programs/AutomationTool/BuildGraph/BgSchema.cs b/Engine/Source/Programs/AutomationTool/BuildGraph/BgSchema.cs
--- a/Engine/Source/Programs/AutomationTool/BuildGraph/BgSchema.cs
+++ b/Engine/Source/Programs/AutomationTool/BuildGraph/BgSchema.cs
@@ -1,13 +1,13 @@
 // Copyright Epic Games, Inc. All Rights Reserved.
 
-using EpicGames.BuildGraph;
-using EpicGames.Core;
 using System;
 using System.Collections.Generic;
 using System.Diagnostics;
 using System.Linq;
 using System.Xml;
 using System.Xml.Schema;
+using EpicGames.BuildGraph;
+using EpicGames.Core;
 
 #nullable enable
 
@@ -68,12 +68,12 @@ namespace AutomationTool
 		/// <summary>
 		/// Constructor
 		/// </summary>
-		public BgScriptTaskParameter(string inName, Type inValueType, TaskParameterValidationType inValidationType, bool bInOptional)
+		public BgScriptTaskParameter(string name, Type valueType, TaskParameterValidationType validationType, bool optional)
 		{
-			Name = inName;
-			ValueType = inValueType;
-			ValidationType = inValidationType;
-			Optional = bInOptional;
+			Name = name;
+			ValueType = valueType;
+			ValidationType = validationType;
+			Optional = optional;
 
 			if (ValueType.IsGenericType && ValueType.GetGenericTypeDefinition() == typeof(Nullable<>))
 			{
diff --git a/Engine/Source/Programs/AutomationTool/BuildGraph/BgScriptNode.cs b/Engine/Source/Programs/AutomationTool/BuildGraph/BgScriptNode.cs
--- a/Engine/Source/Programs/AutomationTool/BuildGraph/BgScriptNode.cs
+++ b/Engine/Source/Programs/AutomationTool/BuildGraph/BgScriptNode.cs
@@ -1,26 +1,5 @@
 // Copyright Epic Games, Inc. All Rights Reserved.
 
-using System;
-using System.Collections.Generic;
-using System.ComponentModel;
-using System.Diagnostics.CodeAnalysis;
-using System.IO;
-using System.Linq;
-using System.Reflection;
-using System.Runtime.CompilerServices;
-using System.Text;
-using System.Text.RegularExpressions;
-using System.Threading.Tasks;
-using System.Xml;
-using System.Xml.Schema;
-using AutomationTool.Tasks;
-using EpicGames.BuildGraph;
-using EpicGames.BuildGraph.Expressions;
-using EpicGames.Core;
-using Microsoft.Extensions.Logging;
-using OpenTracing;
-using OpenTracing.Util;
-
 #nullable enable
 
 namespace AutomationTool
diff --git a/Engine/Source/Programs/AutomationTool/BuildGraph/BgScriptReader.cs b/Engine/Source/Programs/AutomationTool/BuildGraph/BgScriptReader.cs
--- a/Engine/Source/Programs/AutomationTool/BuildGraph/BgScriptReader.cs
+++ b/Engine/Source/Programs/AutomationTool/BuildGraph/BgScriptReader.cs
@@ -1,8 +1,5 @@
 // Copyright Epic Games, Inc. All Rights Reserved.
 
-using EpicGames.BuildGraph;
-using EpicGames.Core;
-using Microsoft.Extensions.Logging;
 using System;
 using System.Collections.Generic;
 using System.Diagnostics.CodeAnalysis;
@@ -13,6 +10,9 @@ using System.Text.RegularExpressions;
 using System.Threading.Tasks;
 using System.Xml;
 using System.Xml.Schema;
+using EpicGames.BuildGraph;
+using EpicGames.Core;
+using Microsoft.Extensions.Logging;
 using UnrealBuildBase;
 
 #nullable enable
@@ -74,10 +74,10 @@ namespace AutomationTool
 		/// <summary>
 		/// Private constructor. Use ScriptDocument.Load to read an XML document.
 		/// </summary>
-		BgScriptDocument(FileReference inFile, ILogger inLogger)
+		BgScriptDocument(FileReference file, ILogger logger)
 		{
-			File = inFile;
-			Logger = inLogger;
+			File = file;
+			Logger = logger;
 		}
 
 		/// <summary>
@@ -727,14 +727,14 @@ namespace AutomationTool
 		/// <summary>
 		/// Find files matching a pattern
 		/// </summary>
-		/// <param name="Pattern"></param>
+		/// <param name="pattern"></param>
 		/// <returns></returns>
-		static IEnumerable<FileReference> FindMatchingFiles(string Pattern)
+		static IEnumerable<FileReference> FindMatchingFiles(string pattern)
 		{
-			FileFilter Filter = new FileFilter();
-			Filter.AddRule(Pattern, FileFilterType.Include);
+			FileFilter filter = new FileFilter();
+			filter.AddRule(pattern, FileFilterType.Include);
 
-			return Filter.ApplyToDirectory(Unreal.RootDirectory, true);
+			return filter.ApplyToDirectory(Unreal.RootDirectory, true);
 		}
 
 		/// <summary>
@@ -875,7 +875,7 @@ namespace AutomationTool
 						string input = ReadAttribute(element, "Input");
 						Match match = regexValue.Match(input);
 
-						bool optional = await BgCondition.EvaluateAsync(ReadAttribute(element, "Optional"));
+						bool optional = await BgCondition.Evaluate(ReadAttribute(element, "Optional"));
 						if (!match.Success)
 						{
 							if (!optional)
@@ -908,8 +908,6 @@ namespace AutomationTool
 				string method = ReadAttribute(element, "Method");
 				string output = ReadAttribute(element, "Output");
 
-				string operationResult = string.Empty;
-
 				string[] arguments = Array.Empty<string>();
 
 				const string ArgumentsName = "Arguments";
@@ -920,10 +918,15 @@ namespace AutomationTool
 				}
 
 				// Supply more string operations here
+				string operationResult;
 				switch (method)
 				{
-					case "ToLower": operationResult = input.ToLower(); break;
-					case "ToUpper": operationResult = input.ToUpper(); break;
+					case "ToLower":
+						operationResult = input.ToLower();
+						break;
+					case "ToUpper":
+						operationResult = input.ToUpper();
+						break;
 					case "Replace":
 						if (arguments.Length != 2)
 						{
@@ -1371,8 +1374,8 @@ namespace AutomationTool
 				string[] producesNames = ReadListAttribute(element, "Produces");
 				string[] afterNames = ReadListAttribute(element, "After");
 				string[] tokenFileNames = ReadListAttribute(element, "Token");
-				bool bRunEarly = ReadBooleanAttribute(element, "RunEarly", false);
-				bool bNotifyOnWarnings = ReadBooleanAttribute(element, "NotifyOnWarnings", true);
+				bool runEarly = ReadBooleanAttribute(element, "RunEarly", false);
+				bool notifyOnWarnings = ReadBooleanAttribute(element, "NotifyOnWarnings", true);
 				Dictionary<string, string> annotations = ReadAnnotationsAttribute(element, "Annotations");
 				string[] ignoreModified = ReadListAttribute(element, "IgnoreModified");
 
@@ -1440,8 +1443,8 @@ namespace AutomationTool
 				{
 					// Add it to the node lookup
 					BgScriptNode newNode = new BgScriptNode(name, inputs.ToArray(), validOutputNames.ToArray(), inputDependencies.ToArray(), orderDependencies.ToArray(), requiredTokens.ToArray(), ignoreModified.ToArray());
-					newNode.RunEarly = bRunEarly;
-					newNode.NotifyOnWarnings = bNotifyOnWarnings;
+					newNode.RunEarly = runEarly;
+					newNode.NotifyOnWarnings = notifyOnWarnings;
 					foreach ((string key, string value) in annotations)
 					{
 						newNode.Annotations[key] = value;
@@ -1632,18 +1635,18 @@ namespace AutomationTool
 						}
 
 						// Make sure none of the required arguments are missing
-						bool bHasMissingArguments = false;
+						bool hasMissingArguments = false;
 						for (int idx = 0; idx < macro.NumRequiredArguments; idx++)
 						{
 							if (arguments[idx] == null)
 							{
 								LogWarning(element, "Macro '{Name}' is missing argument '{ArgName}'", macro.Name, macro.ArgumentNameToIndex.First(x => x.Value == idx).Key);
-								bHasMissingArguments = true;
+								hasMissingArguments = true;
 							}
 						}
 
 						// Expand the function
-						if (!bHasMissingArguments)
+						if (!hasMissingArguments)
 						{
 							EnterScope();
 							foreach (KeyValuePair<string, int> pair in macro.ArgumentNameToIndex)
@@ -1702,8 +1705,8 @@ namespace AutomationTool
 				string[] reportNames = ReadListAttribute(element, "Reports");
 				string[] users = ReadListAttribute(element, "Users");
 				string[] submitters = ReadListAttribute(element, "Submitters");
-				bool? bWarnings = element.HasAttribute("Warnings") ? (bool?)ReadBooleanAttribute(element, "Warnings", true) : null;
-				bool bAbsolute = element.HasAttribute("Absolute") && ReadBooleanAttribute(element, "Absolute", true);
+				bool? warnings = element.HasAttribute("Warnings") ? (bool?)ReadBooleanAttribute(element, "Warnings", true) : null;
+				bool absolute = element.HasAttribute("Absolute") && ReadBooleanAttribute(element, "Absolute", true);
 
 				// Find the list of targets which are included, and recurse through all their dependencies
 				HashSet<BgNodeDef> nodes = new HashSet<BgNodeDef>();
@@ -1736,7 +1739,7 @@ namespace AutomationTool
 				{
 					if (users != null)
 					{
-						if (bAbsolute)
+						if (absolute)
 						{
 							node.NotifyUsers = new HashSet<string>(users);
 						}
@@ -1747,7 +1750,7 @@ namespace AutomationTool
 					}
 					if (submitters != null)
 					{
-						if (bAbsolute)
+						if (absolute)
 						{
 							node.NotifySubmitters = new HashSet<string>(submitters);
 						}
@@ -1756,9 +1759,9 @@ namespace AutomationTool
 							node.NotifySubmitters.UnionWith(submitters);
 						}
 					}
-					if (bWarnings.HasValue)
+					if (warnings.HasValue)
 					{
-						node.NotifyOnWarnings = bWarnings.Value;
+						node.NotifyOnWarnings = warnings.Value;
 					}
 				}
 
@@ -2096,28 +2099,28 @@ namespace AutomationTool
 		/// </summary>
 		/// <param name="element">Element to read the attribute from</param>
 		/// <param name="name">Name of the attribute</param>
-		/// <param name="bDefaultValue">Default value if the attribute is missing</param>
+		/// <param name="defaultValue">Default value if the attribute is missing</param>
 		/// <returns>The value of the attribute field</returns>
-		protected bool ReadBooleanAttribute(BgScriptElement element, string name, bool bDefaultValue)
+		protected bool ReadBooleanAttribute(BgScriptElement element, string name, bool defaultValue)
 		{
-			bool bResult = bDefaultValue;
+			bool result = defaultValue;
 			if (element.HasAttribute(name))
 			{
 				string value = ReadAttribute(element, name).Trim();
 				if (value.Equals("true", StringComparison.OrdinalIgnoreCase))
 				{
-					bResult = true;
+					result = true;
 				}
 				else if (value.Equals("false", StringComparison.OrdinalIgnoreCase))
 				{
-					bResult = false;
+					result = false;
 				}
 				else
 				{
 					LogError(element, "Invalid boolean value '{0}' - expected 'true' or 'false'", value);
 				}
 			}
-			return bResult;
+			return result;
 		}
 
 		/// <summary>
@@ -2216,7 +2219,7 @@ namespace AutomationTool
 			try
 			{
 				string text = ExpandProperties(element, element.GetAttribute("If"));
-				return await BgCondition.EvaluateAsync(text);
+				return await BgCondition.Evaluate(text);
 			}
 			catch (BgConditionException ex)
 			{
@@ -2266,9 +2269,9 @@ namespace AutomationTool
 		}
 
 		/// <inheritdoc/>
-		public static object GetNativePath(string Path)
+		public static object GetNativePath(string path)
 		{
-			return FileReference.Combine(Unreal.RootDirectory, Path).FullName;
+			return FileReference.Combine(Unreal.RootDirectory, path).FullName;
 		}
 	}
 }
\ No newline at end of file
diff --git a/Engine/Source/Programs/AutomationTool/BuildGraph/BgScriptWriter.cs b/Engine/Source/Programs/AutomationTool/BuildGraph/BgScriptWriter.cs
--- a/Engine/Source/Programs/AutomationTool/BuildGraph/BgScriptWriter.cs
+++ b/Engine/Source/Programs/AutomationTool/BuildGraph/BgScriptWriter.cs
@@ -2,17 +2,10 @@
 
 using System;
 using System.Collections.Generic;
-using System.Diagnostics.CodeAnalysis;
-using System.IO;
 using System.Linq;
-using System.Text;
-using System.Text.RegularExpressions;
-using System.Threading.Tasks;
 using System.Xml;
-using System.Xml.Schema;
 using EpicGames.BuildGraph;
 using EpicGames.Core;
-using Microsoft.Extensions.Logging;
 
 #nullable enable
 
diff --git a/Engine/Source/Programs/AutomationTool/BuildGraph/BuildGraph.Automation.csproj b/Engine/Source/Programs/AutomationTool/BuildGraph/BuildGraph.Automation.csproj
--- a/Engine/Source/Programs/AutomationTool/BuildGraph/BuildGraph.Automation.csproj
+++ b/Engine/Source/Programs/AutomationTool/BuildGraph/BuildGraph.Automation.csproj
@@ -1,6 +1,6 @@
 <Project Sdk="Microsoft.NET.Sdk">
   <Import Project="..\..\Shared\UnrealEngine.csproj.props" />
-  
+
   <PropertyGroup>
     <TargetFramework>net6.0</TargetFramework>
     <Configuration Condition=" '$(Configuration)' == '' ">Development</Configuration>
@@ -35,6 +35,7 @@
     <DocumentationFile></DocumentationFile>
     <TreatWarningsAsErrors>false</TreatWarningsAsErrors>
     <RunAnalyzersDuringBuild>True</RunAnalyzersDuringBuild>
+    <EnforceCodeStyleInBuild>True</EnforceCodeStyleInBuild>
     <GenerateDocumentationFile>True</GenerateDocumentationFile>
   </PropertyGroup>
   <ItemGroup>
