EpicGames.BuildGraph: Apply standard naming conventions to BuildGraph.Automation.

diff --git a/Engine/Source/Programs/AutomationTool/BuildGraph/BuildGraph.cs b/Engine/Source/Programs/AutomationTool/BuildGraph/BuildGraph.cs
--- a/Engine/Source/Programs/AutomationTool/BuildGraph/BuildGraph.cs
+++ b/Engine/Source/Programs/AutomationTool/BuildGraph/BuildGraph.cs
@@ -1,12 +1,5 @@
 // Copyright Epic Games, Inc. All Rights Reserved.
 
-using EpicGames.BuildGraph;
-using EpicGames.BuildGraph.Expressions;
-using EpicGames.Core;
-using EpicGames.MCP.Automation;
-using Microsoft.Extensions.Logging;
-using OpenTracing;
-using OpenTracing.Util;
 using System;
 using System.Collections.Generic;
 using System.IO;
@@ -16,6 +9,13 @@ using System.Runtime.InteropServices;
 using System.Text;
 using System.Threading.Tasks;
 using System.Xml;
+using EpicGames.BuildGraph;
+using EpicGames.BuildGraph.Expressions;
+using EpicGames.Core;
+using EpicGames.MCP.Automation;
+using Microsoft.Extensions.Logging;
+using OpenTracing;
+using OpenTracing.Util;
 using UnrealBuildBase;
 using UnrealBuildTool;
 
@@ -31,14 +31,14 @@ namespace AutomationTool
 	{
 		public abstract Type ParameterType { get; }
 
-		public ScriptTaskParameterBinding(string Name, Type ValueType, TaskParameterValidationType ValidationType, bool bOptional)
-			: base(Name, ValueType, ValidationType, bOptional)
+		public ScriptTaskParameterBinding(string name, Type valueType, TaskParameterValidationType validationType, bool optional)
+			: base(name, valueType, validationType, optional)
 		{
 		}
 
 		public abstract string GetXmlDocName();
-		public abstract object? GetValue(object SourceObject);
-		public abstract void SetValue(object TargetObject, object? Value);
+		public abstract object? GetValue(object sourceObject);
+		public abstract void SetValue(object targetObject, object? value);
 	}
 
 	class ScriptTaskParameterPropertyBinding : ScriptTaskParameterBinding
@@ -57,11 +57,11 @@ namespace AutomationTool
 		public override string GetXmlDocName()
 			=> "P:" + _propertyInfo.DeclaringType!.FullName + "." + Name;
 
-		public override object? GetValue(object SourceObject)
-			=> _propertyInfo.GetValue(SourceObject);
+		public override object? GetValue(object sourceObject)
+			=> _propertyInfo.GetValue(sourceObject);
 
-		public override void SetValue(object TargetObject, object? Value)
-			=> _propertyInfo.SetValue(TargetObject, Value);
+		public override void SetValue(object targetObject, object? value)
+			=> _propertyInfo.SetValue(targetObject, value);
 	}
 
 	class ScriptTaskParameterFieldBinding : ScriptTaskParameterBinding
@@ -80,11 +80,11 @@ namespace AutomationTool
 		public override string GetXmlDocName()
 			=> "F:" + _fieldInfo.DeclaringType!.FullName + "." + Name;
 
-		public override object? GetValue(object SourceObject)
-			=> _fieldInfo.GetValue(SourceObject);
+		public override object? GetValue(object sourceObject)
+			=> _fieldInfo.GetValue(sourceObject);
 
-		public override void SetValue(object TargetObject, object? Value)
-			=> _fieldInfo.SetValue(TargetObject, Value);
+		public override void SetValue(object targetObject, object? value)
+			=> _fieldInfo.SetValue(targetObject, value);
 	}
 
 	/// <summary>
@@ -110,48 +110,48 @@ namespace AutomationTool
 		/// <summary>
 		/// Constructor
 		/// </summary>
-		/// <param name="Name">Name of the task</param>
-		/// <param name="TaskClass">Task class to create</param>
-		/// <param name="ParametersClass">Class type of an object to be constructed and passed as an argument to the task class constructor</param>
-		public ScriptTaskBinding(string Name, Type TaskClass, Type ParametersClass)
-			: this(Name, TaskClass, ParametersClass, CreateParameters(ParametersClass))
+		/// <param name="name">Name of the task</param>
+		/// <param name="taskClass">Task class to create</param>
+		/// <param name="parametersClass">Class type of an object to be constructed and passed as an argument to the task class constructor</param>
+		public ScriptTaskBinding(string name, Type taskClass, Type parametersClass)
+			: this(name, taskClass, parametersClass, CreateParameters(parametersClass))
 		{
 		}
 
 		/// <summary>
 		/// Private constructor
 		/// </summary>
-		private ScriptTaskBinding(string Name, Type TaskClass, Type ParametersClass, List<ScriptTaskParameterBinding> Parameters)
-			: base(Name, Parameters.ConvertAll<BgScriptTaskParameter>(x => x))
+		private ScriptTaskBinding(string name, Type taskClass, Type parametersClass, List<ScriptTaskParameterBinding> parameters)
+			: base(name, parameters.ConvertAll<BgScriptTaskParameter>(x => x))
 		{
-			this.TaskClass = TaskClass;
-			this.ParametersClass = ParametersClass;
-			this.NameToParameter = Parameters.ToDictionary(x => x.Name, x => x);
+			this.TaskClass = taskClass;
+			this.ParametersClass = parametersClass;
+			this.NameToParameter = parameters.ToDictionary(x => x.Name, x => x);
 		}
 
-		static List<ScriptTaskParameterBinding> CreateParameters(Type ParametersClass)
+		static List<ScriptTaskParameterBinding> CreateParameters(Type parametersClass)
 		{
-			List<ScriptTaskParameterBinding> ScriptTaskParameters = new List<ScriptTaskParameterBinding>();
-			foreach (FieldInfo Field in ParametersClass.GetFields())
+			List<ScriptTaskParameterBinding> scriptTaskParameters = new List<ScriptTaskParameterBinding>();
+			foreach (FieldInfo field in parametersClass.GetFields())
 			{
-				if (Field.MemberType == MemberTypes.Field)
+				if (field.MemberType == MemberTypes.Field)
 				{
-					TaskParameterAttribute? ParameterAttribute = Field.GetCustomAttribute<TaskParameterAttribute>();
-					if (ParameterAttribute != null)
+					TaskParameterAttribute? parameterAttribute = field.GetCustomAttribute<TaskParameterAttribute>();
+					if (parameterAttribute != null)
 					{
-						ScriptTaskParameters.Add(new ScriptTaskParameterFieldBinding(Field.Name, Field, ParameterAttribute.ValidationType, ParameterAttribute.Optional));
+						scriptTaskParameters.Add(new ScriptTaskParameterFieldBinding(field.Name, field, parameterAttribute.ValidationType, parameterAttribute.Optional));
 					}
 				}
 			}
-			foreach (PropertyInfo Property in ParametersClass.GetProperties(BindingFlags.Public | BindingFlags.Instance))
+			foreach (PropertyInfo property in parametersClass.GetProperties(BindingFlags.Public | BindingFlags.Instance))
 			{
-				TaskParameterAttribute? ParameterAttribute = Property.GetCustomAttribute<TaskParameterAttribute>();
-				if (ParameterAttribute != null)
+				TaskParameterAttribute? parameterAttribute = property.GetCustomAttribute<TaskParameterAttribute>();
+				if (parameterAttribute != null)
 				{
-					ScriptTaskParameters.Add(new ScriptTaskParameterPropertyBinding(Property.Name, Property, ParameterAttribute.ValidationType, ParameterAttribute.Optional));
+					scriptTaskParameters.Add(new ScriptTaskParameterPropertyBinding(property.Name, property, parameterAttribute.ValidationType, parameterAttribute.Optional));
 				}
 			}
-			return ScriptTaskParameters;
+			return scriptTaskParameters;
 		}
 	}
 
@@ -188,18 +188,18 @@ namespace AutomationTool
 		/// <summary>
 		/// Constructor
 		/// </summary>
-		/// <param name="Branch">The current branch</param>
-		/// <param name="Change">Changelist being built</param>
-		/// <param name="CodeChange">Code changelist being built</param>
-		internal BgEnvironment(string? Branch, int? Change, int? CodeChange)
+		/// <param name="branch">The current branch</param>
+		/// <param name="change">Changelist being built</param>
+		/// <param name="codeChange">Code changelist being built</param>
+		internal BgEnvironment(string? branch, int? change, int? codeChange)
 		{
-			this.Stream = Branch ?? "Unknown";
-			this.Change = Change ?? 0;
-			this.CodeChange = CodeChange ?? 0;
+			this.Stream = branch ?? "Unknown";
+			this.Change = change ?? 0;
+			this.CodeChange = codeChange ?? 0;
 			this.IsBuildMachine = CommandUtils.IsBuildMachine;
 
-			ReadOnlyBuildVersion Version = ReadOnlyBuildVersion.Current;
-			this.EngineVersion = (Version.MajorVersion, Version.MinorVersion, Version.PatchVersion);
+			ReadOnlyBuildVersion version = ReadOnlyBuildVersion.Current;
+			this.EngineVersion = (version.MajorVersion, version.MinorVersion, version.PatchVersion);
 		}
 	}
 
@@ -268,221 +268,221 @@ namespace AutomationTool
 		public override async Task<ExitCode> ExecuteAsync()
 		{
 			// Parse the command line parameters
-			string ClassName = ParseParamValue("Class", null);
-			string ScriptFileName = ParseParamValue("Script", null);
-			string[] TargetNames = ParseParamValues("Target").SelectMany(x => x.Split(';', '+').Select(y => y.Trim()).Where(y => y.Length > 0)).ToArray();
-			string DocumentationFileName = ParseParamValue("Documentation", null);
-			string SchemaFileName = ParseParamValue("Schema", null);
-			string ImportSchemaFileName = ParseParamValue("ImportSchema", null);
-			string ExportFileName = ParseParamValue("Export", null);
-			string HordeExportFileName = ParseParamValue("HordeExport", null);
-			string PreprocessedFileName = ParseParamValue("Preprocess", null);
-			string SharedStorageDir = ParseParamValue("SharedStorageDir", null);
-			string SingleNodeName = ParseParamValue("SingleNode", null);
-			string TriggerName = ParseParamValue("Trigger", null);
-			string TokenSignature = ParseParamValue("TokenSignature", null);
-			bool bSkipTargetsWithoutTokens = ParseParam("SkipTargetsWithoutTokens");
-			bool bResume = SingleNodeName != null || ParseParam("Resume");
-			bool bListOnly = ParseParam("ListOnly");
-			bool bShowDiagnostics = ParseParam("ShowDiagnostics");
-			bool bWriteToSharedStorage = ParseParam("WriteToSharedStorage") || CommandUtils.IsBuildMachine;
-			bool bPublicTasksOnly = ParseParam("PublicTasksOnly");
-			bool bSkipValidation = ParseParam("SkipValidation");
-			string ReportName = ParseParamValue("ReportName", null);
-			string BranchOverride = ParseParamValue("Branch", null);
-
-			GraphPrintOptions PrintOptions = GraphPrintOptions.ShowCommandLineOptions;
+			string className = ParseParamValue("Class", null);
+			string scriptFileName = ParseParamValue("Script", null);
+			string[] targetNames = ParseParamValues("Target").SelectMany(x => x.Split(';', '+').Select(y => y.Trim()).Where(y => y.Length > 0)).ToArray();
+			string documentationFileName = ParseParamValue("Documentation", null);
+			string schemaFileName = ParseParamValue("Schema", null);
+			string importSchemaFileName = ParseParamValue("ImportSchema", null);
+			string exportFileName = ParseParamValue("Export", null);
+			string hordeExportFileName = ParseParamValue("HordeExport", null);
+			string preprocessedFileName = ParseParamValue("Preprocess", null);
+			string sharedStorageDir = ParseParamValue("SharedStorageDir", null);
+			string singleNodeName = ParseParamValue("SingleNode", null);
+			string triggerName = ParseParamValue("Trigger", null);
+			string tokenSignature = ParseParamValue("TokenSignature", null);
+			bool skipTargetsWithoutTokens = ParseParam("SkipTargetsWithoutTokens");
+			bool resume = singleNodeName != null || ParseParam("Resume");
+			bool listOnly = ParseParam("ListOnly");
+			bool showDiagnostics = ParseParam("ShowDiagnostics");
+			bool writeToSharedStorage = ParseParam("WriteToSharedStorage") || CommandUtils.IsBuildMachine;
+			bool publicTasksOnly = ParseParam("PublicTasksOnly");
+			bool skipValidation = ParseParam("SkipValidation");
+			string reportName = ParseParamValue("ReportName", null);
+			string branchOverride = ParseParamValue("Branch", null);
+
+			GraphPrintOptions printOptions = GraphPrintOptions.ShowCommandLineOptions;
 			if (ParseParam("ShowDeps"))
 			{
-				PrintOptions |= GraphPrintOptions.ShowDependencies;
+				printOptions |= GraphPrintOptions.ShowDependencies;
 			}
 			if (ParseParam("ShowNotifications"))
 			{
-				PrintOptions |= GraphPrintOptions.ShowNotifications;
+				printOptions |= GraphPrintOptions.ShowNotifications;
 			}
 
-			if (SchemaFileName == null && ParseParam("Schema"))
+			if (schemaFileName == null && ParseParam("Schema"))
 			{
-				SchemaFileName = FileReference.Combine(Unreal.EngineDirectory, "Build", "Graph", "Schema.xsd").FullName;
+				schemaFileName = FileReference.Combine(Unreal.EngineDirectory, "Build", "Graph", "Schema.xsd").FullName;
 			}
 
 			// Parse any specific nodes to clean
-			List<string> CleanNodes = new List<string>();
-			foreach (string NodeList in ParseParamValues("CleanNode"))
+			List<string> cleanNodes = new List<string>();
+			foreach (string nodeList in ParseParamValues("CleanNode"))
 			{
-				foreach (string NodeName in NodeList.Split('+', ';'))
+				foreach (string nodeName in nodeList.Split('+', ';'))
 				{
-					CleanNodes.Add(NodeName);
+					cleanNodes.Add(nodeName);
 				}
 			}
 
 			// Get the standard P4 properties, defaulting to the environment variables if not set. This allows setting p4-like properties without having a P4 connection.
-			string? Branch = P4Enabled ? P4Env.Branch : GetEnvVarOrNull(EnvVarNames.BuildRootP4);
-			int? Change = P4Enabled ? P4Env.Changelist : GetEnvVarIntOrNull(EnvVarNames.Changelist);
-			int? CodeChange = P4Enabled ? P4Env.CodeChangelist : GetEnvVarIntOrNull(EnvVarNames.CodeChangelist);
+			string? branch = P4Enabled ? P4Env.Branch : GetEnvVarOrNull(EnvVarNames.BuildRootP4);
+			int? change = P4Enabled ? P4Env.Changelist : GetEnvVarIntOrNull(EnvVarNames.Changelist);
+			int? codeChange = P4Enabled ? P4Env.CodeChangelist : GetEnvVarIntOrNull(EnvVarNames.CodeChangelist);
 
 			// Set up the standard properties which build scripts might need
-			Dictionary<string, string?> DefaultProperties = new Dictionary<string, string?>(StringComparer.InvariantCultureIgnoreCase);
-			DefaultProperties["Branch"] = Branch ?? "Unknown";
-			DefaultProperties["Depot"] = (Branch != null && Branch.StartsWith("//", StringComparison.Ordinal)) ? Branch.Substring(2).Split('/').First() : "Unknown";
-			DefaultProperties["EscapedBranch"] = String.IsNullOrEmpty(Branch) ? "Unknown" : CommandUtils.EscapePath(Branch);
-			DefaultProperties["Change"] = (Change ?? 0).ToString();
-			DefaultProperties["CodeChange"] = (CodeChange ?? 0).ToString();
-			DefaultProperties["IsBuildMachine"] = IsBuildMachine ? "true" : "false";
-			DefaultProperties["HostPlatform"] = HostPlatform.Current.HostEditorPlatform.ToString();
-			DefaultProperties["HostArchitecture"] = RuntimeInformation.ProcessArchitecture.ToString();
-			DefaultProperties["RestrictedFolderNames"] = String.Join(";", RestrictedFolder.GetNames());
-			DefaultProperties["RestrictedFolderFilter"] = String.Join(";", RestrictedFolder.GetNames().Select(x => String.Format(".../{0}/...", x)));
-			DefaultProperties["DataDrivenPlatforms"] = String.Join(";", DataDrivenPlatformInfo.GetAllPlatformInfos().Keys);
+			Dictionary<string, string?> defaultProperties = new Dictionary<string, string?>(StringComparer.InvariantCultureIgnoreCase);
+			defaultProperties["Branch"] = branch ?? "Unknown";
+			defaultProperties["Depot"] = (branch != null && branch.StartsWith("//", StringComparison.Ordinal)) ? branch.Substring(2).Split('/').First() : "Unknown";
+			defaultProperties["EscapedBranch"] = String.IsNullOrEmpty(branch) ? "Unknown" : CommandUtils.EscapePath(branch);
+			defaultProperties["Change"] = (change ?? 0).ToString();
+			defaultProperties["CodeChange"] = (codeChange ?? 0).ToString();
+			defaultProperties["IsBuildMachine"] = IsBuildMachine ? "true" : "false";
+			defaultProperties["HostPlatform"] = HostPlatform.Current.HostEditorPlatform.ToString();
+			defaultProperties["HostArchitecture"] = RuntimeInformation.ProcessArchitecture.ToString();
+			defaultProperties["RestrictedFolderNames"] = String.Join(";", RestrictedFolder.GetNames());
+			defaultProperties["RestrictedFolderFilter"] = String.Join(";", RestrictedFolder.GetNames().Select(x => String.Format(".../{0}/...", x)));
+			defaultProperties["DataDrivenPlatforms"] = String.Join(";", DataDrivenPlatformInfo.GetAllPlatformInfos().Keys);
 
 			// Look for overrides
-			if (!string.IsNullOrEmpty(BranchOverride))
+			if (!string.IsNullOrEmpty(branchOverride))
 			{
-				Logger.LogInformation("Overriding default branch '{Branch}' with '{BranchOverride}'", DefaultProperties["Branch"], BranchOverride);
-				DefaultProperties["Branch"] = BranchOverride;
-				DefaultProperties["EscapedBranch"] = CommandUtils.EscapePath(DefaultProperties["Branch"]);
+				Logger.LogInformation("Overriding default branch '{Branch}' with '{BranchOverride}'", defaultProperties["Branch"], branchOverride);
+				defaultProperties["Branch"] = branchOverride;
+				defaultProperties["EscapedBranch"] = CommandUtils.EscapePath(defaultProperties["Branch"]);
 			}
 
 			// Prevent expansion of the root directory if we're just preprocessing the output. They may vary by machine.
-			if (PreprocessedFileName == null)
+			if (preprocessedFileName == null)
 			{
-				DefaultProperties["RootDir"] = Unreal.RootDirectory.FullName;
+				defaultProperties["RootDir"] = Unreal.RootDirectory.FullName;
 			}
 			else
 			{
-				DefaultProperties["RootDir"] = null;
+				defaultProperties["RootDir"] = null;
 			}
 
 			// Attempt to read existing Build Version information
-			BuildVersion? Version;
-			if (BuildVersion.TryRead(BuildVersion.GetDefaultFileName(), out Version))
+			BuildVersion? version;
+			if (BuildVersion.TryRead(BuildVersion.GetDefaultFileName(), out version))
 			{
-				DefaultProperties["EngineMajorVersion"] = Version.MajorVersion.ToString();
-				DefaultProperties["EngineMinorVersion"] = Version.MinorVersion.ToString();
-				DefaultProperties["EnginePatchVersion"] = Version.PatchVersion.ToString();
-				DefaultProperties["EngineCompatibleChange"] = Version.CompatibleChangelist.ToString();
+				defaultProperties["EngineMajorVersion"] = version.MajorVersion.ToString();
+				defaultProperties["EngineMinorVersion"] = version.MinorVersion.ToString();
+				defaultProperties["EnginePatchVersion"] = version.PatchVersion.ToString();
+				defaultProperties["EngineCompatibleChange"] = version.CompatibleChangelist.ToString();
 			}
 
 			// If the -project flag is given, pass useful information into the graph
-			FileReference ProjectFile = ParseProjectParam();
-			if (ProjectFile != null)
+			FileReference projectFile = ParseProjectParam();
+			if (projectFile != null)
 			{
-				DefaultProperties["ProjectName"] = ProjectFile.GetFileNameWithoutExtension();
-				DefaultProperties["ProjectFile"] = ProjectFile.FullName;
-				DefaultProperties["ProjectDir"] = ProjectFile.Directory.FullName;
+				defaultProperties["ProjectName"] = projectFile.GetFileNameWithoutExtension();
+				defaultProperties["ProjectFile"] = projectFile.FullName;
+				defaultProperties["ProjectDir"] = projectFile.Directory.FullName;
 			}
 
 			// Add any additional custom arguments from the command line (of the form -Set:X=Y)
-			Dictionary<string, string> Arguments = new Dictionary<string, string>(StringComparer.InvariantCultureIgnoreCase);
-			foreach (string Param in Params)
+			Dictionary<string, string> arguments = new Dictionary<string, string>(StringComparer.InvariantCultureIgnoreCase);
+			foreach (string param in Params)
 			{
 				const string SetPrefix = "set:";
-				if (Param.StartsWith(SetPrefix, StringComparison.InvariantCultureIgnoreCase))
+				if (param.StartsWith(SetPrefix, StringComparison.InvariantCultureIgnoreCase))
 				{
-					int EqualsIdx = Param.IndexOf('=', StringComparison.Ordinal);
-					if (EqualsIdx >= 0)
+					int equalsIdx = param.IndexOf('=', StringComparison.Ordinal);
+					if (equalsIdx >= 0)
 					{
-						Arguments[Param.Substring(SetPrefix.Length, EqualsIdx - SetPrefix.Length)] = Param.Substring(EqualsIdx + 1);
+						arguments[param.Substring(SetPrefix.Length, equalsIdx - SetPrefix.Length)] = param.Substring(equalsIdx + 1);
 					}
 					else
 					{
-						Logger.LogWarning("Missing value for '{Arg0}'", Param.Substring(SetPrefix.Length));
+						Logger.LogWarning("Missing value for '{Arg0}'", param.Substring(SetPrefix.Length));
 					}
 				}
 
 				const string AppendPrefix = "append:";
-				if (Param.StartsWith(AppendPrefix, StringComparison.InvariantCultureIgnoreCase))
+				if (param.StartsWith(AppendPrefix, StringComparison.InvariantCultureIgnoreCase))
 				{
-					int EqualsIdx = Param.IndexOf('=', StringComparison.Ordinal);
-					if (EqualsIdx >= 0)
+					int equalsIdx = param.IndexOf('=', StringComparison.Ordinal);
+					if (equalsIdx >= 0)
 					{
-						string Property = Param.Substring(AppendPrefix.Length, EqualsIdx - AppendPrefix.Length);
-						string Value = Param.Substring(EqualsIdx + 1);
-						if (Arguments.ContainsKey(Property))
+						string property = param.Substring(AppendPrefix.Length, equalsIdx - AppendPrefix.Length);
+						string value = param.Substring(equalsIdx + 1);
+						if (arguments.ContainsKey(property))
 						{
-							Arguments[Property] = Arguments[Property] + ";" + Value;
+							arguments[property] = arguments[property] + ";" + value;
 						}
 						else
 						{
-							Arguments[Property] = Value;
+							arguments[property] = value;
 						}
 					}
 					else
 					{
-						Logger.LogWarning("Missing value for '{Arg0}'", Param.Substring(AppendPrefix.Length));
+						Logger.LogWarning("Missing value for '{Arg0}'", param.Substring(AppendPrefix.Length));
 					}
 				}
 			}
 
 			// Find all the tasks from the loaded assemblies
-			Dictionary<string, ScriptTaskBinding> NameToTask = new Dictionary<string, ScriptTaskBinding>();
-			if (!FindAvailableTasks(NameToTask, bPublicTasksOnly))
+			Dictionary<string, ScriptTaskBinding> nameToTask = new Dictionary<string, ScriptTaskBinding>();
+			if (!FindAvailableTasks(nameToTask, publicTasksOnly))
 			{
 				return ExitCode.Error_Unknown;
 			}
 
 			// Generate documentation
-			if (DocumentationFileName != null)
+			if (documentationFileName != null)
 			{
-				WriteDocumentation(NameToTask, new FileReference(DocumentationFileName));
+				WriteDocumentation(nameToTask, new FileReference(documentationFileName));
 				return ExitCode.Success;
 			}
 
 			// Create the graph
-			BgGraphDef? Graph;
-			if (ClassName != null)
+			BgGraphDef? graph;
+			if (className != null)
 			{
 				// Find all the graph builders
-				Dictionary<string, Type> NameToType = new Dictionary<string, Type>();
-				FindAvailableGraphs(NameToType);
+				Dictionary<string, Type> nameToType = new Dictionary<string, Type>();
+				FindAvailableGraphs(nameToType);
 
-				Type? BuilderType;
-				if (!NameToType.TryGetValue(ClassName, out BuilderType))
+				Type? builderType;
+				if (!nameToType.TryGetValue(className, out builderType))
 				{
-					Logger.LogError("Unable to find graph '{GraphName}'", ClassName);
+					Logger.LogError("Unable to find graph '{GraphName}'", className);
 					Logger.LogInformation("");
 					Logger.LogInformation("Available graphs:");
-					foreach (string Name in NameToType.Keys.OrderBy(x => x))
+					foreach (string name in nameToType.Keys.OrderBy(x => x))
 					{
-						Logger.LogInformation("  {GraphName}", Name);
+						Logger.LogInformation("  {GraphName}", name);
 					}
 					return ExitCode.Error_Unknown;
 				}
 
-				BgGraphBuilder Builder = (BgGraphBuilder)Activator.CreateInstance(BuilderType)!;
-				BgGraph GraphSpec = Builder.CreateGraph(new BgEnvironment(Branch, Change, CodeChange));
+				BgGraphBuilder builder = (BgGraphBuilder)Activator.CreateInstance(builderType)!;
+				BgGraph graphSpec = builder.CreateGraph(new BgEnvironment(branch, change, codeChange));
 
-				(byte[] Data, BgThunkDef[] Methods) = BgCompiler.Compile(GraphSpec);
+				(byte[] data, BgThunkDef[] methods) = BgCompiler.Compile(graphSpec);
 
-				BgInterpreter Interpreter = new BgInterpreter(Data, Methods, Arguments);
-				Interpreter.Disassemble(Logger);
-				Graph = ((BgObjectDef)Interpreter.Evaluate()).Deserialize<BgGraphExpressionDef>().ToGraphDef();
+				BgInterpreter interpreter = new BgInterpreter(data, methods, arguments);
+				interpreter.Disassemble(Logger);
+				graph = ((BgObjectDef)interpreter.Evaluate()).Deserialize<BgGraphExpressionDef>().ToGraphDef();
 			}
 			else
 			{
 				// Import schema if one is passed in
-				BgScriptSchema? Schema;
-				if (ImportSchemaFileName != null)
+				BgScriptSchema? schema;
+				if (importSchemaFileName != null)
 				{
-					Schema = BgScriptSchema.Import(FileReference.FromString(ImportSchemaFileName));
+					schema = BgScriptSchema.Import(FileReference.FromString(importSchemaFileName));
 				}
 				else
 				{
 					// Add any primitive types
-					List<(Type, ScriptSchemaStandardType)> PrimitiveTypes = new List<(Type, ScriptSchemaStandardType)>();
-					PrimitiveTypes.Add((typeof(FileReference), ScriptSchemaStandardType.BalancedString));
-					PrimitiveTypes.Add((typeof(DirectoryReference), ScriptSchemaStandardType.BalancedString));
-					PrimitiveTypes.Add((typeof(UnrealTargetPlatform), ScriptSchemaStandardType.BalancedString));
-					PrimitiveTypes.Add((typeof(MCPPlatform), ScriptSchemaStandardType.BalancedString));
+					List<(Type, ScriptSchemaStandardType)> primitiveTypes = new List<(Type, ScriptSchemaStandardType)>();
+					primitiveTypes.Add((typeof(FileReference), ScriptSchemaStandardType.BalancedString));
+					primitiveTypes.Add((typeof(DirectoryReference), ScriptSchemaStandardType.BalancedString));
+					primitiveTypes.Add((typeof(UnrealTargetPlatform), ScriptSchemaStandardType.BalancedString));
+					primitiveTypes.Add((typeof(MCPPlatform), ScriptSchemaStandardType.BalancedString));
 
 					// Create a schema for the given tasks
-					Schema = new BgScriptSchema(NameToTask.Values, PrimitiveTypes);
-					if (SchemaFileName != null)
+					schema = new BgScriptSchema(nameToTask.Values, primitiveTypes);
+					if (schemaFileName != null)
 					{
-						FileReference FullSchemaFileName = new FileReference(SchemaFileName);
-						Logger.LogInformation("Writing schema to {Arg0}...", FullSchemaFileName.FullName);
-						Schema.Export(FullSchemaFileName);
-						if (ScriptFileName == null)
+						FileReference fullSchemaFileName = new FileReference(schemaFileName);
+						Logger.LogInformation("Writing schema to {Arg0}...", fullSchemaFileName.FullName);
+						schema.Export(fullSchemaFileName);
+						if (scriptFileName == null)
 						{
 							return ExitCode.Success;
 						}
@@ -490,139 +490,139 @@ namespace AutomationTool
 				}
 
 				// Check there was a script specified
-				if (ScriptFileName == null)
+				if (scriptFileName == null)
 				{
 					Logger.LogError("Missing -Script= parameter for BuildGraph");
 					return ExitCode.Error_Unknown;
 				}
 
 				// Normalize the script filename
-				FileReference FullScriptFile = FileReference.Combine(Unreal.RootDirectory, ScriptFileName);
+				FileReference fullScriptFile = FileReference.Combine(Unreal.RootDirectory, scriptFileName);
 
 				// Read the script from disk
-				Graph = await BgScriptReader.ReadAsync(FullScriptFile, Arguments, DefaultProperties, Schema, Logger, SingleNodeName);
-				if (Graph == null)
+				graph = await BgScriptReader.ReadAsync(fullScriptFile, arguments, defaultProperties, schema, Logger, singleNodeName);
+				if (graph == null)
 				{
 					return ExitCode.Error_Unknown;
 				}
 			}
 
 			// Create the temp storage handler
-			DirectoryReference RootDir = new DirectoryReference(CommandUtils.CmdEnv.LocalRoot);
-			TempStorage Storage = new TempStorage(RootDir, DirectoryReference.Combine(RootDir, "Engine", "Saved", "BuildGraph"), (SharedStorageDir == null) ? null : new DirectoryReference(SharedStorageDir), bWriteToSharedStorage);
-			if (!bResume)
+			DirectoryReference rootDir = new DirectoryReference(CommandUtils.CmdEnv.LocalRoot);
+			TempStorage storage = new TempStorage(rootDir, DirectoryReference.Combine(rootDir, "Engine", "Saved", "BuildGraph"), (sharedStorageDir == null) ? null : new DirectoryReference(sharedStorageDir), writeToSharedStorage);
+			if (!resume)
 			{
-				Storage.CleanLocal();
+				storage.CleanLocal();
 			}
-			foreach (string CleanNode in CleanNodes)
+			foreach (string cleanNode in cleanNodes)
 			{
-				Storage.CleanLocalNode(CleanNode);
+				storage.CleanLocalNode(cleanNode);
 			}
 
 			// Convert the supplied target references into nodes 
-			HashSet<BgNodeDef> TargetNodes = new HashSet<BgNodeDef>();
-			if (TargetNames.Length == 0)
+			HashSet<BgNodeDef> targetNodes = new HashSet<BgNodeDef>();
+			if (targetNames.Length == 0)
 			{
-				if (!bListOnly && SingleNodeName == null)
+				if (!listOnly && singleNodeName == null)
 				{
 					Logger.LogError("Missing -Target= parameter for BuildGraph");
 					return ExitCode.Error_Unknown;
 				}
-				TargetNodes.UnionWith(Graph.Agents.SelectMany(x => x.Nodes));
+				targetNodes.UnionWith(graph.Agents.SelectMany(x => x.Nodes));
 			}
 			else
 			{
-				IEnumerable<string>? NodesToResolve = null;
+				IEnumerable<string>? nodesToResolve = null;
 
 				// If we're only building a single node and using a preprocessed reference we only need to try to resolve the references
 				// for that node.
-				if (SingleNodeName != null && PreprocessedFileName != null)
+				if (singleNodeName != null && preprocessedFileName != null)
 				{
-					NodesToResolve = new List<string> { SingleNodeName };
+					nodesToResolve = new List<string> { singleNodeName };
 				}
 				else
 				{
-					NodesToResolve = TargetNames;
+					nodesToResolve = targetNames;
 				}
 
-				foreach (string TargetName in NodesToResolve)
+				foreach (string targetName in nodesToResolve)
 				{
-					BgNodeDef[]? Nodes;
-					if (!Graph.TryResolveReference(TargetName, out Nodes))
+					BgNodeDef[]? nodes;
+					if (!graph.TryResolveReference(targetName, out nodes))
 					{
-						Logger.LogError("Target '{TargetName}' is not in graph", TargetName);
+						Logger.LogError("Target '{TargetName}' is not in graph", targetName);
 						return ExitCode.Error_Unknown;
 					}
-					TargetNodes.UnionWith(Nodes);
+					targetNodes.UnionWith(nodes);
 				}
 			}
 
 			// Try to acquire tokens for all the target nodes we want to build
-			if (TokenSignature != null)
+			if (tokenSignature != null)
 			{
 				// Find all the lock files
-				HashSet<FileReference> RequiredTokens = new HashSet<FileReference>(TargetNodes.SelectMany(x => x.RequiredTokens));
+				HashSet<FileReference> requiredTokens = new HashSet<FileReference>(targetNodes.SelectMany(x => x.RequiredTokens));
 
 				// List out all the required tokens
-				if (SingleNodeName == null)
+				if (singleNodeName == null)
 				{
 					Logger.LogInformation("Required tokens:");
-					foreach (BgNodeDef Node in TargetNodes)
+					foreach (BgNodeDef node in targetNodes)
 					{
-						foreach (FileReference RequiredToken in Node.RequiredTokens)
+						foreach (FileReference requiredToken in node.RequiredTokens)
 						{
-							Logger.LogInformation("  '{Node}' requires {RequiredToken}", Node, RequiredToken);
+							Logger.LogInformation("  '{Node}' requires {RequiredToken}", node, requiredToken);
 						}
 					}
 				}
 
 				// Try to create all the lock files
-				List<FileReference> CreatedTokens = new List<FileReference>();
-				if (!bListOnly)
+				List<FileReference> createdTokens = new List<FileReference>();
+				if (!listOnly)
 				{
-					CreatedTokens.AddRange(RequiredTokens.Where(x => WriteTokenFile(x, TokenSignature)));
+					createdTokens.AddRange(requiredTokens.Where(x => WriteTokenFile(x, tokenSignature)));
 				}
 
 				// Find all the tokens that we don't have
-				Dictionary<FileReference, string> MissingTokens = new Dictionary<FileReference, string>();
-				foreach (FileReference RequiredToken in RequiredTokens)
+				Dictionary<FileReference, string> missingTokens = new Dictionary<FileReference, string>();
+				foreach (FileReference requiredToken in requiredTokens)
 				{
-					string? CurrentOwner = ReadTokenFile(RequiredToken);
-					if (CurrentOwner != null && CurrentOwner != TokenSignature)
+					string? currentOwner = ReadTokenFile(requiredToken);
+					if (currentOwner != null && currentOwner != tokenSignature)
 					{
-						MissingTokens.Add(RequiredToken, CurrentOwner);
+						missingTokens.Add(requiredToken, currentOwner);
 					}
 				}
 
 				// If we want to skip all the nodes with missing locks, adjust the target nodes to account for it
-				if (MissingTokens.Count > 0)
+				if (missingTokens.Count > 0)
 				{
-					if (bSkipTargetsWithoutTokens)
+					if (skipTargetsWithoutTokens)
 					{
 						// Find all the nodes we're going to skip
-						HashSet<BgNodeDef> SkipNodes = new HashSet<BgNodeDef>();
-						foreach (IGrouping<string, FileReference> MissingTokensForBuild in MissingTokens.GroupBy(x => x.Value, x => x.Key))
+						HashSet<BgNodeDef> skipNodes = new HashSet<BgNodeDef>();
+						foreach (IGrouping<string, FileReference> missingTokensForBuild in missingTokens.GroupBy(x => x.Value, x => x.Key))
 						{
-							Logger.LogInformation("Skipping the following nodes due to {Arg0}:", MissingTokensForBuild.Key);
-							foreach (FileReference MissingToken in MissingTokensForBuild)
+							Logger.LogInformation("Skipping the following nodes due to {Arg0}:", missingTokensForBuild.Key);
+							foreach (FileReference missingToken in missingTokensForBuild)
 							{
-								foreach (BgNodeDef SkipNode in TargetNodes.Where(x => x.RequiredTokens.Contains(MissingToken) && SkipNodes.Add(x)))
+								foreach (BgNodeDef skipNode in targetNodes.Where(x => x.RequiredTokens.Contains(missingToken) && skipNodes.Add(x)))
 								{
-									Logger.LogInformation("    {SkipNode}", SkipNode);
+									Logger.LogInformation("    {SkipNode}", skipNode);
 								}
 							}
 						}
 
 						// Write a list of everything left over
-						if (SkipNodes.Count > 0)
+						if (skipNodes.Count > 0)
 						{
-							TargetNodes.ExceptWith(SkipNodes);
+							targetNodes.ExceptWith(skipNodes);
 							Logger.LogInformation("Remaining target nodes:");
-							foreach (BgNodeDef TargetNode in TargetNodes)
+							foreach (BgNodeDef targetNode in targetNodes)
 							{
-								Logger.LogInformation("    {TargetNode}", TargetNode);
+								Logger.LogInformation("    {TargetNode}", targetNode);
 							}
-							if (TargetNodes.Count == 0)
+							if (targetNodes.Count == 0)
 							{
 								Logger.LogInformation("    None.");
 							}
@@ -630,14 +630,14 @@ namespace AutomationTool
 					}
 					else
 					{
-						foreach (KeyValuePair<FileReference, string> Pair in MissingTokens)
+						foreach (KeyValuePair<FileReference, string> pair in missingTokens)
 						{
-							List<BgNodeDef> SkipNodes = TargetNodes.Where(x => x.RequiredTokens.Contains(Pair.Key)).ToList();
-							Logger.LogError("Cannot run {Arg0} due to previous build: {Arg1}", String.Join(", ", SkipNodes), Pair.Value);
+							List<BgNodeDef> skipNodes = targetNodes.Where(x => x.RequiredTokens.Contains(pair.Key)).ToList();
+							Logger.LogError("Cannot run {Arg0} due to previous build: {Arg1}", String.Join(", ", skipNodes), pair.Value);
 						}
-						foreach (FileReference CreatedToken in CreatedTokens)
+						foreach (FileReference createdToken in createdTokens)
 						{
-							FileReference.Delete(CreatedToken);
+							FileReference.Delete(createdToken);
 						}
 						return ExitCode.Error_Unknown;
 					}
@@ -645,110 +645,110 @@ namespace AutomationTool
 			}
 
 			// Cull the graph to include only those nodes
-			Graph.Select(TargetNodes);
+			graph.Select(targetNodes);
 
 			// If a report for the whole build was requested, insert it into the graph
-			if (ReportName != null)
+			if (reportName != null)
 			{
-				BgReport NewReport = new BgReport(ReportName);
-				NewReport.Nodes.UnionWith(Graph.Agents.SelectMany(x => x.Nodes));
-				Graph.NameToReport.Add(ReportName, NewReport);
+				BgReport newReport = new BgReport(reportName);
+				newReport.Nodes.UnionWith(graph.Agents.SelectMany(x => x.Nodes));
+				graph.NameToReport.Add(reportName, newReport);
 			}
 
 			// Export the graph for Horde
-			if (HordeExportFileName != null)
+			if (hordeExportFileName != null)
 			{
-				Graph.ExportForHorde(new FileReference(HordeExportFileName));
+				graph.ExportForHorde(new FileReference(hordeExportFileName));
 			}
 
 			// Write out the preprocessed script
-			if (PreprocessedFileName != null)
+			if (preprocessedFileName != null)
 			{
-				FileReference PreprocessedFileLocation = new FileReference(PreprocessedFileName);
-				Logger.LogInformation("Writing {PreprocessedFileLocation}...", PreprocessedFileLocation);
-				Graph.Write(PreprocessedFileLocation, (SchemaFileName != null) ? new FileReference(SchemaFileName) : null);
-				bListOnly = true;
+				FileReference preprocessedFileLocation = new FileReference(preprocessedFileName);
+				Logger.LogInformation("Writing {PreprocessedFileLocation}...", preprocessedFileLocation);
+				graph.Write(preprocessedFileLocation, (schemaFileName != null) ? new FileReference(schemaFileName) : null);
+				listOnly = true;
 			}
 
 			// If we're just building a single node, find it 
-			BgNodeDef? SingleNode = null;
-			if (SingleNodeName != null && !Graph.NameToNode.TryGetValue(SingleNodeName, out SingleNode))
+			BgNodeDef? singleNode = null;
+			if (singleNodeName != null && !graph.NameToNode.TryGetValue(singleNodeName, out singleNode))
 			{
-				Logger.LogError("Node '{SingleNodeName}' is not in the trimmed graph", SingleNodeName);
+				Logger.LogError("Node '{SingleNodeName}' is not in the trimmed graph", singleNodeName);
 				return ExitCode.Error_Unknown;
 			}
 
 			// If we just want to show the contents of the graph, do so and exit.
-			if (bListOnly)
+			if (listOnly)
 			{
-				HashSet<BgNodeDef> CompletedNodes = FindCompletedNodes(Graph, Storage);
-				Graph.Print(CompletedNodes, PrintOptions, Log.Logger);
+				HashSet<BgNodeDef> completedNodes = FindCompletedNodes(graph, storage);
+				graph.Print(completedNodes, printOptions, Log.Logger);
 			}
 
 			// Print out all the diagnostic messages which still apply, unless we're running a step as part of a build system or just listing the contents of the file. 
-			if (SingleNode == null && (!bListOnly || bShowDiagnostics))
+			if (singleNode == null && (!listOnly || showDiagnostics))
 			{
-				List<BgDiagnosticDef> Diagnostics = Graph.GetAllDiagnostics();
-				foreach (BgDiagnosticDef Diagnostic in Diagnostics)
+				List<BgDiagnosticDef> diagnostics = graph.GetAllDiagnostics();
+				foreach (BgDiagnosticDef diagnostic in diagnostics)
 				{
-					if (Diagnostic.Level == LogLevel.Information)
+					if (diagnostic.Level == LogLevel.Information)
 					{
-						Logger.LogInformation("{Arg0}({Arg1}): {Arg2}", Diagnostic.File, Diagnostic.Line, Diagnostic.Message);
+						Logger.LogInformation("{Arg0}({Arg1}): {Arg2}", diagnostic.File, diagnostic.Line, diagnostic.Message);
 					}
-					else if (Diagnostic.Level == LogLevel.Warning)
+					else if (diagnostic.Level == LogLevel.Warning)
 					{
-						Logger.LogWarning("{Arg0}({Arg1}): warning: {Arg2}", Diagnostic.File, Diagnostic.Line, Diagnostic.Message);
+						Logger.LogWarning("{Arg0}({Arg1}): warning: {Arg2}", diagnostic.File, diagnostic.Line, diagnostic.Message);
 					}
 					else
 					{
-						Logger.LogError("{Arg0}({Arg1}): error: {Arg2}", Diagnostic.File, Diagnostic.Line, Diagnostic.Message);
+						Logger.LogError("{Arg0}({Arg1}): error: {Arg2}", diagnostic.File, diagnostic.Line, diagnostic.Message);
 					}
 				}
-				if (Diagnostics.Any(x => x.Level == LogLevel.Error))
+				if (diagnostics.Any(x => x.Level == LogLevel.Error))
 				{
 					return ExitCode.Error_Unknown;
 				}
 			}
 
 			// Export the graph to a file
-			if (ExportFileName != null)
+			if (exportFileName != null)
 			{
-				HashSet<BgNodeDef> CompletedNodes = FindCompletedNodes(Graph, Storage);
-				Graph.Print(CompletedNodes, PrintOptions, Log.Logger);
-				Graph.Export(new FileReference(ExportFileName), CompletedNodes);
+				HashSet<BgNodeDef> completedNodes = FindCompletedNodes(graph, storage);
+				graph.Print(completedNodes, printOptions, Log.Logger);
+				graph.Export(new FileReference(exportFileName), completedNodes);
 				return ExitCode.Success;
 			}
 
 			// Create tasks for the entire graph
-			Dictionary<BgNodeDef, BgNodeExecutor> NodeToExecutor = new Dictionary<BgNodeDef, BgNodeExecutor>();
-			if (bSkipValidation && SingleNode != null)
+			Dictionary<BgNodeDef, BgNodeExecutor> nodeToExecutor = new Dictionary<BgNodeDef, BgNodeExecutor>();
+			if (skipValidation && singleNode != null)
 			{
-				if (!BindNodes(SingleNode, NameToTask, Graph.TagNameToNodeOutput, NodeToExecutor))
+				if (!BindNodes(singleNode, nameToTask, graph.TagNameToNodeOutput, nodeToExecutor))
 				{
 					return ExitCode.Error_Unknown;
 				}
 			}
 			else
 			{
-				if (!BindNodes(Graph, NameToTask, NodeToExecutor))
+				if (!BindNodes(graph, nameToTask, nodeToExecutor))
 				{
 					return ExitCode.Error_Unknown;
 				}
 			}
 
-			// Execute the command
-			if (!bListOnly)
+			// ExecuteAsync the command
+			if (!listOnly)
 			{
-				if (SingleNode != null)
+				if (singleNode != null)
 				{
-					if (!await BuildNodeAsync(Graph, SingleNode, NodeToExecutor, Storage, bWithBanner: true))
+					if (!await BuildNodeAsync(graph, singleNode, nodeToExecutor, storage, withBanner: true))
 					{
 						return ExitCode.Error_Unknown;
 					}
 				}
 				else
 				{
-					if (!await BuildAllNodesAsync(Graph, NodeToExecutor, Storage))
+					if (!await BuildAllNodesAsync(graph, nodeToExecutor, storage))
 					{
 						return ExitCode.Error_Unknown;
 					}
@@ -757,31 +757,31 @@ namespace AutomationTool
 			return ExitCode.Success;
 		}
 
-		static bool BindNodes(BgGraphDef Graph, Dictionary<string, ScriptTaskBinding> NameToTask, Dictionary<BgNodeDef, BgNodeExecutor> NodeToExecutor)
+		static bool BindNodes(BgGraphDef graph, Dictionary<string, ScriptTaskBinding> nameToTask, Dictionary<BgNodeDef, BgNodeExecutor> nodeToExecutor)
 		{
-			bool bResult = true;
-			foreach (BgAgentDef Agent in Graph.Agents)
+			bool result = true;
+			foreach (BgAgentDef agent in graph.Agents)
 			{
-				foreach (BgNodeDef Node in Agent.Nodes)
+				foreach (BgNodeDef node in agent.Nodes)
 				{
-					bResult &= BindNodes(Node, NameToTask, Graph.TagNameToNodeOutput, NodeToExecutor);
+					result &= BindNodes(node, nameToTask, graph.TagNameToNodeOutput, nodeToExecutor);
 				}
 			}
-			return bResult;
+			return result;
 		}
 
-		static bool BindNodes(BgNodeDef Node, Dictionary<string, ScriptTaskBinding> NameToTask, Dictionary<string, BgNodeOutput> TagNameToNodeOutput, Dictionary<BgNodeDef, BgNodeExecutor> NodeToExecutor)
+		static bool BindNodes(BgNodeDef node, Dictionary<string, ScriptTaskBinding> nameToTask, Dictionary<string, BgNodeOutput> tagNameToNodeOutput, Dictionary<BgNodeDef, BgNodeExecutor> nodeToExecutor)
 		{
-			if (Node is BgScriptNode ScriptNode)
+			if (node is BgScriptNode scriptNode)
 			{
-				BgScriptNodeExecutor executor = new BgScriptNodeExecutor(ScriptNode);
-				NodeToExecutor[Node] = executor;
-				return executor.Bind(NameToTask, TagNameToNodeOutput, Logger);
+				BgScriptNodeExecutor executor = new BgScriptNodeExecutor(scriptNode);
+				nodeToExecutor[node] = executor;
+				return executor.Bind(nameToTask, tagNameToNodeOutput, Logger);
 			}
-			else if (Node is BgNodeExpressionDef BytecodeNode)
+			else if (node is BgNodeExpressionDef bytecodeNode)
 			{
-				BgBytecodeNodeExecutor executor = new BgBytecodeNodeExecutor(BytecodeNode);
-				NodeToExecutor[Node] = executor;
+				BgBytecodeNodeExecutor executor = new BgBytecodeNodeExecutor(bytecodeNode);
+				nodeToExecutor[node] = executor;
 				return BgBytecodeNodeExecutor.Bind(Logger);
 			}
 			else
@@ -790,26 +790,26 @@ namespace AutomationTool
 			}
 		}
 
-		static void FindAvailableGraphs(Dictionary<string, Type> NameToType)
+		static void FindAvailableGraphs(Dictionary<string, Type> nameToType)
 		{
-			foreach (Assembly LoadedAssembly in ScriptManager.AllScriptAssemblies)
+			foreach (Assembly loadedAssembly in ScriptManager.AllScriptAssemblies)
 			{
-				Type[] Types;
+				Type[] types;
 				try
 				{
-					Types = LoadedAssembly.GetTypes();
+					types = loadedAssembly.GetTypes();
 				}
 				catch (ReflectionTypeLoadException ex)
 				{
-					Logger.LogWarning(ex, "Exception {Ex} while trying to get types from assembly {LoadedAssembly}. LoaderExceptions: {Arg2}", ex, LoadedAssembly, string.Join("\n", ex.LoaderExceptions.Select(x => x?.Message)));
+					Logger.LogWarning(ex, "Exception {Ex} while trying to get types from assembly {LoadedAssembly}. LoaderExceptions: {Arg2}", ex, loadedAssembly, string.Join("\n", ex.LoaderExceptions.Select(x => x?.Message)));
 					continue;
 				}
 
-				foreach (Type Type in Types)
+				foreach (Type type in types)
 				{
-					if (Type.IsSubclassOf(typeof(BgGraphBuilder)))
+					if (type.IsSubclassOf(typeof(BgGraphBuilder)))
 					{
-						NameToType.Add(Type.Name, Type);
+						nameToType.Add(type.Name, type);
 					}
 				}
 			}
@@ -818,44 +818,44 @@ namespace AutomationTool
 		/// <summary>
 		/// Find all the tasks which are available from the loaded assemblies
 		/// </summary>
-		/// <param name="NameToTask">Mapping from task name to information about how to serialize it</param>
-		/// <param name="bPublicTasksOnly">Whether to include just public tasks, or all the tasks in any loaded assemblies</param>
-		static bool FindAvailableTasks(Dictionary<string, ScriptTaskBinding> NameToTask, bool bPublicTasksOnly)
+		/// <param name="nameToTask">Mapping from task name to information about how to serialize it</param>
+		/// <param name="publicTasksOnly">Whether to include just public tasks, or all the tasks in any loaded assemblies</param>
+		static bool FindAvailableTasks(Dictionary<string, ScriptTaskBinding> nameToTask, bool publicTasksOnly)
 		{
-			IEnumerable<Assembly> LoadedScriptAssemblies = ScriptManager.AllScriptAssemblies;
+			IEnumerable<Assembly> loadedScriptAssemblies = ScriptManager.AllScriptAssemblies;
 
-			if (bPublicTasksOnly)
+			if (publicTasksOnly)
 			{
-				LoadedScriptAssemblies = LoadedScriptAssemblies.Where(x => IsPublicAssembly(new FileReference(x.Location)));
+				loadedScriptAssemblies = loadedScriptAssemblies.Where(x => IsPublicAssembly(new FileReference(x.Location)));
 			}
-			foreach (Assembly LoadedAssembly in LoadedScriptAssemblies)
+			foreach (Assembly loadedAssembly in loadedScriptAssemblies)
 			{
-				Type[] Types;
+				Type[] types;
 				try
 				{
-					Types = LoadedAssembly.GetTypes();
+					types = loadedAssembly.GetTypes();
 				}
 				catch (ReflectionTypeLoadException ex)
 				{
-					Logger.LogWarning(ex, "Exception {Ex} while trying to get types from assembly {LoadedAssembly}. LoaderExceptions: {Arg2}", ex, LoadedAssembly, string.Join("\n", ex.LoaderExceptions.Select(x => x?.Message)));
+					Logger.LogWarning(ex, "Exception {Ex} while trying to get types from assembly {LoadedAssembly}. LoaderExceptions: {Arg2}", ex, loadedAssembly, string.Join("\n", ex.LoaderExceptions.Select(x => x?.Message)));
 					continue;
 				}
 
-				foreach (Type Type in Types)
+				foreach (Type type in types)
 				{
-					foreach (TaskElementAttribute ElementAttribute in Type.GetCustomAttributes<TaskElementAttribute>())
+					foreach (TaskElementAttribute elementAttribute in type.GetCustomAttributes<TaskElementAttribute>())
 					{
-						if (!Type.IsSubclassOf(typeof(BgTaskImpl)))
+						if (!type.IsSubclassOf(typeof(BgTaskImpl)))
 						{
-							Logger.LogError("Class '{Arg0}' has TaskElementAttribute, but is not derived from 'BgTaskImpl'", Type.Name);
+							Logger.LogError("Class '{Arg0}' has TaskElementAttribute, but is not derived from 'BgTaskImpl'", type.Name);
 							return false;
 						}
-						if (NameToTask.ContainsKey(ElementAttribute.Name))
+						if (nameToTask.ContainsKey(elementAttribute.Name))
 						{
-							Logger.LogError("Found multiple handlers for task elements called '{Arg0}'", ElementAttribute.Name);
+							Logger.LogError("Found multiple handlers for task elements called '{Arg0}'", elementAttribute.Name);
 							return false;
 						}
-						NameToTask.Add(ElementAttribute.Name, new ScriptTaskBinding(ElementAttribute.Name, Type, ElementAttribute.ParametersType));
+						nameToTask.Add(elementAttribute.Name, new ScriptTaskBinding(elementAttribute.Name, type, elementAttribute.ParametersType));
 					}
 				}
 			}
@@ -866,19 +866,19 @@ namespace AutomationTool
 		/// Reads the contents of the given token
 		/// </summary>
 		/// <returns>Contents of the token, or null if it does not exist</returns>
-		public static string? ReadTokenFile(FileReference Location)
+		public static string? ReadTokenFile(FileReference location)
 		{
-			return FileReference.Exists(Location) ? File.ReadAllText(Location.FullName) : null;
+			return FileReference.Exists(location) ? File.ReadAllText(location.FullName) : null;
 		}
 
 		/// <summary>
 		/// Attempts to write an owner to a token file transactionally
 		/// </summary>
 		/// <returns>True if the lock was acquired, false otherwise</returns>
-		public static bool WriteTokenFile(FileReference Location, string Signature)
+		public static bool WriteTokenFile(FileReference location, string signature)
 		{
 			// Check it doesn't already exist
-			if (FileReference.Exists(Location))
+			if (FileReference.Exists(location))
 			{
 				return false;
 			}
@@ -886,30 +886,30 @@ namespace AutomationTool
 			// Make sure the directory exists
 			try
 			{
-				DirectoryReference.CreateDirectory(Location.Directory);
+				DirectoryReference.CreateDirectory(location.Directory);
 			}
-			catch (Exception Ex)
+			catch (Exception ex)
 			{
-				throw new AutomationException(Ex, "Unable to create '{0}'", Location.Directory);
+				throw new AutomationException(ex, "Unable to create '{0}'", location.Directory);
 			}
 
 			// Create a temp file containing the owner name
-			string TempFileName;
-			for (int Idx = 0; ; Idx++)
+			string tempFileName;
+			for (int idx = 0; ; idx++)
 			{
-				TempFileName = String.Format("{0}.{1}.tmp", Location.FullName, Idx);
+				tempFileName = String.Format("{0}.{1}.tmp", location.FullName, idx);
 				try
 				{
-					byte[] Bytes = Encoding.UTF8.GetBytes(Signature);
-					using (FileStream Stream = File.Open(TempFileName, FileMode.CreateNew, FileAccess.Write, FileShare.None))
+					byte[] bytes = Encoding.UTF8.GetBytes(signature);
+					using (FileStream stream = File.Open(tempFileName, FileMode.CreateNew, FileAccess.Write, FileShare.None))
 					{
-						Stream.Write(Bytes, 0, Bytes.Length);
+						stream.Write(bytes, 0, bytes.Length);
 					}
 					break;
 				}
 				catch (IOException)
 				{
-					if (!File.Exists(TempFileName))
+					if (!File.Exists(tempFileName))
 					{
 						throw;
 					}
@@ -919,12 +919,12 @@ namespace AutomationTool
 			// Try to move the temporary file into place. 
 			try
 			{
-				File.Move(TempFileName, Location.FullName);
+				File.Move(tempFileName, location.FullName);
 				return true;
 			}
 			catch
 			{
-				if (!File.Exists(TempFileName))
+				if (!File.Exists(tempFileName))
 				{
 					throw;
 				}
@@ -935,15 +935,15 @@ namespace AutomationTool
 		/// <summary>
 		/// Checks whether the given assembly is a publically distributed engine assembly.
 		/// </summary>
-		/// <param name="File">Assembly location</param>
+		/// <param name="file">Assembly location</param>
 		/// <returns>True if the assembly is distributed publically</returns>
-		static bool IsPublicAssembly(FileReference File)
+		static bool IsPublicAssembly(FileReference file)
 		{
-			DirectoryReference EngineDirectory = Unreal.EngineDirectory;
-			if (File.IsUnderDirectory(EngineDirectory))
+			DirectoryReference engineDirectory = Unreal.EngineDirectory;
+			if (file.IsUnderDirectory(engineDirectory))
 			{
-				string[] PathFragments = File.MakeRelativeTo(EngineDirectory).Split(Path.DirectorySeparatorChar, Path.AltDirectorySeparatorChar);
-				if (PathFragments.All(x => !x.Equals("NotForLicensees", StringComparison.OrdinalIgnoreCase) && !x.Equals("NoRedist", StringComparison.OrdinalIgnoreCase)))
+				string[] pathFragments = file.MakeRelativeTo(engineDirectory).Split(Path.DirectorySeparatorChar, Path.AltDirectorySeparatorChar);
+				if (pathFragments.All(x => !x.Equals("NotForLicensees", StringComparison.OrdinalIgnoreCase) && !x.Equals("NoRedist", StringComparison.OrdinalIgnoreCase)))
 				{
 					return true;
 				}
@@ -954,55 +954,55 @@ namespace AutomationTool
 		/// <summary>
 		/// Find all the nodes in the graph which are already completed
 		/// </summary>
-		/// <param name="Graph">The graph instance</param>
-		/// <param name="Storage">The temp storage backend which stores the shared state</param>
-		static HashSet<BgNodeDef> FindCompletedNodes(BgGraphDef Graph, TempStorage Storage)
+		/// <param name="graph">The graph instance</param>
+		/// <param name="storage">The temp storage backend which stores the shared state</param>
+		static HashSet<BgNodeDef> FindCompletedNodes(BgGraphDef graph, TempStorage storage)
 		{
-			HashSet<BgNodeDef> CompletedNodes = new HashSet<BgNodeDef>();
-			foreach (BgNodeDef Node in Graph.Agents.SelectMany(x => x.Nodes))
+			HashSet<BgNodeDef> completedNodes = new HashSet<BgNodeDef>();
+			foreach (BgNodeDef node in graph.Agents.SelectMany(x => x.Nodes))
 			{
-				if (Storage.IsComplete(Node.Name))
+				if (storage.IsComplete(node.Name))
 				{
-					CompletedNodes.Add(Node);
+					completedNodes.Add(node);
 				}
 			}
-			return CompletedNodes;
+			return completedNodes;
 		}
 
 		/// <summary>
 		/// Builds all the nodes in the graph
 		/// </summary>
-		/// <param name="Graph">The graph instance</param>
-		/// <param name="NodeToExecutor">Map from node to executor</param>
-		/// <param name="Storage">The temp storage backend which stores the shared state</param>
+		/// <param name="graph">The graph instance</param>
+		/// <param name="nodeToExecutor">Map from node to executor</param>
+		/// <param name="storage">The temp storage backend which stores the shared state</param>
 		/// <returns>True if everything built successfully</returns>
-		async Task<bool> BuildAllNodesAsync(BgGraphDef Graph, Dictionary<BgNodeDef, BgNodeExecutor> NodeToExecutor, TempStorage Storage)
+		async Task<bool> BuildAllNodesAsync(BgGraphDef graph, Dictionary<BgNodeDef, BgNodeExecutor> nodeToExecutor, TempStorage storage)
 		{
 			// Build a flat list of nodes to execute, in order
-			BgNodeDef[] NodesToExecute = Graph.Agents.SelectMany(x => x.Nodes).ToArray();
+			BgNodeDef[] nodesToExecute = graph.Agents.SelectMany(x => x.Nodes).ToArray();
 
 			// Check the integrity of any local nodes that have been completed. It's common to run formal builds locally between regular development builds, so we may have 
 			// stale local state. Rather than failing later, detect and clean them up now.
-			HashSet<BgNodeDef> CleanedNodes = new HashSet<BgNodeDef>();
-			foreach (BgNodeDef NodeToExecute in NodesToExecute)
+			HashSet<BgNodeDef> cleanedNodes = new HashSet<BgNodeDef>();
+			foreach (BgNodeDef nodeToExecute in nodesToExecute)
 			{
-				FileFilter IgnoreModifiedFilter = new FileFilter(NodeToExecute.IgnoreModified);
-				if (NodeToExecute.InputDependencies.Any(x => CleanedNodes.Contains(x)) || !Storage.CheckLocalIntegrity(NodeToExecute.Name, NodeToExecute.Outputs.Select(x => x.TagName), IgnoreModifiedFilter))
+				FileFilter ignoreModifiedFilter = new FileFilter(nodeToExecute.IgnoreModified);
+				if (nodeToExecute.InputDependencies.Any(x => cleanedNodes.Contains(x)) || !storage.CheckLocalIntegrity(nodeToExecute.Name, nodeToExecute.Outputs.Select(x => x.TagName), ignoreModifiedFilter))
 				{
-					Storage.CleanLocalNode(NodeToExecute.Name);
-					CleanedNodes.Add(NodeToExecute);
+					storage.CleanLocalNode(nodeToExecute.Name);
+					cleanedNodes.Add(nodeToExecute);
 				}
 			}
 
-			// Execute them in order
-			int NodeIdx = 0;
-			foreach (BgNodeDef NodeToExecute in NodesToExecute)
+			// ExecuteAsync them in order
+			int nodeIdx = 0;
+			foreach (BgNodeDef nodeToExecute in nodesToExecute)
 			{
-				Logger.LogInformation("****** [{Arg0}/{Arg1}] {Arg2}", ++NodeIdx, NodesToExecute.Length, NodeToExecute.Name);
-				if (!Storage.IsComplete(NodeToExecute.Name))
+				Logger.LogInformation("****** [{Arg0}/{Arg1}] {Arg2}", ++nodeIdx, nodesToExecute.Length, nodeToExecute.Name);
+				if (!storage.IsComplete(nodeToExecute.Name))
 				{
 					Logger.LogInformation("");
-					if (!await BuildNodeAsync(Graph, NodeToExecute, NodeToExecutor, Storage, bWithBanner: false))
+					if (!await BuildNodeAsync(graph, nodeToExecute, nodeToExecutor, storage, withBanner: false))
 					{
 						return false;
 					}
@@ -1024,8 +1024,8 @@ namespace AutomationTool
 			{
 				_logger = logger;
 
-				string? CleanupScriptEnvVar = Environment.GetEnvironmentVariable(CustomTask.CleanupScriptEnvVarName);
-				if (String.IsNullOrEmpty(CleanupScriptEnvVar))
+				string? cleanupScriptEnvVar = Environment.GetEnvironmentVariable(CustomTask.CleanupScriptEnvVarName);
+				if (String.IsNullOrEmpty(cleanupScriptEnvVar))
 				{
 					string extension = RuntimeInformation.IsOSPlatform(OSPlatform.Windows) ? "bat" : "sh";
 					_scriptFile = FileReference.Combine(Unreal.EngineDirectory, "Intermediate", $"OnExit.{extension}");
@@ -1058,162 +1058,162 @@ namespace AutomationTool
 		/// <summary>
 		/// Build a node
 		/// </summary>
-		/// <param name="Graph">The graph to which the node belongs. Used to determine which outputs need to be transferred to temp storage.</param>
-		/// <param name="Node">The node to build</param>
-		/// <param name="NodeToExecutor">Map from node to executor</param>
-		/// <param name="Storage">The temp storage backend which stores the shared state</param>
-		/// <param name="bWithBanner">Whether to write a banner before and after this node's log output</param>
+		/// <param name="graph">The graph to which the node belongs. Used to determine which outputs need to be transferred to temp storage.</param>
+		/// <param name="node">The node to build</param>
+		/// <param name="nodeToExecutor">Map from node to executor</param>
+		/// <param name="storage">The temp storage backend which stores the shared state</param>
+		/// <param name="withBanner">Whether to write a banner before and after this node's log output</param>
 		/// <returns>True if the node built successfully, false otherwise.</returns>
-		async Task<bool> BuildNodeAsync(BgGraphDef Graph, BgNodeDef Node, Dictionary<BgNodeDef, BgNodeExecutor> NodeToExecutor, TempStorage Storage, bool bWithBanner)
+		async Task<bool> BuildNodeAsync(BgGraphDef graph, BgNodeDef node, Dictionary<BgNodeDef, BgNodeExecutor> nodeToExecutor, TempStorage storage, bool withBanner)
 		{
-			DirectoryReference RootDir = new DirectoryReference(CommandUtils.CmdEnv.LocalRoot);
+			DirectoryReference rootDir = new DirectoryReference(CommandUtils.CmdEnv.LocalRoot);
 
 			// Register something to execute cleanup commands
-			using CleanupScriptRunner CleanupRunner = new CleanupScriptRunner(Logger);
+			using CleanupScriptRunner cleanupRunner = new CleanupScriptRunner(Logger);
 
 			// Create a filter for modified files that should be ignored
-			FileFilter IgnoreModifiedFilter = new FileFilter(Node.IgnoreModified);
+			FileFilter ignoreModifiedFilter = new FileFilter(node.IgnoreModified);
 
 			// Create the mapping of tag names to file sets
-			Dictionary<string, HashSet<FileReference>> TagNameToFileSet = new Dictionary<string, HashSet<FileReference>>();
+			Dictionary<string, HashSet<FileReference>> tagNameToFileSet = new Dictionary<string, HashSet<FileReference>>();
 
 			// Read all the input tags for this node, and build a list of referenced input storage blocks
-			HashSet<TempStorageBlock> InputStorageBlocks = new HashSet<TempStorageBlock>();
-			foreach (BgNodeOutput Input in Node.Inputs)
+			HashSet<TempStorageBlock> inputStorageBlocks = new HashSet<TempStorageBlock>();
+			foreach (BgNodeOutput input in node.Inputs)
 			{
-				TempStorageTagManifest FileList = Storage.ReadFileList(Input.ProducingNode.Name, Input.TagName) ?? throw new InvalidOperationException();
-				TagNameToFileSet[Input.TagName] = FileList.ToFileSet(RootDir);
-				InputStorageBlocks.UnionWith(FileList.Blocks);
+				TempStorageTagManifest fileList = storage.ReadFileList(input.ProducingNode.Name, input.TagName) ?? throw new InvalidOperationException();
+				tagNameToFileSet[input.TagName] = fileList.ToFileSet(rootDir);
+				inputStorageBlocks.UnionWith(fileList.Blocks);
 			}
 
 			// Read the manifests for all the input storage blocks
-			Dictionary<TempStorageBlock, TempStorageManifest> InputManifests = new Dictionary<TempStorageBlock, TempStorageManifest>();
-			using (IScope Scope = GlobalTracer.Instance.BuildSpan("TempStorage").WithTag("resource", "read").StartActive())
+			Dictionary<TempStorageBlock, TempStorageManifest> inputManifests = new Dictionary<TempStorageBlock, TempStorageManifest>();
+			using (IScope scope = GlobalTracer.Instance.BuildSpan("TempStorage").WithTag("resource", "read").StartActive())
 			{
-				Scope.Span.SetTag("blocks", InputStorageBlocks.Count);
-				foreach (TempStorageBlock InputStorageBlock in InputStorageBlocks)
+				scope.Span.SetTag("blocks", inputStorageBlocks.Count);
+				foreach (TempStorageBlock inputStorageBlock in inputStorageBlocks)
 				{
-					TempStorageManifest Manifest = Storage.Retrieve(InputStorageBlock.NodeName, InputStorageBlock.OutputName, IgnoreModifiedFilter);
-					InputManifests[InputStorageBlock] = Manifest;
+					TempStorageManifest manifest = storage.Retrieve(inputStorageBlock.NodeName, inputStorageBlock.OutputName, ignoreModifiedFilter);
+					inputManifests[inputStorageBlock] = manifest;
 				}
-				Scope.Span.SetTag("size", InputManifests.Sum(x => x.Value.GetTotalSize()));
+				scope.Span.SetTag("size", inputManifests.Sum(x => x.Value.GetTotalSize()));
 			}
 
 			// Read all the input storage blocks, keeping track of which block each file came from
-			Dictionary<FileReference, TempStorageBlock> FileToStorageBlock = new Dictionary<FileReference, TempStorageBlock>();
-			foreach (KeyValuePair<TempStorageBlock, TempStorageManifest> Pair in InputManifests)
+			Dictionary<FileReference, TempStorageBlock> fileToStorageBlock = new Dictionary<FileReference, TempStorageBlock>();
+			foreach (KeyValuePair<TempStorageBlock, TempStorageManifest> pair in inputManifests)
 			{
-				TempStorageBlock InputStorageBlock = Pair.Key;
-				foreach (FileReference File in Pair.Value.Files.Select(x => x.ToFileReference(RootDir)))
+				TempStorageBlock inputStorageBlock = pair.Key;
+				foreach (FileReference file in pair.Value.Files.Select(x => x.ToFileReference(rootDir)))
 				{
-					TempStorageBlock? CurrentStorageBlock;
-					if (FileToStorageBlock.TryGetValue(File, out CurrentStorageBlock) && !TempStorage.IsDuplicateBuildProduct(File))
+					TempStorageBlock? currentStorageBlock;
+					if (fileToStorageBlock.TryGetValue(file, out currentStorageBlock) && !TempStorage.IsDuplicateBuildProduct(file))
 					{
-						Logger.LogError("File '{File}' was produced by {InputStorageBlock} and {CurrentStorageBlock}", File, InputStorageBlock, CurrentStorageBlock);
+						Logger.LogError("File '{File}' was produced by {InputStorageBlock} and {CurrentStorageBlock}", file, inputStorageBlock, currentStorageBlock);
 					}
-					FileToStorageBlock[File] = InputStorageBlock;
+					fileToStorageBlock[file] = inputStorageBlock;
 				}
 			}
 
 			// Add placeholder outputs for the current node
-			foreach (BgNodeOutput Output in Node.Outputs)
+			foreach (BgNodeOutput output in node.Outputs)
 			{
-				TagNameToFileSet.Add(Output.TagName, new HashSet<FileReference>());
+				tagNameToFileSet.Add(output.TagName, new HashSet<FileReference>());
 			}
 
-			// Execute the node
-			if (bWithBanner)
+			// ExecuteAsync the node
+			if (withBanner)
 			{
 				Console.WriteLine();
-				Logger.LogInformation("========== Starting: {Arg0} ==========", Node.Name);
+				Logger.LogInformation("========== Starting: {Arg0} ==========", node.Name);
 			}
-			if (!await NodeToExecutor[Node].ExecuteAsync(new JobContext(Node.Name, this), TagNameToFileSet))
+			if (!await nodeToExecutor[node].Execute(new JobContext(node.Name, this), tagNameToFileSet))
 			{
 				return false;
 			}
-			if (bWithBanner)
+			if (withBanner)
 			{
-				Logger.LogInformation("========== Finished: {Arg0} ==========", Node.Name);
+				Logger.LogInformation("========== Finished: {Arg0} ==========", node.Name);
 				Console.WriteLine();
 			}
 
 			// Check that none of the inputs have been clobbered
-			Dictionary<string, string> ModifiedFiles = new Dictionary<string, string>(StringComparer.InvariantCultureIgnoreCase);
-			foreach (TempStorageFile File in InputManifests.Values.SelectMany(x => x.Files))
+			Dictionary<string, string> modifiedFiles = new Dictionary<string, string>(StringComparer.InvariantCultureIgnoreCase);
+			foreach (TempStorageFile file in inputManifests.Values.SelectMany(x => x.Files))
 			{
-				string? Message;
-				if (!ModifiedFiles.ContainsKey(File.RelativePath) && !IgnoreModifiedFilter.Matches(File.ToFileReference(Unreal.RootDirectory).FullName) && !File.Compare(Unreal.RootDirectory, out Message))
+				string? message;
+				if (!modifiedFiles.ContainsKey(file.RelativePath) && !ignoreModifiedFilter.Matches(file.ToFileReference(Unreal.RootDirectory).FullName) && !file.Compare(Unreal.RootDirectory, out message))
 				{
 					// look up the previous nodes to help with error diagnosis
-					List<string> PreviousNodeNames = InputManifests.Where(x => x.Value.Files.Contains(File))
+					List<string> previousNodeNames = inputManifests.Where(x => x.Value.Files.Contains(file))
 						.Select(x => x.Key.NodeName)
 						.ToList();
-					if (PreviousNodeNames.Count > 0)
+					if (previousNodeNames.Count > 0)
 					{
-						Message += $" (previous {(PreviousNodeNames.Count == 1 ? "step" : "steps")}: {string.Join(" + ", PreviousNodeNames)})";
+						message += $" (previous {(previousNodeNames.Count == 1 ? "step" : "steps")}: {string.Join(" + ", previousNodeNames)})";
 					}
 
-					ModifiedFiles.Add(File.RelativePath, Message);
+					modifiedFiles.Add(file.RelativePath, message);
 				}
 			}
-			if (ModifiedFiles.Count > 0)
+			if (modifiedFiles.Count > 0)
 			{
 				string modifiedFileList = "";
-				if (ModifiedFiles.Count < 100)
+				if (modifiedFiles.Count < 100)
 				{
-					modifiedFileList = String.Join("\n", ModifiedFiles.Select(x => $"  {x.Value}"));
+					modifiedFileList = String.Join("\n", modifiedFiles.Select(x => $"  {x.Value}"));
 				}
 				else
 				{
-					modifiedFileList = String.Join("\n", ModifiedFiles.Take(100).Select(x => $"  {x.Value}"));
-					modifiedFileList += $"\n  ...and {ModifiedFiles.Count - 100} more.";
+					modifiedFileList = String.Join("\n", modifiedFiles.Take(100).Select(x => $"  {x.Value}"));
+					modifiedFileList += $"\n  ...and {modifiedFiles.Count - 100} more.";
 				}
-				throw new AutomationException("Build {0} from a previous step have been modified:\n{1}\nOutput overlapping artifacts to a different location, or ignore them using the Node's IgnoreModified attribute.", (ModifiedFiles.Count == 1) ? "product" : "products", modifiedFileList);
+				throw new AutomationException("Build {0} from a previous step have been modified:\n{1}\nOutput overlapping artifacts to a different location, or ignore them using the Node's IgnoreModified attribute.", (modifiedFiles.Count == 1) ? "product" : "products", modifiedFileList);
 			}
 
 			// Determine all the output files which are required to be copied to temp storage (because they're referenced by nodes in another agent)
-			HashSet<FileReference> ReferencedOutputFiles = new HashSet<FileReference>();
-			foreach (BgAgentDef Agent in Graph.Agents)
+			HashSet<FileReference> referencedOutputFiles = new HashSet<FileReference>();
+			foreach (BgAgentDef agent in graph.Agents)
 			{
-				bool bSameAgent = Agent.Nodes.Contains(Node);
-				foreach (BgNodeDef OtherNode in Agent.Nodes)
+				bool sameAgent = agent.Nodes.Contains(node);
+				foreach (BgNodeDef otherNode in agent.Nodes)
 				{
-					if (!bSameAgent)
+					if (!sameAgent)
 					{
-						foreach (BgNodeOutput Input in OtherNode.Inputs.Where(x => x.ProducingNode == Node))
+						foreach (BgNodeOutput input in otherNode.Inputs.Where(x => x.ProducingNode == node))
 						{
-							ReferencedOutputFiles.UnionWith(TagNameToFileSet[Input.TagName]);
+							referencedOutputFiles.UnionWith(tagNameToFileSet[input.TagName]);
 						}
 					}
 				}
 			}
 
 			// Find a block name for all new outputs
-			Dictionary<FileReference, string> FileToOutputName = new Dictionary<FileReference, string>();
-			foreach (BgNodeOutput Output in Node.Outputs)
+			Dictionary<FileReference, string> fileToOutputName = new Dictionary<FileReference, string>();
+			foreach (BgNodeOutput output in node.Outputs)
 			{
-				HashSet<FileReference> Files = TagNameToFileSet[Output.TagName];
-				foreach (FileReference File in Files)
+				HashSet<FileReference> files = tagNameToFileSet[output.TagName];
+				foreach (FileReference file in files)
 				{
-					if (!FileToStorageBlock.ContainsKey(File) && File.IsUnderDirectory(RootDir))
+					if (!fileToStorageBlock.ContainsKey(file) && file.IsUnderDirectory(rootDir))
 					{
-						if (Output == Node.DefaultOutput)
+						if (output == node.DefaultOutput)
 						{
-							if (!FileToOutputName.ContainsKey(File))
+							if (!fileToOutputName.ContainsKey(file))
 							{
-								FileToOutputName[File] = "";
+								fileToOutputName[file] = "";
 							}
 						}
 						else
 						{
-							string? OutputName;
-							if (FileToOutputName.TryGetValue(File, out OutputName) && OutputName.Length > 0)
+							string? outputName;
+							if (fileToOutputName.TryGetValue(file, out outputName) && outputName.Length > 0)
 							{
-								FileToOutputName[File] = String.Format("{0}+{1}", OutputName, Output.TagName.Substring(1));
+								fileToOutputName[file] = String.Format("{0}+{1}", outputName, output.TagName.Substring(1));
 							}
 							else
 							{
-								FileToOutputName[File] = Output.TagName.Substring(1);
+								fileToOutputName[file] = output.TagName.Substring(1);
 							}
 						}
 					}
@@ -1221,82 +1221,82 @@ namespace AutomationTool
 			}
 
 			// Invert the dictionary to make a mapping of storage block to the files each contains
-			Dictionary<string, HashSet<FileReference>> OutputStorageBlockToFiles = new Dictionary<string, HashSet<FileReference>>();
-			foreach (KeyValuePair<FileReference, string> Pair in FileToOutputName)
+			Dictionary<string, HashSet<FileReference>> outputStorageBlockToFiles = new Dictionary<string, HashSet<FileReference>>();
+			foreach (KeyValuePair<FileReference, string> pair in fileToOutputName)
 			{
-				HashSet<FileReference>? Files;
-				if (!OutputStorageBlockToFiles.TryGetValue(Pair.Value, out Files))
+				HashSet<FileReference>? files;
+				if (!outputStorageBlockToFiles.TryGetValue(pair.Value, out files))
 				{
-					Files = new HashSet<FileReference>();
-					OutputStorageBlockToFiles.Add(Pair.Value, Files);
+					files = new HashSet<FileReference>();
+					outputStorageBlockToFiles.Add(pair.Value, files);
 				}
-				Files.Add(Pair.Key);
+				files.Add(pair.Key);
 			}
 
 			// Write all the storage blocks, and update the mapping from file to storage block
 			using (GlobalTracer.Instance.BuildSpan("TempStorage").WithTag("resource", "Write").StartActive())
 			{
-				foreach (KeyValuePair<string, HashSet<FileReference>> Pair in OutputStorageBlockToFiles)
+				foreach (KeyValuePair<string, HashSet<FileReference>> pair in outputStorageBlockToFiles)
 				{
-					TempStorageBlock OutputBlock = new TempStorageBlock(Node.Name, Pair.Key);
-					foreach (FileReference File in Pair.Value)
+					TempStorageBlock outputBlock = new TempStorageBlock(node.Name, pair.Key);
+					foreach (FileReference file in pair.Value)
 					{
-						FileToStorageBlock.Add(File, OutputBlock);
+						fileToStorageBlock.Add(file, outputBlock);
 					}
-					Storage.Archive(Node.Name, Pair.Key, Pair.Value.ToArray(), Pair.Value.Any(x => ReferencedOutputFiles.Contains(x)));
+					storage.Archive(node.Name, pair.Key, pair.Value.ToArray(), pair.Value.Any(x => referencedOutputFiles.Contains(x)));
 				}
 
 				// Find all the output tags that are published artifacts
-				Dictionary<string, BgArtifactDef> outputNameToArtifact = Graph.Artifacts.ToDictionary(x => x.TagName, x => x);
+				Dictionary<string, BgArtifactDef> outputNameToArtifact = graph.Artifacts.ToDictionary(x => x.TagName, x => x);
 
 				// Publish all the output tags
-				foreach (BgNodeOutput Output in Node.Outputs)
+				foreach (BgNodeOutput output in node.Outputs)
 				{
-					HashSet<FileReference> Files = TagNameToFileSet[Output.TagName];
+					HashSet<FileReference> files = tagNameToFileSet[output.TagName];
 
-					HashSet<TempStorageBlock> StorageBlocks = new HashSet<TempStorageBlock>();
-					foreach (FileReference File in Files)
+					HashSet<TempStorageBlock> storageBlocks = new HashSet<TempStorageBlock>();
+					foreach (FileReference file in files)
 					{
-						TempStorageBlock? StorageBlock;
-						if (FileToStorageBlock.TryGetValue(File, out StorageBlock))
+						TempStorageBlock? storageBlock;
+						if (fileToStorageBlock.TryGetValue(file, out storageBlock))
 						{
-							StorageBlocks.Add(StorageBlock);
+							storageBlocks.Add(storageBlock);
 						}
 					}
 
-					IEnumerable<string> Keys = Enumerable.Empty<string>();
-					if (outputNameToArtifact.TryGetValue(Output.TagName, out BgArtifactDef? artifact))
+					IEnumerable<string> keys = Enumerable.Empty<string>();
+					if (outputNameToArtifact.TryGetValue(output.TagName, out BgArtifactDef? artifact))
 					{
-						Keys = artifact.Keys;
+						keys = artifact.Keys;
 					}
 
-					Storage.WriteFileList(Node.Name, Output.TagName, Files, StorageBlocks.ToArray(), Keys);
+					storage.WriteFileList(node.Name, output.TagName, files, storageBlocks.ToArray(), keys);
 				}
 			}
 
 			// Mark the node as succeeded
-			Storage.MarkAsComplete(Node.Name);
+			storage.MarkAsComplete(node.Name);
 			return true;
 		}
 
 		/// <summary>
 		/// Gets an environment variable, returning null if it's not set or empty.
 		/// </summary>
-		static string? GetEnvVarOrNull(string Name)
+		static string? GetEnvVarOrNull(string name)
 		{
-			string? EnvVar = Environment.GetEnvironmentVariable(Name);
-			return String.IsNullOrEmpty(EnvVar) ? null : EnvVar;
+			string? envVar = Environment.GetEnvironmentVariable(name);
+			return String.IsNullOrEmpty(envVar) ? null : envVar;
 		}
 
 		/// <summary>
 		/// Gets an environment variable as an integer, returning null if it's not set or empty.
 		/// </summary>
-		static int? GetEnvVarIntOrNull(string Name)
+		static int? GetEnvVarIntOrNull(string name)
 		{
-			string? EnvVar = GetEnvVarOrNull(Name);
-			if (EnvVar != null && Int32.TryParse(EnvVar, out int Value))
+			string? envVar = GetEnvVarOrNull(name);
+			if (envVar != null && Int32.TryParse(envVar, out int value))
 			{
-				return Value;
+				return value;
 			}
 			return null;
 		}
@@ -1304,160 +1304,160 @@ namespace AutomationTool
 		/// <summary>
 		/// Generate HTML documentation for all the tasks
 		/// </summary>
-		/// <param name="NameToTask">Map of task name to implementation</param>
-		/// <param name="OutputFile">Output file</param>
-		static void WriteDocumentation(Dictionary<string, ScriptTaskBinding> NameToTask, FileReference OutputFile)
+		/// <param name="nameToTask">Map of task name to implementation</param>
+		/// <param name="outputFile">Output file</param>
+		static void WriteDocumentation(Dictionary<string, ScriptTaskBinding> nameToTask, FileReference outputFile)
 		{
 			// Find all the assemblies containing tasks
-			Assembly[] TaskAssemblies = NameToTask.Values.Select(x => x.ParametersClass.Assembly).Distinct().ToArray();
+			Assembly[] taskAssemblies = nameToTask.Values.Select(x => x.ParametersClass.Assembly).Distinct().ToArray();
 
 			// Read documentation for each of them
-			Dictionary<string, XmlElement> MemberNameToElement = new Dictionary<string, XmlElement>();
-			foreach (Assembly TaskAssembly in TaskAssemblies)
+			Dictionary<string, XmlElement> memberNameToElement = new Dictionary<string, XmlElement>();
+			foreach (Assembly taskAssembly in taskAssemblies)
 			{
-				string XmlFileName = Path.ChangeExtension(TaskAssembly.Location, ".xml");
-				if (File.Exists(XmlFileName))
+				string xmlFileName = Path.ChangeExtension(taskAssembly.Location, ".xml");
+				if (File.Exists(xmlFileName))
 				{
 					// Read the document
-					XmlDocument Document = new XmlDocument();
-					Document.Load(XmlFileName);
+					XmlDocument document = new XmlDocument();
+					document.Load(xmlFileName);
 
 					// Parse all the members, and add them to the map
-					foreach (XmlElement Element in Document.SelectNodes("/doc/members/member")!)
+					foreach (XmlElement element in document.SelectNodes("/doc/members/member")!)
 					{
-						string Name = Element.GetAttribute("name");
-						MemberNameToElement.Add(Name, Element);
+						string name = element.GetAttribute("name");
+						memberNameToElement.Add(name, element);
 					}
 				}
 			}
 
 			// Create the output directory
-			if (FileReference.Exists(OutputFile))
+			if (FileReference.Exists(outputFile))
 			{
-				FileReference.MakeWriteable(OutputFile);
+				FileReference.MakeWriteable(outputFile);
 			}
 			else
 			{
-				DirectoryReference.CreateDirectory(OutputFile.Directory);
+				DirectoryReference.CreateDirectory(outputFile.Directory);
 			}
 
 			// Write the output file
-			if (OutputFile.HasExtension(".udn"))
+			if (outputFile.HasExtension(".udn"))
 			{
-				WriteDocumentationUDN(NameToTask, MemberNameToElement, OutputFile);
+				WriteDocumentationUdn(nameToTask, memberNameToElement, outputFile);
 			}
-			else if (OutputFile.HasExtension(".html"))
+			else if (outputFile.HasExtension(".html"))
 			{
-				WriteDocumentationHTML(NameToTask, MemberNameToElement, OutputFile);
+				WriteDocumentationHtml(nameToTask, memberNameToElement, outputFile);
 			}
 			else
 			{
-				throw new BuildException("Unable to detect format from extension of output file ({0})", OutputFile);
+				throw new BuildException("Unable to detect format from extension of output file ({0})", outputFile);
 			}
 		}
 
 		/// <summary>
 		/// Writes documentation to a UDN file
 		/// </summary>
-		/// <param name="NameToTask">Map of name to script task</param>
-		/// <param name="MemberNameToElement">Map of field name to XML documenation element</param>
-		/// <param name="OutputFile">The output file to write to</param>
-		static void WriteDocumentationUDN(Dictionary<string, ScriptTaskBinding> NameToTask, Dictionary<string, XmlElement> MemberNameToElement, FileReference OutputFile)
+		/// <param name="nameToTask">Map of name to script task</param>
+		/// <param name="memberNameToElement">Map of field name to XML documenation element</param>
+		/// <param name="outputFile">The output file to write to</param>
+		static void WriteDocumentationUdn(Dictionary<string, ScriptTaskBinding> nameToTask, Dictionary<string, XmlElement> memberNameToElement, FileReference outputFile)
 		{
-			using (StreamWriter Writer = new StreamWriter(OutputFile.FullName))
-			{
-				Writer.WriteLine("Availability: NoPublish");
-				Writer.WriteLine("Title: BuildGraph Predefined Tasks");
-				Writer.WriteLine("Crumbs: %ROOT%, Programming, Programming/Development, Programming/Development/BuildGraph, Programming/Development/BuildGraph/BuildGraphScriptTasks");
-				Writer.WriteLine("Description: This is a procedurally generated markdown page.");
-				Writer.WriteLine("version: {0}.{1}", ReadOnlyBuildVersion.Current.MajorVersion, ReadOnlyBuildVersion.Current.MinorVersion);
-				Writer.WriteLine("parent:Programming/Development/BuildGraph/BuildGraphScriptTasks");
-				Writer.WriteLine();
-				foreach (string TaskName in NameToTask.Keys.OrderBy(x => x))
+			using (StreamWriter writer = new StreamWriter(outputFile.FullName))
+			{
+				writer.WriteLine("Availability: NoPublish");
+				writer.WriteLine("Title: BuildGraph Predefined Tasks");
+				writer.WriteLine("Crumbs: %ROOT%, Programming, Programming/Development, Programming/Development/BuildGraph, Programming/Development/BuildGraph/BuildGraphScriptTasks");
+				writer.WriteLine("Description: This is a procedurally generated markdown page.");
+				writer.WriteLine("version: {0}.{1}", ReadOnlyBuildVersion.Current.MajorVersion, ReadOnlyBuildVersion.Current.MinorVersion);
+				writer.WriteLine("parent:Programming/Development/BuildGraph/BuildGraphScriptTasks");
+				writer.WriteLine();
+				foreach (string taskName in nameToTask.Keys.OrderBy(x => x))
 				{
 					// Get the task object
-					ScriptTaskBinding Task = NameToTask[TaskName];
+					ScriptTaskBinding task = nameToTask[taskName];
 
 					// Get the documentation for this task
-					XmlElement? TaskElement;
-					if (MemberNameToElement.TryGetValue("T:" + Task.TaskClass.FullName, out TaskElement))
+					XmlElement? taskElement;
+					if (memberNameToElement.TryGetValue("T:" + task.TaskClass.FullName, out taskElement))
 					{
 						// Write the task heading
-						Writer.WriteLine("### {0}", TaskName);
-						Writer.WriteLine();
-						Writer.WriteLine(ConvertToMarkdown(TaskElement.SelectSingleNode("summary")!));
-						Writer.WriteLine();
+						writer.WriteLine("### {0}", taskName);
+						writer.WriteLine();
+						writer.WriteLine(ConvertToMarkdown(taskElement.SelectSingleNode("summary")!));
+						writer.WriteLine();
 
 						// Document the parameters
-						List<string[]> Rows = new List<string[]>();
-						foreach (string ParameterName in Task.NameToParameter.Keys)
+						List<string[]> rows = new List<string[]>();
+						foreach (string parameterName in task.NameToParameter.Keys)
 						{
 							// Get the parameter data
-							ScriptTaskParameterBinding Parameter = Task.NameToParameter[ParameterName];
+							ScriptTaskParameterBinding parameter = task.NameToParameter[parameterName];
 
 							// Get the documentation for this parameter
-							XmlElement? ParameterElement;
-							if (MemberNameToElement.TryGetValue(Parameter.GetXmlDocName(), out ParameterElement))
+							XmlElement? parameterElement;
+							if (memberNameToElement.TryGetValue(parameter.GetXmlDocName(), out parameterElement))
 							{
-								Type FieldType = Parameter.ParameterType;
-								if (FieldType.IsGenericType && FieldType.GetGenericTypeDefinition() == typeof(Nullable<>))
+								Type fieldType = parameter.ParameterType;
+								if (fieldType.IsGenericType && fieldType.GetGenericTypeDefinition() == typeof(Nullable<>))
 								{
-									FieldType = FieldType.GetGenericArguments()[0];
+									fieldType = fieldType.GetGenericArguments()[0];
 								}
 
-								string TypeName;
-								if (Parameter.ValidationType != TaskParameterValidationType.Default)
+								string typeName;
+								if (parameter.ValidationType != TaskParameterValidationType.Default)
 								{
-									TypeName = Parameter.ValidationType.ToString();
+									typeName = parameter.ValidationType.ToString();
 								}
-								else if (FieldType == typeof(int))
+								else if (fieldType == typeof(int))
 								{
-									TypeName = "Integer";
+									typeName = "Integer";
 								}
-								else if (FieldType == typeof(HashSet<DirectoryReference>))
+								else if (fieldType == typeof(HashSet<DirectoryReference>))
 								{
-									TypeName = "DirectoryList";
+									typeName = "DirectoryList";
 								}
 								else
 								{
-									TypeName = FieldType.Name;
+									typeName = fieldType.Name;
 								}
 
-								string[] Columns = new string[4];
-								Columns[0] = ParameterName;
-								Columns[1] = TypeName;
-								Columns[2] = Parameter.Optional ? "Optional" : "Required";
-								Columns[3] = ConvertToMarkdown(ParameterElement.SelectSingleNode("summary")!);
-								Rows.Add(Columns);
+								string[] columns = new string[4];
+								columns[0] = parameterName;
+								columns[1] = typeName;
+								columns[2] = parameter.Optional ? "Optional" : "Required";
+								columns[3] = ConvertToMarkdown(parameterElement.SelectSingleNode("summary")!);
+								rows.Add(columns);
 							}
 						}
 
 						// Always include the "If" attribute
-						string[] IfColumns = new string[4];
-						IfColumns[0] = "If";
-						IfColumns[1] = "Condition";
-						IfColumns[2] = "Optional";
-						IfColumns[3] = "Whether to execute this task. It is ignored if this condition evaluates to false.";
-						Rows.Add(IfColumns);
+						string[] ifColumns = new string[4];
+						ifColumns[0] = "If";
+						ifColumns[1] = "Condition";
+						ifColumns[2] = "Optional";
+						ifColumns[3] = "Whether to execute this task. It is ignored if this condition evaluates to false.";
+						rows.Add(ifColumns);
 
 						// Get the width of each column
-						int[] Widths = new int[4];
-						for (int Idx = 0; Idx < 4; Idx++)
+						int[] widths = new int[4];
+						for (int idx = 0; idx < 4; idx++)
 						{
-							Widths[Idx] = Rows.Max(x => x[Idx].Length);
+							widths[idx] = rows.Max(x => x[idx].Length);
 						}
 
 						// Format the markdown table
-						string Format = String.Format("| {{0,-{0}}} | {{1,-{1}}} | {{2,-{2}}} | {{3,-{3}}} |", Widths[0], Widths[1], Widths[2], Widths[3]);
-						Writer.WriteLine(Format, "", "", "", "");
-						Writer.WriteLine(Format, new string('-', Widths[0]), new string('-', Widths[1]), new string('-', Widths[2]), new string('-', Widths[3]));
-						for (int Idx = 0; Idx < Rows.Count; Idx++)
+						string format = String.Format("| {{0,-{0}}} | {{1,-{1}}} | {{2,-{2}}} | {{3,-{3}}} |", widths[0], widths[1], widths[2], widths[3]);
+						writer.WriteLine(format, "", "", "", "");
+						writer.WriteLine(format, new string('-', widths[0]), new string('-', widths[1]), new string('-', widths[2]), new string('-', widths[3]));
+						for (int idx = 0; idx < rows.Count; idx++)
 						{
-							Writer.WriteLine(Format, Rows[Idx][0], Rows[Idx][1], Rows[Idx][2], Rows[Idx][3]);
+							writer.WriteLine(format, rows[idx][0], rows[idx][1], rows[idx][2], rows[idx][3]);
 						}
 
 						// Blank line before next task
-						Writer.WriteLine();
+						writer.WriteLine();
 					}
 				}
 			}
@@ -1466,121 +1466,121 @@ namespace AutomationTool
 		/// <summary>
 		/// Writes documentation to an HTML file
 		/// </summary>
-		/// <param name="NameToTask">Map of name to script task</param>
-		/// <param name="MemberNameToElement">Map of field name to XML documenation element</param>
-		/// <param name="OutputFile">The output file to write to</param>
-		static void WriteDocumentationHTML(Dictionary<string, ScriptTaskBinding> NameToTask, Dictionary<string, XmlElement> MemberNameToElement, FileReference OutputFile)
+		/// <param name="nameToTask">Map of name to script task</param>
+		/// <param name="memberNameToElement">Map of field name to XML documenation element</param>
+		/// <param name="outputFile">The output file to write to</param>
+		static void WriteDocumentationHtml(Dictionary<string, ScriptTaskBinding> nameToTask, Dictionary<string, XmlElement> memberNameToElement, FileReference outputFile)
 		{
-			Logger.LogInformation("Writing {OutputFile}...", OutputFile);
-			using (StreamWriter Writer = new StreamWriter(OutputFile.FullName))
-			{
-				Writer.WriteLine("<html>");
-				Writer.WriteLine("  <head>");
-				Writer.WriteLine("    <style>");
-				Writer.WriteLine("      table { border-collapse: collapse }");
-				Writer.WriteLine("      table, th, td { border: 1px solid black; }");
-				Writer.WriteLine("    </style>");
-				Writer.WriteLine("  </head>");
-				Writer.WriteLine("  <body>");
-				Writer.WriteLine("    <h1>BuildGraph Tasks</h1>");
-				foreach (string TaskName in NameToTask.Keys.OrderBy(x => x))
+			Logger.LogInformation("Writing {OutputFile}...", outputFile);
+			using (StreamWriter writer = new StreamWriter(outputFile.FullName))
+			{
+				writer.WriteLine("<html>");
+				writer.WriteLine("  <head>");
+				writer.WriteLine("    <style>");
+				writer.WriteLine("      table { border-collapse: collapse }");
+				writer.WriteLine("      table, th, td { border: 1px solid black; }");
+				writer.WriteLine("    </style>");
+				writer.WriteLine("  </head>");
+				writer.WriteLine("  <body>");
+				writer.WriteLine("    <h1>BuildGraph Tasks</h1>");
+				foreach (string taskName in nameToTask.Keys.OrderBy(x => x))
 				{
 					// Get the task object
-					ScriptTaskBinding Task = NameToTask[TaskName];
+					ScriptTaskBinding task = nameToTask[taskName];
 
 					// Get the documentation for this task
-					XmlElement? TaskElement;
-					if (MemberNameToElement.TryGetValue("T:" + Task.TaskClass.FullName, out TaskElement))
+					XmlElement? taskElement;
+					if (memberNameToElement.TryGetValue("T:" + task.TaskClass.FullName, out taskElement))
 					{
 						// Write the task heading
-						Writer.WriteLine("    <h2>{0}</h2>", TaskName);
-						Writer.WriteLine("    <p>{0}</p>", TaskElement.SelectSingleNode("summary")!.InnerXml.Trim());
+						writer.WriteLine("    <h2>{0}</h2>", taskName);
+						writer.WriteLine("    <p>{0}</p>", taskElement.SelectSingleNode("summary")!.InnerXml.Trim());
 
 						// Start the parameter table
-						Writer.WriteLine("    <table>");
-						Writer.WriteLine("      <tr>");
-						Writer.WriteLine("        <th>Attribute</th>");
-						Writer.WriteLine("        <th>Type</th>");
-						Writer.WriteLine("        <th>Usage</th>");
-						Writer.WriteLine("        <th>Description</th>");
-						Writer.WriteLine("      </tr>");
+						writer.WriteLine("    <table>");
+						writer.WriteLine("      <tr>");
+						writer.WriteLine("        <th>Attribute</th>");
+						writer.WriteLine("        <th>Type</th>");
+						writer.WriteLine("        <th>Usage</th>");
+						writer.WriteLine("        <th>Description</th>");
+						writer.WriteLine("      </tr>");
 
 						// Document the parameters
-						foreach (string ParameterName in Task.NameToParameter.Keys)
+						foreach (string parameterName in task.NameToParameter.Keys)
 						{
 							// Get the parameter data
-							ScriptTaskParameterBinding Parameter = Task.NameToParameter[ParameterName];
+							ScriptTaskParameterBinding parameter = task.NameToParameter[parameterName];
 
 							// Get the documentation for this parameter
-							XmlElement? ParameterElement;
-							if (MemberNameToElement.TryGetValue(Parameter.GetXmlDocName(), out ParameterElement))
+							XmlElement? parameterElement;
+							if (memberNameToElement.TryGetValue(parameter.GetXmlDocName(), out parameterElement))
 							{
-								string TypeName = Parameter.ParameterType.Name;
-								if (Parameter.ValidationType != TaskParameterValidationType.Default)
+								string typeName = parameter.ParameterType.Name;
+								if (parameter.ValidationType != TaskParameterValidationType.Default)
 								{
-									StringBuilder NewTypeName = new StringBuilder(Parameter.ValidationType.ToString());
-									for (int Idx = 1; Idx < NewTypeName.Length; Idx++)
+									StringBuilder newTypeName = new StringBuilder(parameter.ValidationType.ToString());
+									for (int idx = 1; idx < newTypeName.Length; idx++)
 									{
-										if (Char.IsLower(NewTypeName[Idx - 1]) && Char.IsUpper(NewTypeName[Idx]))
+										if (Char.IsLower(newTypeName[idx - 1]) && Char.IsUpper(newTypeName[idx]))
 										{
-											NewTypeName.Insert(Idx, ' ');
+											newTypeName.Insert(idx, ' ');
 										}
 									}
-									TypeName = NewTypeName.ToString();
+									typeName = newTypeName.ToString();
 								}
 
-								Writer.WriteLine("      <tr>");
-								Writer.WriteLine("         <td>{0}</td>", ParameterName);
-								Writer.WriteLine("         <td>{0}</td>", TypeName);
-								Writer.WriteLine("         <td>{0}</td>", Parameter.Optional ? "Optional" : "Required");
-								Writer.WriteLine("         <td>{0}</td>", ParameterElement.SelectSingleNode("summary")!.InnerXml.Trim());
-								Writer.WriteLine("      </tr>");
+								writer.WriteLine("      <tr>");
+								writer.WriteLine("         <td>{0}</td>", parameterName);
+								writer.WriteLine("         <td>{0}</td>", typeName);
+								writer.WriteLine("         <td>{0}</td>", parameter.Optional ? "Optional" : "Required");
+								writer.WriteLine("         <td>{0}</td>", parameterElement.SelectSingleNode("summary")!.InnerXml.Trim());
+								writer.WriteLine("      </tr>");
 							}
 						}
 
 						// Always include the "If" attribute
-						Writer.WriteLine("     <tr>");
-						Writer.WriteLine("       <td>If</td>");
-						Writer.WriteLine("       <td>Condition</td>");
-						Writer.WriteLine("       <td>Optional</td>");
-						Writer.WriteLine("       <td>Whether to execute this task. It is ignored if this condition evaluates to false.</td>");
-						Writer.WriteLine("     </tr>");
+						writer.WriteLine("     <tr>");
+						writer.WriteLine("       <td>If</td>");
+						writer.WriteLine("       <td>Condition</td>");
+						writer.WriteLine("       <td>Optional</td>");
+						writer.WriteLine("       <td>Whether to execute this task. It is ignored if this condition evaluates to false.</td>");
+						writer.WriteLine("     </tr>");
 
 						// Close the table
-						Writer.WriteLine("    <table>");
+						writer.WriteLine("    <table>");
 					}
 				}
-				Writer.WriteLine("  </body>");
-				Writer.WriteLine("</html>");
+				writer.WriteLine("  </body>");
+				writer.WriteLine("</html>");
 			}
 		}
 
 		/// <summary>
 		/// Converts an XML documentation node to markdown
 		/// </summary>
-		/// <param name="Node">The node to read</param>
+		/// <param name="node">The node to read</param>
 		/// <returns>Text in markdown format</returns>
-		static string ConvertToMarkdown(XmlNode Node)
+		static string ConvertToMarkdown(XmlNode node)
 		{
-			string Text = Node.InnerXml;
+			string text = node.InnerXml;
 
-			StringBuilder Result = new StringBuilder();
-			for (int Idx = 0; Idx < Text.Length; Idx++)
+			StringBuilder result = new StringBuilder();
+			for (int idx = 0; idx < text.Length; idx++)
 			{
-				if (Char.IsWhiteSpace(Text[Idx]))
+				if (Char.IsWhiteSpace(text[idx]))
 				{
-					Result.Append(' ');
-					while (Idx + 1 < Text.Length && Char.IsWhiteSpace(Text[Idx + 1]))
+					result.Append(' ');
+					while (idx + 1 < text.Length && Char.IsWhiteSpace(text[idx + 1]))
 					{
-						Idx++;
+						idx++;
 					}
 				}
 				else
 				{
-					Result.Append(Text[Idx]);
+					result.Append(text[idx]);
 				}
 			}
-			return Result.ToString().Trim();
+			return result.ToString().Trim();
 		}
 	}
 
diff --git a/Engine/Source/Programs/AutomationTool/BuildGraph/JobContext.cs b/Engine/Source/Programs/AutomationTool/BuildGraph/JobContext.cs
--- a/Engine/Source/Programs/AutomationTool/BuildGraph/JobContext.cs
+++ b/Engine/Source/Programs/AutomationTool/BuildGraph/JobContext.cs
@@ -1,11 +1,5 @@
 // Copyright Epic Games, Inc. All Rights Reserved.
 
-using System;
-using System.Collections.Generic;
-using System.Linq;
-using System.Text;
-using System.Threading.Tasks;
-
 #nullable enable
 
 namespace AutomationTool
@@ -28,20 +22,20 @@ namespace AutomationTool
 		/// <summary>
 		/// Constructor
 		/// </summary>
-		/// <param name="InOwnerCommand">The command running the current job</param>
-		public JobContext(BuildCommand InOwnerCommand) : this("Unknown", InOwnerCommand)
+		/// <param name="inOwnerCommand">The command running the current job</param>
+		public JobContext(BuildCommand inOwnerCommand) : this("Unknown", inOwnerCommand)
 		{
 		}
 
 		/// <summary>
 		/// Constructor
 		/// </summary>
-		/// <param name="InCurrentNode">The current node being executed</param>
-		/// <param name="InOwnerCommand">The command running the current job</param>
-		public JobContext(string InCurrentNode, BuildCommand InOwnerCommand)
+		/// <param name="inCurrentNode">The current node being executed</param>
+		/// <param name="inOwnerCommand">The command running the current job</param>
+		public JobContext(string inCurrentNode, BuildCommand inOwnerCommand)
 		{
-			CurrentNode = InCurrentNode;
-			OwnerCommand = InOwnerCommand;
+			CurrentNode = inCurrentNode;
+			OwnerCommand = inOwnerCommand;
 		}
 	}
 }
