EpicGames.BuildGraph: Apply standard naming conventions to BuildGraph.Automation.

diff --git a/Engine/Source/Programs/AutomationTool/BuildGraph/Scripts/InstalledBuild.cs b/Engine/Source/Programs/AutomationTool/BuildGraph/Scripts/InstalledBuild.cs
--- a/Engine/Source/Programs/AutomationTool/BuildGraph/Scripts/InstalledBuild.cs
+++ b/Engine/Source/Programs/AutomationTool/BuildGraph/Scripts/InstalledBuild.cs
@@ -1,18 +1,16 @@
 // Copyright Epic Games, Inc. All Rights Reserved.
 
-using EpicGames.Core;
-using EpicGames.BuildGraph;
-using EpicGames.BuildGraph.Expressions;
 using System;
-using UnrealBuildTool;
-using AutomationTool.Tasks;
 using System.Collections.Generic;
-using System.Linq;
-using UnrealBuildBase;
 using System.Threading.Tasks;
-
-using static AutomationTool.Tasks.StandardTasks;
+using AutomationTool.Tasks;
+using EpicGames.BuildGraph;
+using EpicGames.BuildGraph.Expressions;
+using EpicGames.Core;
 using Microsoft.Extensions.Logging;
+using UnrealBuildBase;
+using UnrealBuildTool;
+using static AutomationTool.Tasks.StandardTasks;
 
 #nullable enable
 
@@ -25,207 +23,204 @@ namespace AutomationTool
 		static DirectoryReference IntermediateDir => DirectoryReference.Combine(RootDir, "Engine", "Intermediate", "Installed");
 		static DirectoryReference TempMiscDir => DirectoryReference.Combine(IntermediateDir, "General");
 		static DirectoryReference TempCsToolsDir => DirectoryReference.Combine(IntermediateDir, "CsTools");
-		static DirectoryReference TempDdcDir = DirectoryReference.Combine(IntermediateDir, "DDC");
-
-//		static string[] PluginsExceptions =
-//		{
-//			"Engine/Plugins/Enterprise/DatasmithCADImporter/...",
-//			"Engine/Plugins/Enterprise/DatasmithC4DImporter/...",
-//			"Engine/Plugins/Enterprise/AxFImporter/...",
-//			"Engine/Plugins/Enterprise/MDLImporter/..."
-//		};
-
-//		static string[] WinSignFilter =
-//		{
-//			"*.exe",
-//			"*.dll"
-//		};
-
-		static string[] WinStripFilter =
+		static DirectoryReference TempDdcDir => DirectoryReference.Combine(IntermediateDir, "DDC");
+
+		//		static string[] PluginsExceptions =
+		//		{
+		//			"Engine/Plugins/Enterprise/DatasmithCADImporter/...",
+		//			"Engine/Plugins/Enterprise/DatasmithC4DImporter/...",
+		//			"Engine/Plugins/Enterprise/AxFImporter/...",
+		//			"Engine/Plugins/Enterprise/MDLImporter/..."
+		//		};
+
+		//		static string[] WinSignFilter =
+		//		{
+		//			"*.exe",
+		//			"*.dll"
+		//		};
+
+		static readonly string[] s_winStripFilter =
 		{
 			"*.pdb",
 			"-/Engine/Binaries/Win64/UnrealEditor*.pdb",
 			"-/Engine/Plugins/.../Binaries/Win64/UnrealEditor*.pdb",
 		};
 
-		static List<string> GetDdcProjects(UnrealTargetPlatform Platform)
+		static List<string> GetDdcProjects(UnrealTargetPlatform platform)
 		{
-			List<string> Projects = new List<string>();
-			Projects.Add("Templates/TP_AEC_ArchvisBP/TP_AEC_ArchvisBP.uproject");
-			Projects.Add("Templates/TP_AEC_BlankBP/TP_AEC_BlankBP.uproject");
-			Projects.Add("Templates/TP_AEC_CollabBP/TP_AEC_CollabBP.uproject");
-			Projects.Add("Templates/TP_AEC_ProdConfigBP/TP_AEC_ProdConfigBP.uproject");
-			Projects.Add("Templates/TP_FirstPersonBP/TP_FirstPersonBP.uproject");
-			Projects.Add("Templates/TP_HandheldARBP/TP_HandheldARBP.uproject");
-			Projects.Add("Templates/TP_AEC_HandheldARBP/TP_AEC_HandheldARBP.uproject");
-			Projects.Add("Templates/TP_MFG_HandheldARBP/TP_MFG_HandheldARBP.uproject");
-			Projects.Add("Templates/TP_MFG_CollabBP/TP_MFG_CollabBP.uproject");
-			Projects.Add("Templates/TP_MFG_ProdConfigBP/TP_MFG_ProdConfigBP.uproject");
-			Projects.Add("Templates/TP_PhotoStudioBP/TP_PhotoStudioBP.uproject");
-			Projects.Add("Templates/TP_PuzzleBP/TP_PuzzleBP.uproject");
-			Projects.Add("Templates/TP_ThirdPersonBP/TP_ThirdPersonBP.uproject");
-			Projects.Add("Templates/TP_TopDownBP/TP_TopDownBP.uproject");
-			Projects.Add("Templates/TP_VehicleAdvBP/TP_VehicleAdvBP.uproject");
-			Projects.Add("Templates/TP_VirtualRealityBP/TP_VirtualRealityBP.uproject");
-			Projects.Add("Templates/TP_DMXBP/TP_DMXBP.uproject");
+			List<string> projects = new List<string>();
+			projects.Add("Templates/TP_AEC_ArchvisBP/TP_AEC_ArchvisBP.uproject");
+			projects.Add("Templates/TP_AEC_BlankBP/TP_AEC_BlankBP.uproject");
+			projects.Add("Templates/TP_AEC_CollabBP/TP_AEC_CollabBP.uproject");
+			projects.Add("Templates/TP_AEC_ProdConfigBP/TP_AEC_ProdConfigBP.uproject");
+			projects.Add("Templates/TP_FirstPersonBP/TP_FirstPersonBP.uproject");
+			projects.Add("Templates/TP_HandheldARBP/TP_HandheldARBP.uproject");
+			projects.Add("Templates/TP_AEC_HandheldARBP/TP_AEC_HandheldARBP.uproject");
+			projects.Add("Templates/TP_MFG_HandheldARBP/TP_MFG_HandheldARBP.uproject");
+			projects.Add("Templates/TP_MFG_CollabBP/TP_MFG_CollabBP.uproject");
+			projects.Add("Templates/TP_MFG_ProdConfigBP/TP_MFG_ProdConfigBP.uproject");
+			projects.Add("Templates/TP_PhotoStudioBP/TP_PhotoStudioBP.uproject");
+			projects.Add("Templates/TP_PuzzleBP/TP_PuzzleBP.uproject");
+			projects.Add("Templates/TP_ThirdPersonBP/TP_ThirdPersonBP.uproject");
+			projects.Add("Templates/TP_TopDownBP/TP_TopDownBP.uproject");
+			projects.Add("Templates/TP_VehicleAdvBP/TP_VehicleAdvBP.uproject");
+			projects.Add("Templates/TP_VirtualRealityBP/TP_VirtualRealityBP.uproject");
+			projects.Add("Templates/TP_DMXBP/TP_DMXBP.uproject");
 			//			Projects.Add("Samples/StarterContent/StarterContent.uproject");
 
-			if (Platform == UnrealTargetPlatform.Win64)
+			if (platform == UnrealTargetPlatform.Win64)
 			{
-				Projects.Add("Templates/TP_InCamVFXBP/TP_InCamVFXBP.uproject");
+				projects.Add("Templates/TP_InCamVFXBP/TP_InCamVFXBP.uproject");
 			}
-			return Projects;
+			return projects;
 		}
 
-		public override BgGraph CreateGraph(BgEnvironment Context)
+		public override BgGraph CreateGraph(BgEnvironment context)
 		{
-			string RootDir = CommandUtils.CmdEnv.LocalRoot;
+			string rootDir = CommandUtils.CmdEnv.LocalRoot;
 
-			BgBoolOption HostPlatformOnly = new BgBoolOption("HostPlatformOnly", "A helper option to make an installed build for your host platform only, so that you don't have to disable each platform individually", false);
-			BgBoolOption HostPlatformEditorOnly = new BgBoolOption("HostPlatformEditorOnly", "A helper option to make an installed build for your host platform only, so that you don't have to disable each platform individually", false);
-			BgBoolOption AllPlatforms = new BgBoolOption("AllPlatforms", "Include all target platforms by default", false);
-			BgBoolOption CompileDatasmithPlugins = new BgBoolOption("CompileDatasmithPlugins", "If Datasmith plugins should be compiled on a separate node.", false);
+			BgBoolOption hostPlatformOnly = new BgBoolOption("HostPlatformOnly", "A helper option to make an installed build for your host platform only, so that you don't have to disable each platform individually", false);
+			BgBoolOption hostPlatformEditorOnly = new BgBoolOption("HostPlatformEditorOnly", "A helper option to make an installed build for your host platform only, so that you don't have to disable each platform individually", false);
+			BgBoolOption allPlatforms = new BgBoolOption("AllPlatforms", "Include all target platforms by default", false);
+			BgBoolOption compileDatasmithPlugins = new BgBoolOption("CompileDatasmithPlugins", "If Datasmith plugins should be compiled on a separate node.", false);
 
-			UnrealTargetPlatform CurrentHostPlatform = HostPlatform.Current.HostEditorPlatform;
+			UnrealTargetPlatform currentHostPlatform = HostPlatform.Current.HostEditorPlatform;
 
-			BgBool DefaultWithWin64 = !(HostPlatformEditorOnly | (HostPlatformOnly & (CurrentHostPlatform != UnrealTargetPlatform.Win64)));
-			BgBool DefaultWithMac = !(HostPlatformEditorOnly | (HostPlatformOnly & (CurrentHostPlatform != UnrealTargetPlatform.Mac)));
-			BgBool DefaultWithLinux = !(HostPlatformEditorOnly | (HostPlatformOnly & (CurrentHostPlatform != UnrealTargetPlatform.Linux)));
-			BgBool DefaultWithLinuxArm64 = !(HostPlatformEditorOnly | (HostPlatformOnly & (CurrentHostPlatform != UnrealTargetPlatform.Linux)));
-			BgBool DefaultWithPlatform = !(HostPlatformEditorOnly | HostPlatformOnly);
-			BgBool DefaultWithIOS = !((CurrentHostPlatform != UnrealTargetPlatform.Mac) & !AllPlatforms);
+			BgBool defaultWithWin64 = !(hostPlatformEditorOnly | (hostPlatformOnly & (currentHostPlatform != UnrealTargetPlatform.Win64)));
+			BgBool defaultWithMac = !(hostPlatformEditorOnly | (hostPlatformOnly & (currentHostPlatform != UnrealTargetPlatform.Mac)));
+			BgBool defaultWithLinux = !(hostPlatformEditorOnly | (hostPlatformOnly & (currentHostPlatform != UnrealTargetPlatform.Linux)));
+			BgBool defaultWithLinuxArm64 = !(hostPlatformEditorOnly | (hostPlatformOnly & (currentHostPlatform != UnrealTargetPlatform.Linux)));
+			BgBool defaultWithPlatform = !(hostPlatformEditorOnly | hostPlatformOnly);
+			BgBool defaultWithIos = !((currentHostPlatform != UnrealTargetPlatform.Mac) & !allPlatforms);
 
-			BgBoolOption WithWin64 = new BgBoolOption("WithWin64", "Include the Win64 target platform", DefaultWithWin64);
-			BgBoolOption WithMac = new BgBoolOption("WithMac", "Include the Mac target platform", DefaultWithMac);
-			BgBoolOption WithAndroid = new BgBoolOption("WithAndroid", "Include the Android target platform", DefaultWithPlatform);
-			BgBoolOption WithIOS = new BgBoolOption("WithIOS", "Include the iOS target platform", DefaultWithIOS);
-			BgBoolOption WithTVOS = new BgBoolOption("WithTVOS", "Include the tvOS target platform", DefaultWithIOS);
-			BgBoolOption WithLinux = new BgBoolOption("WithLinux", "Include the Linux target platform", DefaultWithLinux);
-			BgBoolOption WithLinuxArm64 = new BgBoolOption("WithLinuxArm64", "Include the Linux AArch64 target platform", DefaultWithLinuxArm64);
+			BgBoolOption withWin64 = new BgBoolOption("WithWin64", "Include the Win64 target platform", defaultWithWin64);
+			BgBoolOption withMac = new BgBoolOption("WithMac", "Include the Mac target platform", defaultWithMac);
+			BgBoolOption withAndroid = new BgBoolOption("WithAndroid", "Include the Android target platform", defaultWithPlatform);
+			BgBoolOption withIos = new BgBoolOption("WithIOS", "Include the iOS target platform", defaultWithIos);
+			BgBoolOption withTvos = new BgBoolOption("WithTVOS", "Include the tvOS target platform", defaultWithIos);
+			BgBoolOption withLinux = new BgBoolOption("WithLinux", "Include the Linux target platform", defaultWithLinux);
+			BgBoolOption withLinuxArm64 = new BgBoolOption("WithLinuxArm64", "Include the Linux AArch64 target platform", defaultWithLinuxArm64);
 
-			BgBoolOption WithClient = new BgBoolOption("WithClient", "Include precompiled client targets", false);
-			BgBoolOption WithServer = new BgBoolOption("WithServer", "Include precompiled server targets", false);
-			BgBoolOption WithDDC = new BgBoolOption("WithDDC", "Build a standalone derived-data cache for the engine content and templates", true);
-			BgBoolOption HostPlatformDDCOnly = new BgBoolOption("HostPlatformDDCOnly", "Whether to include DDC for the host platform only", true);
-			BgBoolOption SignExecutables = new BgBoolOption("SignExecutables", "Sign the executables produced where signing is available", false);
+			BgBoolOption withClient = new BgBoolOption("WithClient", "Include precompiled client targets", false);
+			BgBoolOption withServer = new BgBoolOption("WithServer", "Include precompiled server targets", false);
+			BgBoolOption withDdc = new BgBoolOption("WithDDC", "Build a standalone derived-data cache for the engine content and templates", true);
+			BgBoolOption hostPlatformDdcOnly = new BgBoolOption("HostPlatformDDCOnly", "Whether to include DDC for the host platform only", true);
+			BgBoolOption signExecutables = new BgBoolOption("SignExecutables", "Sign the executables produced where signing is available", false);
 
-			BgStringOption AnalyticsTypeOverride = new BgStringOption("AnalyticsTypeOverride", "Identifier for analytic events to send", "");
-			BgBool EmbedSrcSrvInfo = new BgBoolOption("EmbedSrcSrvInfo", "Whether to add Source indexing to Windows game apps so they can be added to a symbol server", false);
+			BgStringOption analyticsTypeOverride = new BgStringOption("AnalyticsTypeOverride", "Identifier for analytic events to send", "");
+			BgBool embedSrcSrvInfo = new BgBoolOption("EmbedSrcSrvInfo", "Whether to add Source indexing to Windows game apps so they can be added to a symbol server", false);
 
-			BgList<BgString> DefaultGameConfigurations = BgList<BgString>.Create(nameof(UnrealTargetConfiguration.DebugGame), nameof(UnrealTargetConfiguration.Development), nameof(UnrealTargetConfiguration.Shipping));
-			BgList<BgString> GameConfigurationStrings = new BgListOption("GameConfigurations", description: "Which game configurations to include for packaged applications", style: BgListOptionStyle.CheckList, values: DefaultGameConfigurations);
-			BgList<BgEnum<UnrealTargetConfiguration>> GameConfigurations = GameConfigurationStrings.Select(x => BgEnum<UnrealTargetConfiguration>.Parse(x));
+			BgList<BgString> defaultGameConfigurations = BgList<BgString>.Create(nameof(UnrealTargetConfiguration.DebugGame), nameof(UnrealTargetConfiguration.Development), nameof(UnrealTargetConfiguration.Shipping));
+			BgList<BgString> gameConfigurationStrings = new BgListOption("GameConfigurations", description: "Which game configurations to include for packaged applications", style: BgListOptionStyle.CheckList, values: defaultGameConfigurations);
+			BgList<BgEnum<UnrealTargetConfiguration>> gameConfigurations = gameConfigurationStrings.Select(x => BgEnum<UnrealTargetConfiguration>.Parse(x));
 
-			BgBoolOption WithFullDebugInfo = new BgBoolOption("WithFullDebugInfo", "Generate full debug info for binary editor and packaged application builds", false);
+			BgBoolOption withFullDebugInfo = new BgBoolOption("WithFullDebugInfo", "Generate full debug info for binary editor and packaged application builds", false);
 
-			BgStringOption BuiltDirectory = new BgStringOption("BuiltDirectory", "Directory for outputting the built engine", RootDir + "/LocalBuilds/Engine");
+			BgStringOption builtDirectory = new BgStringOption("BuiltDirectory", "Directory for outputting the built engine", rootDir + "/LocalBuilds/Engine");
 
-			BgStringOption CrashReporterAPIURL = new BgStringOption("CrashReporterAPIURL", "The URL to use to talk to the CrashReporterClient API.", "");
-			BgStringOption CrashReporterAPIKey = new BgStringOption("CrashReporterAPIKey", "The API key to use to talk to the CrashReporterClient API.", "");
-			BgStringOption BuildId = new BgStringOption("BuildId", "The unique build identifier to associate with this installed build", "");
+			BgStringOption crashReporterApiurl = new BgStringOption("CrashReporterAPIURL", "The URL to use to talk to the CrashReporterClient API.", "");
+			BgStringOption crashReporterApiKey = new BgStringOption("CrashReporterAPIKey", "The API key to use to talk to the CrashReporterClient API.", "");
+			BgStringOption buildId = new BgStringOption("BuildId", "The unique build identifier to associate with this installed build", "");
 
-			BgString CrashReporterCompileArgs = "";
-			CrashReporterCompileArgs = CrashReporterCompileArgs.If(CrashReporterAPIURL != "" & CrashReporterAPIKey != "", BgString.Format("-define:CRC_TELEMETRY_URL=\"{0}\" -define:CRC_TELEMETRY_KEY_DEV=\"{1}\" -define:CRC_TELEMETRY_KEY_RELEASE=\"{1}\" -OverrideBuildEnvironment", CrashReporterAPIURL, CrashReporterAPIKey));
+			BgString crashReporterCompileArgs = "";
+			crashReporterCompileArgs = crashReporterCompileArgs.If(crashReporterApiurl != "" & crashReporterApiKey != "", BgString.Format("-define:CRC_TELEMETRY_URL=\"{0}\" -define:CRC_TELEMETRY_KEY_DEV=\"{1}\" -define:CRC_TELEMETRY_KEY_RELEASE=\"{1}\" -OverrideBuildEnvironment", crashReporterApiurl, crashReporterApiKey));
 
-			List<BgAggregate> Aggregates = new List<BgAggregate>();
+			List<BgAggregate> aggregates = new List<BgAggregate>();
 
 			/////// EDITORS ////////////////////////////////////////////////
 
 			//// Windows ////
-			BgAgent EditorWin64 = new BgAgent("Editor Win64", "Win64_Licensee");
+			BgAgent editorWin64 = new BgAgent("Editor Win64", "Win64_Licensee");
 
-			BgNode VersionFilesNode = EditorWin64
+			BgNode versionFilesNode = editorWin64
 				.AddNode(x => UpdateVersionFilesAsync(x));
 
-			BgNode<BgFileSet> WinEditorNode = EditorWin64
-				.AddNode(x => CompileUnrealEditorWin64Async(x, CrashReporterCompileArgs, EmbedSrcSrvInfo, CompileDatasmithPlugins, WithFullDebugInfo, SignExecutables))
-				.Requires(VersionFilesNode);
-
-			Aggregates.Add(new BgAggregate("Win64 Editor", WinEditorNode, label: "Editors/Win64"));
+			BgNode<BgFileSet> winEditorNode = editorWin64
+				.AddNode(x => CompileUnrealEditorWin64Async(x, crashReporterCompileArgs, embedSrcSrvInfo, compileDatasmithPlugins, withFullDebugInfo, signExecutables))
+				.Requires(versionFilesNode);
 
+			aggregates.Add(new BgAggregate("Win64 Editor", winEditorNode, label: "Editors/Win64"));
 
 			/////// TARGET PLATFORMS ////////////////////////////////////////////////
 
 			//// Win64 ////
 
-			BgAgent TargetWin64 = new BgAgent("Target Win64", "Win64_Licensee");
+			BgAgent targetWin64 = new BgAgent("Target Win64", "Win64_Licensee");
 
-			BgNode<BgFileSet> WinGame = TargetWin64
-				.AddNode(x => CompileUnrealGameWin64(x, GameConfigurations, EmbedSrcSrvInfo, WithFullDebugInfo, SignExecutables))
-				.Requires(VersionFilesNode);
+			BgNode<BgFileSet> winGame = targetWin64
+				.AddNode(x => CompileUnrealGameWin64Async(x, gameConfigurations, embedSrcSrvInfo, withFullDebugInfo, signExecutables))
+				.Requires(versionFilesNode);
 
-			Aggregates.Add(new BgAggregate("TargetPlatforms_Win64", WinGame));
+			aggregates.Add(new BgAggregate("TargetPlatforms_Win64", winGame));
 
 			/////// TOOLS //////////////////////////////////////////////////////////
 
 			//// Build Rules ////
 
-			BgAgent BuildRules = new BgAgent("BuildRules", "Win64_Licensee");
+			BgAgent buildRules = new BgAgent("BuildRules", "Win64_Licensee");
 
-			BgNode RulesAssemblies = BuildRules
-				.AddNode(x => CompileRulesAssemblies(x));
+			BgNode rulesAssemblies = buildRules
+				.AddNode(x => CompileRulesAssembliesAsync(x));
 
 			//// Win Tools ////
 
-			BgAgent ToolsGroupWin64 = new BgAgent("Tools Group Win64", "Win64_Licensee");
+			BgAgent toolsGroupWin64 = new BgAgent("Tools Group Win64", "Win64_Licensee");
 
-			BgNode<BgFileSet> WinTools = ToolsGroupWin64
-				.AddNode(x => BuildToolsWin64Async(x, CrashReporterCompileArgs))
-				.Requires(VersionFilesNode);
-
-			BgNode<BgFileSet> CsTools = ToolsGroupWin64
-				.AddNode(x => BuildToolsCSAsync(x, SignExecutables))
-				.Requires(VersionFilesNode);
+			BgNode<BgFileSet> winTools = toolsGroupWin64
+				.AddNode(x => BuildToolsWin64Async(x, crashReporterCompileArgs))
+				.Requires(versionFilesNode);
 
+			BgNode<BgFileSet> csTools = toolsGroupWin64
+				.AddNode(x => BuildToolsCsAsync(x, signExecutables))
+				.Requires(versionFilesNode);
 
 			/////// DDC //////////////////////////////////////////////////////////
 
-			BgAgent DDCGroupWin64 = new BgAgent("DDC Group Win64", "Win64_Licensee");
-
-			BgList<BgString> DDCPlatformsWin64 = BgList<BgString>.Create("WindowsEditor");
-			DDCPlatformsWin64 = DDCPlatformsWin64.If(WithWin64, x => x.Add("Windows"));
+			BgAgent ddcGroupWin64 = new BgAgent("DDC Group Win64", "Win64_Licensee");
 
-			BgNode DdcNode = DDCGroupWin64
-				.AddNode(x => BuildDDCWin64Async(x, DDCPlatformsWin64, BgList<BgFileSet>.Create(WinEditorNode.Output, WinTools.Output)))
-				.Requires(WinEditorNode, WinTools);
+			BgList<BgString> ddcPlatformsWin64 = BgList<BgString>.Create("WindowsEditor");
+			ddcPlatformsWin64 = ddcPlatformsWin64.If(withWin64, x => x.Add("Windows"));
 
+			BgNode ddcNode = ddcGroupWin64
+				.AddNode(x => BuildDdcWin64Async(x, ddcPlatformsWin64, BgList<BgFileSet>.Create(winEditorNode.Output, winTools.Output)))
+				.Requires(winEditorNode, winTools);
 
 			/////// STAGING ///////
 
 			// Windows 
-			BgAgent WinStageAgent = new BgAgent("Installed Build Group Win64", "Win64_Licensee");
+			BgAgent winStageAgent = new BgAgent("Installed Build Group Win64", "Win64_Licensee");
 
-			BgList<BgFileSet> WinInstalledFiles = BgList<BgFileSet>.Empty;
-			WinInstalledFiles = WinInstalledFiles.Add(WinEditorNode.Output);
-			WinInstalledFiles = WinInstalledFiles.Add(WinTools.Output);
-			WinInstalledFiles = WinInstalledFiles.Add(CsTools.Output);
-			WinInstalledFiles = WinInstalledFiles.If(WithWin64, x => x.Add(WinGame.Output));
+			BgList<BgFileSet> winInstalledFiles = BgList<BgFileSet>.Empty;
+			winInstalledFiles = winInstalledFiles.Add(winEditorNode.Output);
+			winInstalledFiles = winInstalledFiles.Add(winTools.Output);
+			winInstalledFiles = winInstalledFiles.Add(csTools.Output);
+			winInstalledFiles = winInstalledFiles.If(withWin64, x => x.Add(winGame.Output));
 
-			BgList<BgString> WinPlatforms = BgList<BgString>.Empty;
-			WinPlatforms = WinPlatforms.If(WithWin64, x => x.Add("Win64"));
+			BgList<BgString> winPlatforms = BgList<BgString>.Empty;
+			winPlatforms = winPlatforms.If(withWin64, x => x.Add("Win64"));
 
-			BgList<BgString> WinContentOnlyPlatforms = BgList<BgString>.Empty;
+			BgList<BgString> winContentOnlyPlatforms = BgList<BgString>.Empty;
 
-			BgString WinOutputDir = "LocalBuilds/Engine/Windows";
-			BgString WinFinalizeArgs = BgString.Format("-OutputDir=\"{0}\" -Platforms={1} -ContentOnlyPlatforms={2}", WinOutputDir, BgString.Join(";", WinPlatforms), BgString.Join(";", WinContentOnlyPlatforms));
+			BgString winOutputDir = "LocalBuilds/Engine/Windows";
+			BgString winFinalizeArgs = BgString.Format("-OutputDir=\"{0}\" -Platforms={1} -ContentOnlyPlatforms={2}", winOutputDir, BgString.Join(";", winPlatforms), BgString.Join(";", winContentOnlyPlatforms));
 
-			BgNode WinInstalledNode = WinStageAgent
-				.AddNode(x => MakeInstalledBuildWin64Async(x, WinInstalledFiles, WinFinalizeArgs, WinOutputDir))
-				.Requires(WinInstalledFiles);
+			BgNode winInstalledNode = winStageAgent
+				.AddNode(x => MakeInstalledBuildWin64Async(x, winInstalledFiles, winFinalizeArgs, winOutputDir))
+				.Requires(winInstalledFiles);
 
-			Aggregates.Add(new BgAggregate("HostPlatforms_Win64", WinInstalledNode, label: "Builds/Win64"));
+			aggregates.Add(new BgAggregate("HostPlatforms_Win64", winInstalledNode, label: "Builds/Win64"));
 
-			return new BgGraph(BgList<BgNode>.Empty, Aggregates);
+			return new BgGraph(BgList<BgNode>.Empty, aggregates);
 		}
 
 		/// <summary>
 		/// Update the build version
 		/// </summary>
-		static async Task UpdateVersionFilesAsync(BgContext State)
+		static async Task UpdateVersionFilesAsync(BgContext state)
 		{
-			if (State.IsBuildMachine)
+			if (state.IsBuildMachine)
 			{
-				await SetVersionAsync(State.Change, State.Stream.Replace('/', '+'));
+				await SetVersionAsync(state.Change, state.Stream.Replace('/', '+'));
 			}
 		}
 
@@ -233,175 +228,185 @@ namespace AutomationTool
 		/// Builds the Windows editor
 		/// </summary>
 		[BgNodeName("Compile UnrealEditor Win64")]
-		static async Task<BgFileSet> CompileUnrealEditorWin64Async(BgContext State, BgString CrashReporterCompileArgs, BgBool EmbedSrcSrvInfo, BgBool CompileDatasmithPlugins, BgBool WithFullDebugInfo, BgBool SignExecutables)
+		static async Task<BgFileSet> CompileUnrealEditorWin64Async(BgContext state, BgString crashReporterCompileArgs, BgBool embedSrcSrvInfo, BgBool compileDatasmithPlugins, BgBool withFullDebugInfo, BgBool signExecutables)
 		{
-			FileSet OutputFiles = FileSet.Empty;
-			OutputFiles += await CompileAsync("UnrealEditor", UnrealTargetPlatform.Win64, UnrealTargetConfiguration.DebugGame, Arguments: $"-precompile -allmodules {State.Get(CrashReporterCompileArgs)}");
-			OutputFiles += await CompileAsync("UnrealEditor", UnrealTargetPlatform.Win64, UnrealTargetConfiguration.Development, Arguments: $"-precompile -allmodules {State.Get(CrashReporterCompileArgs)}");
+			_ = compileDatasmithPlugins;
+
+			FileSet outputFiles = FileSet.Empty;
+			outputFiles += await CompileAsync("UnrealEditor", UnrealTargetPlatform.Win64, UnrealTargetConfiguration.DebugGame, arguments: $"-precompile -allmodules {state.Get(crashReporterCompileArgs)}");
+			outputFiles += await CompileAsync("UnrealEditor", UnrealTargetPlatform.Win64, UnrealTargetConfiguration.Development, arguments: $"-precompile -allmodules {state.Get(crashReporterCompileArgs)}");
 
-			if (State.Get(EmbedSrcSrvInfo))
+			if (state.Get(embedSrcSrvInfo))
 			{
 				// Embed source info into the PDB files. Should be done from this machine to ensure that paths are correct.
 				Logger.LogInformation("Embedding source file information into PDB files...");
-				FileSet SourceFiles = Workspace.Filter("Engine/Source/...;Engine/Plugins/...").Except("Engine/Source/ThirdParty/...").Filter("*.c;*.h;*.cpp;*.hpp;*.inl");
-//				State.SrcSrv(BinaryFiles: Full, SourceFiles: SourceFiles);
+				FileSet sourceFiles = Workspace.Filter("Engine/Source/...;Engine/Plugins/...").Except("Engine/Source/ThirdParty/...").Filter("*.c;*.h;*.cpp;*.hpp;*.inl");
+				//				State.SrcSrv(BinaryFiles: Full, SourceFiles: SourceFiles);
+				_ = sourceFiles;
 			}
 
-			if (!State.Get(WithFullDebugInfo))
+			if (!state.Get(withFullDebugInfo))
 			{
-				FileSet UnstrippedFiles = OutputFiles.Filter(WinStripFilter);
-				OutputFiles += await StripAsync(UnstrippedFiles, UnrealTargetPlatform.Win64, BaseDir: RootDir, OutputDir: TempMiscDir);
-				OutputFiles -= UnstrippedFiles;
+				FileSet unstrippedFiles = outputFiles.Filter(s_winStripFilter);
+				outputFiles += await StripAsync(unstrippedFiles, UnrealTargetPlatform.Win64, baseDir: RootDir, outputDir: TempMiscDir);
+				outputFiles -= unstrippedFiles;
 			}
 
-			if(State.Get(SignExecutables))
+			if (state.Get(signExecutables))
 			{
-//				FileSet UnsignedFiles = OutputFiles.Filter(WinSignFilter);
-//				FileSet FilesToCopy = UnsignedFiles.Except((new HashSet<FileReference>(UnsignedFiles.Where(x => !x.IsUnderDirectory(TempMiscDir)));
-//				UnsignedFiles += State.Copy(FilesToCopy.Flatten(TempMiscDir));
-//				UnsignedFiles.Sign();
-//				OutputFiles -= FilesToCopy;
-//				OutputFiles += UnsignedFiles;
+				//				FileSet UnsignedFiles = OutputFiles.Filter(WinSignFilter);
+				//				FileSet FilesToCopy = UnsignedFiles.Except((new HashSet<FileReference>(UnsignedFiles.Where(x => !x.IsUnderDirectory(TempMiscDir)));
+				//				UnsignedFiles += State.Copy(FilesToCopy.Flatten(TempMiscDir));
+				//				UnsignedFiles.Sign();
+				//				OutputFiles -= FilesToCopy;
+				//				OutputFiles += UnsignedFiles;
 			}
 
-			return OutputFiles;
+			return outputFiles;
 		}
 
 		/// <summary>
 		/// Builds the game target
 		/// </summary>
 		[BgNodeName("Compile UnrealGame Win64")]
-		static async Task<BgFileSet> CompileUnrealGameWin64(BgContext State, BgList<BgEnum<UnrealTargetConfiguration>> Configurations, BgBool EmbedSrcSrvInfo, BgBool WithFullDebugInfo, BgBool SignExecutables)
+		static async Task<BgFileSet> CompileUnrealGameWin64Async(BgContext state, BgList<BgEnum<UnrealTargetConfiguration>> configurations, BgBool embedSrcSrvInfo, BgBool withFullDebugInfo, BgBool signExecutables)
 		{
-			FileSet Files = FileSet.Empty;
+			_ = embedSrcSrvInfo;
+			_ = withFullDebugInfo;
+			_ = signExecutables;
 
-			List<UnrealTargetConfiguration> ConfigurationsValue = State.Get(Configurations);
-			foreach (UnrealTargetConfiguration Configuration in ConfigurationsValue)
+			FileSet files = FileSet.Empty;
+
+			List<UnrealTargetConfiguration> configurationsValue = state.Get(configurations);
+			foreach (UnrealTargetConfiguration configuration in configurationsValue)
 			{
-				Files += await CompileAsync("UnrealGame", UnrealTargetPlatform.Win64, Configuration, Arguments: "-precompile -allmodules -nolink");
-				Files += await CompileAsync("UnrealGame", UnrealTargetPlatform.Win64, Configuration, Arguments: "-precompile", Clean: false);
+				files += await CompileAsync("UnrealGame", UnrealTargetPlatform.Win64, configuration, arguments: "-precompile -allmodules -nolink");
+				files += await CompileAsync("UnrealGame", UnrealTargetPlatform.Win64, configuration, arguments: "-precompile", clean: false);
 			}
 
-			return Files;
+			return files;
 		}
 
-		static async Task<BgFileSet> CompileRulesAssemblies(BgContext State)
+		static async Task<BgFileSet> CompileRulesAssembliesAsync(BgContext state)
 		{
-			FileReference UnrealBuildToolDll = FileReference.Combine(RootDir, "Engine/Binaries/DotNET/UnrealBuildTool/UnrealBuildTool.dll");
+			_ = state;
+
+			FileReference unrealBuildToolDll = FileReference.Combine(RootDir, "Engine/Binaries/DotNET/UnrealBuildTool/UnrealBuildTool.dll");
 
-			await SpawnAsync(Unreal.DotnetPath.FullName, Arguments: $"\"{UnrealBuildToolDll}\" -Mode=QueryTargets");
+			await SpawnAsync(Unreal.DotnetPath.FullName, arguments: $"\"{unrealBuildToolDll}\" -Mode=QueryTargets");
 
 			return Workspace.Filter("Engine/Intermediate/Build/BuildRules/...");
 		}
 
-		static async Task<BgFileSet> BuildToolsWin64Async(BgContext State, BgString CrashReporterCompileArgs)
+		static async Task<BgFileSet> BuildToolsWin64Async(BgContext state, BgString crashReporterCompileArgs)
 		{
-			string CrashReportClientArgs = State.Get(CrashReporterCompileArgs);
+			string crashReportClientArgs = state.Get(crashReporterCompileArgs);
 
 			// State.Tag(Files: "#NotForLicensee Build Tools Win64", With: Files);
 
-			FileSet Files = FileSet.Empty;
-			Files += await CompileAsync("CrashReportClient", UnrealTargetPlatform.Win64, UnrealTargetConfiguration.Shipping, Arguments: CrashReportClientArgs);
-			Files += await CompileAsync("CrashReportClientEditor", UnrealTargetPlatform.Win64, UnrealTargetConfiguration.Shipping, Arguments: CrashReportClientArgs);
-			Files += await CompileAsync("ShaderCompileWorker", UnrealTargetPlatform.Win64, UnrealTargetConfiguration.Development);
-			Files += await CompileAsync("EpicWebHelper", UnrealTargetPlatform.Win64, UnrealTargetConfiguration.Development);
-			Files += await CompileAsync("UnrealInsights", UnrealTargetPlatform.Win64, UnrealTargetConfiguration.Development);
-			Files += await CompileAsync("UnrealFrontend", UnrealTargetPlatform.Win64, UnrealTargetConfiguration.Development);
-			Files += await CompileAsync("UnrealLightmass", UnrealTargetPlatform.Win64, UnrealTargetConfiguration.Development);
-			Files += await CompileAsync("InterchangeWorker", UnrealTargetPlatform.Win64, UnrealTargetConfiguration.Development);
-			Files += await CompileAsync("UnrealPak", UnrealTargetPlatform.Win64, UnrealTargetConfiguration.Development);
-			Files += await CompileAsync("UnrealMultiUserServer", UnrealTargetPlatform.Win64, UnrealTargetConfiguration.Development);
-			Files += await CompileAsync("UnrealRecoverySvc", UnrealTargetPlatform.Win64, UnrealTargetConfiguration.Development);
-			Files += await CompileAsync("LiveCodingConsole", UnrealTargetPlatform.Win64, UnrealTargetConfiguration.Development);
-			Files += await CompileAsync("BootstrapPackagedGame", UnrealTargetPlatform.Win64, UnrealTargetConfiguration.Shipping);
-			Files += await CompileAsync("BuildPatchTool", UnrealTargetPlatform.Win64, UnrealTargetConfiguration.Shipping);
-			Files += await CompileAsync("SwitchboardListener", UnrealTargetPlatform.Win64, UnrealTargetConfiguration.Development);
-			Files += FileSet.FromFile(RootDir, "Engine/Binaries/Win64/XGEControlWorker.exe");
-			Files += FileSet.FromFile(RootDir, "Engine/Saved/UnrealBuildTool/BuildConfiguration.Schema.xsd");
-
-			return Files;
+			FileSet files = FileSet.Empty;
+			files += await CompileAsync("CrashReportClient", UnrealTargetPlatform.Win64, UnrealTargetConfiguration.Shipping, arguments: crashReportClientArgs);
+			files += await CompileAsync("CrashReportClientEditor", UnrealTargetPlatform.Win64, UnrealTargetConfiguration.Shipping, arguments: crashReportClientArgs);
+			files += await CompileAsync("ShaderCompileWorker", UnrealTargetPlatform.Win64, UnrealTargetConfiguration.Development);
+			files += await CompileAsync("EpicWebHelper", UnrealTargetPlatform.Win64, UnrealTargetConfiguration.Development);
+			files += await CompileAsync("UnrealInsights", UnrealTargetPlatform.Win64, UnrealTargetConfiguration.Development);
+			files += await CompileAsync("UnrealFrontend", UnrealTargetPlatform.Win64, UnrealTargetConfiguration.Development);
+			files += await CompileAsync("UnrealLightmass", UnrealTargetPlatform.Win64, UnrealTargetConfiguration.Development);
+			files += await CompileAsync("InterchangeWorker", UnrealTargetPlatform.Win64, UnrealTargetConfiguration.Development);
+			files += await CompileAsync("UnrealPak", UnrealTargetPlatform.Win64, UnrealTargetConfiguration.Development);
+			files += await CompileAsync("UnrealMultiUserServer", UnrealTargetPlatform.Win64, UnrealTargetConfiguration.Development);
+			files += await CompileAsync("UnrealRecoverySvc", UnrealTargetPlatform.Win64, UnrealTargetConfiguration.Development);
+			files += await CompileAsync("LiveCodingConsole", UnrealTargetPlatform.Win64, UnrealTargetConfiguration.Development);
+			files += await CompileAsync("BootstrapPackagedGame", UnrealTargetPlatform.Win64, UnrealTargetConfiguration.Shipping);
+			files += await CompileAsync("BuildPatchTool", UnrealTargetPlatform.Win64, UnrealTargetConfiguration.Shipping);
+			files += await CompileAsync("SwitchboardListener", UnrealTargetPlatform.Win64, UnrealTargetConfiguration.Development);
+			files += FileSet.FromFile(RootDir, "Engine/Binaries/Win64/XGEControlWorker.exe");
+			files += FileSet.FromFile(RootDir, "Engine/Saved/UnrealBuildTool/BuildConfiguration.Schema.xsd");
+
+			return files;
 		}
 
 		[BgNodeName("Build Tools CS")]
-		static async Task<BgFileSet> BuildToolsCSAsync(BgContext State, BgBool SignExecutables)
+		static async Task<BgFileSet> BuildToolsCsAsync(BgContext state, BgBool signExecutables)
 		{
 			FileUtils.ForceDeleteDirectory(TempCsToolsDir);
 
 			// Copy Source and referenced libraries to a new location with Confidential folders removed
-			FileSet UatProjects = Workspace.Filter("Engine/Source/Programs/AutomationTool/....csproj");
+			FileSet uatProjects = Workspace.Filter("Engine/Source/Programs/AutomationTool/....csproj");
+			_ = uatProjects;
 
-			FileSet RedistUatSource = Workspace.Filter("Engine/Binaries/DotNET/...;Engine/Binaries/ThirdParty/Newtonsoft/...;Engine/Binaries/ThirdParty/IOS/...;Engine/Binaries/ThirdParty/VisualStudio/...;Engine/Source/Programs/...;Engine/Platforms/*/Source/Programs/...;Engine/Source/Editor/SwarmInterface/...");
-			await RedistUatSource.CopyToAsync(Unreal.RootDirectory);
+			FileSet redistUatSource = Workspace.Filter("Engine/Binaries/DotNET/...;Engine/Binaries/ThirdParty/Newtonsoft/...;Engine/Binaries/ThirdParty/IOS/...;Engine/Binaries/ThirdParty/VisualStudio/...;Engine/Source/Programs/...;Engine/Platforms/*/Source/Programs/...;Engine/Source/Editor/SwarmInterface/...");
+			await redistUatSource.CopyToAsync(Unreal.RootDirectory);
 
 			// Compile all the tools
-			CsCompileOutput Output = CsCompileOutput.Empty;
+			CsCompileOutput output = CsCompileOutput.Empty;
 			if (HostPlatform.Current.HostEditorPlatform == UnrealTargetPlatform.Win64)
 			{
-				Output += await CsCompileAsync(FileReference.Combine(TempCsToolsDir, "Engine/Source/Programs/UnrealSwarm/SwarmCoordinator/SwarmCoordinator.csproj"), Platform: "AnyCPU", Configuration: "Development");
-				Output += await CsCompileAsync(FileReference.Combine(TempCsToolsDir, "Engine/Source/Programs/UnrealSwarm/Agent/Agent.csproj"), Platform: "AnyCPU", Configuration: "Development");
-				Output += await CsCompileAsync(FileReference.Combine(TempCsToolsDir, "Engine/Source/Editor/SwarmInterface/DotNET/SwarmInterface.csproj"), Platform: "AnyCPU", Configuration: "Development");
-				Output += await CsCompileAsync(FileReference.Combine(TempCsToolsDir, "Engine/Source/Programs/UnrealControls/UnrealControls.csproj"), Platform: "AnyCPU", Configuration: "Development");
-				Output += await CsCompileAsync(FileReference.Combine(TempCsToolsDir, "Engine/Source/Programs/IOS/iPhonePackager/iPhonePackager.csproj"), Platform: "AnyCPU", Configuration: "Development", Arguments: "/verbosity:minimal /target:Rebuild");
-				Output += await CsCompileAsync(FileReference.Combine(TempCsToolsDir, "Engine/Source/Programs/NetworkProfiler/NetworkProfiler/NetworkProfiler.csproj"), Platform: "AnyCPU", Configuration: "Development");
+				output += await CsCompileAsync(FileReference.Combine(TempCsToolsDir, "Engine/Source/Programs/UnrealSwarm/SwarmCoordinator/SwarmCoordinator.csproj"), platform: "AnyCPU", configuration: "Development");
+				output += await CsCompileAsync(FileReference.Combine(TempCsToolsDir, "Engine/Source/Programs/UnrealSwarm/Agent/Agent.csproj"), platform: "AnyCPU", configuration: "Development");
+				output += await CsCompileAsync(FileReference.Combine(TempCsToolsDir, "Engine/Source/Editor/SwarmInterface/DotNET/SwarmInterface.csproj"), platform: "AnyCPU", configuration: "Development");
+				output += await CsCompileAsync(FileReference.Combine(TempCsToolsDir, "Engine/Source/Programs/UnrealControls/UnrealControls.csproj"), platform: "AnyCPU", configuration: "Development");
+				output += await CsCompileAsync(FileReference.Combine(TempCsToolsDir, "Engine/Source/Programs/IOS/iPhonePackager/iPhonePackager.csproj"), platform: "AnyCPU", configuration: "Development", arguments: "/verbosity:minimal /target:Rebuild");
+				output += await CsCompileAsync(FileReference.Combine(TempCsToolsDir, "Engine/Source/Programs/NetworkProfiler/NetworkProfiler/NetworkProfiler.csproj"), platform: "AnyCPU", configuration: "Development");
 			}
 
-			FileSet Binaries = Output.Binaries + Output.References;
+			FileSet binaries = output.Binaries + output.References;
 
 			// Tag AutomationTool and UnrealBuildTool folders recursively as NET Core dependencies are not currently handled by CsCompile
-			Binaries += FileSet.FromDirectory(TempCsToolsDir).Filter("Engine/Binaries/DotNET/AutomationTool/...");
-			Binaries += FileSet.FromDirectory(TempCsToolsDir).Filter("Engine/Binaries/DotNET/UnrealBuildTool/...");
+			binaries += FileSet.FromDirectory(TempCsToolsDir).Filter("Engine/Binaries/DotNET/AutomationTool/...");
+			binaries += FileSet.FromDirectory(TempCsToolsDir).Filter("Engine/Binaries/DotNET/UnrealBuildTool/...");
 
 			// Tag AutomationTool Script module build records, so that prebuilt modules may be discovered in the absence of source code
-			Binaries += FileSet.FromDirectory(TempCsToolsDir).Filter("Engine/Intermediate/ScriptModules/...");
+			binaries += FileSet.FromDirectory(TempCsToolsDir).Filter("Engine/Intermediate/ScriptModules/...");
 
-			if(State.Get(SignExecutables))
+			if (state.Get(signExecutables))
 			{
-//				Binaries.Sign();
+				//				Binaries.Sign();
 			}
 
-			return Binaries;
+			return binaries;
 		}
 
 		/// <summary>
 		/// Creates a DDC pack file for the supported platforms
 		/// </summary>
 		[BgNodeName("Build DDC Win64")]
-		static async Task<BgFileSet> BuildDDCWin64Async(BgContext State, BgList<BgString> Platforms, BgList<BgFileSet> Dependencies)
+		static async Task<BgFileSet> BuildDdcWin64Async(BgContext state, BgList<BgString> platforms, BgList<BgFileSet> dependencies)
 		{
 			// Build up a list of files needed to build DDC
-			FileSet ToCopy = State.Get(Dependencies);
-			ToCopy += await CsCompileAsync(FileReference.Combine(RootDir, "Engine/Source/Programs/UnrealBuildTool/UnrealBuildTool.csproj"), Platform: "AnyCPU", Configuration: "Development", EnumerateOnly: true).MergeAsync();
-			ToCopy += await ToCopy.Filter("*.target").TagReceiptsAsync(RuntimeDependencies: true);
-			ToCopy += FileSet.FromFile(RootDir, "Engine/Binaries/DotNET/Ionic.Zip.Reduced.dll");
-			ToCopy += FileSet.FromFile(RootDir, "Engine/Binaries/DotNET/OneSky.dll");
-			ToCopy += FileSet.FromDirectory(RootDir).Filter("Templates/TemplateResources/...");
-
-			foreach (DirectoryReference ExtensionDir in Unreal.GetExtensionDirs(Unreal.EngineDirectory))
+			FileSet toCopy = state.Get(dependencies);
+			toCopy += await CsCompileAsync(FileReference.Combine(RootDir, "Engine/Source/Programs/UnrealBuildTool/UnrealBuildTool.csproj"), platform: "AnyCPU", configuration: "Development", enumerateOnly: true).MergeAsync();
+			toCopy += await toCopy.Filter("*.target").TagReceiptsAsync(runtimeDependencies: true);
+			toCopy += FileSet.FromFile(RootDir, "Engine/Binaries/DotNET/Ionic.Zip.Reduced.dll");
+			toCopy += FileSet.FromFile(RootDir, "Engine/Binaries/DotNET/OneSky.dll");
+			toCopy += FileSet.FromDirectory(RootDir).Filter("Templates/TemplateResources/...");
+
+			foreach (DirectoryReference extensionDir in Unreal.GetExtensionDirs(Unreal.EngineDirectory))
 			{
-				ToCopy += FileSet.FromDirectory(ExtensionDir).Filter("Content/...").Except("....psd;....pdn;....fbx;....po");
-				ToCopy += FileSet.FromDirectory(ExtensionDir).Filter("Config/...").Except("....vdf");
-				ToCopy += FileSet.FromDirectory(ExtensionDir).Filter("Shaders/...");
+				toCopy += FileSet.FromDirectory(extensionDir).Filter("Content/...").Except("....psd;....pdn;....fbx;....po");
+				toCopy += FileSet.FromDirectory(extensionDir).Filter("Config/...").Except("....vdf");
+				toCopy += FileSet.FromDirectory(extensionDir).Filter("Shaders/...");
 
-				DirectoryReference ExtensionPluginsDir = DirectoryReference.Combine(ExtensionDir, "Plugins");
-				ToCopy += FileSet.FromDirectory(ExtensionPluginsDir).Filter("....uplugin;.../Config/...;.../Content/...;.../Resources/...;.../Shaders/...;.../Templates/...").Except(".../TwitchLiveStreaming/...");
+				DirectoryReference extensionPluginsDir = DirectoryReference.Combine(extensionDir, "Plugins");
+				toCopy += FileSet.FromDirectory(extensionPluginsDir).Filter("....uplugin;.../Config/...;.../Content/...;.../Resources/...;.../Shaders/...;.../Templates/...").Except(".../TwitchLiveStreaming/...");
 			}
 
 			// Filter out the files not needed to build DDC. Removing confidential folders can affect DDC keys, so we want to be sure that we're making DDC with a build that can use it.
-			FileSet FilteredCopyList = ToCopy.Except(".../Source/...;.../Intermediate/...");
+			FileSet filteredCopyList = toCopy.Except(".../Source/...;.../Intermediate/...");
 
-			Dictionary<FileReference, FileReference> TargetFileToSourceFile = new Dictionary<FileReference, FileReference>();
-			MapFilesToOutputDir(FilteredCopyList, TempDdcDir, TargetFileToSourceFile);
+			Dictionary<FileReference, FileReference> targetFileToSourceFile = new Dictionary<FileReference, FileReference>();
+			MapFilesToOutputDir(filteredCopyList, TempDdcDir, targetFileToSourceFile);
 
 			FileUtils.ForceDeleteDirectoryContents(TempDdcDir);
-			await FilteredCopyList.CopyToAsync(TempDdcDir);
+			await filteredCopyList.CopyToAsync(TempDdcDir);
 
 			// Run the DDC commandlet
-			List<string> Arguments = new List<string>();
-			Arguments.Add($"-TempDir=\"{TempDdcDir}\"");
-			Arguments.Add($"-FeaturePacks=\"{String.Join(";", GetDdcProjects(UnrealTargetPlatform.Win64))}\"");
-			Arguments.Add($"-TargetPlatforms={String.Join("+", State.Get(Platforms))}");
-			Arguments.Add($"-HostPlatform=Win64");
-			await State.CommandAsync("BuildDerivedDataCache", Arguments: String.Join(" ", Arguments));
+			List<string> arguments = new List<string>();
+			arguments.Add($"-TempDir=\"{TempDdcDir}\"");
+			arguments.Add($"-FeaturePacks=\"{String.Join(";", GetDdcProjects(UnrealTargetPlatform.Win64))}\"");
+			arguments.Add($"-TargetPlatforms={String.Join("+", state.Get(platforms))}");
+			arguments.Add($"-HostPlatform=Win64");
+			await state.CommandAsync("BuildDerivedDataCache", arguments: String.Join(" ", arguments));
 
 			// Return a tag for the output file
 			return FileSet.FromFile(TempDdcDir, "Engine/DerivedDataCache/Compressed.ddp");
@@ -411,69 +416,69 @@ namespace AutomationTool
 		/// Copy all the build artifacts to the output folder
 		/// </summary>
 		[BgNodeName("Make Installed Build Win64")]
-		static async Task<BgFileSet> MakeInstalledBuildWin64Async(BgContext State, BgList<BgFileSet> InputFiles, BgString FinalizeArgs, BgString OutputDir)
+		static async Task<BgFileSet> MakeInstalledBuildWin64Async(BgContext state, BgList<BgFileSet> inputFiles, BgString finalizeArgs, BgString outputDir)
 		{
 			// Find all the input files, and add any runtime dependencies from the receipts into the list
-			FileSet SourceFiles = State.Get(InputFiles);
-			SourceFiles += await SourceFiles.Filter("*.target").TagReceiptsAsync(RuntimeDependencies: true);
+			FileSet sourceFiles = state.Get(inputFiles);
+			sourceFiles += await sourceFiles.Filter("*.target").TagReceiptsAsync(runtimeDependencies: true);
 
 			// Include any files referenced by dependency lists
-			FileSet DependencyListFiles = SourceFiles.Filter(".../DependencyList.txt;.../DependencyList-AllModules.txt");
-			foreach (FileReference DependencyListFile in DependencyListFiles)
+			FileSet dependencyListFiles = sourceFiles.Filter(".../DependencyList.txt;.../DependencyList-AllModules.txt");
+			foreach (FileReference dependencyListFile in dependencyListFiles)
 			{
-				string[] Lines = await FileReference.ReadAllLinesAsync(DependencyListFile);
-				foreach (string Line in Lines)
+				string[] lines = await FileReference.ReadAllLinesAsync(dependencyListFile);
+				foreach (string line in lines)
 				{
-					string TrimLine = Line.Trim();
-					if (TrimLine.Length > 0)
+					string trimLine = line.Trim();
+					if (trimLine.Length > 0)
 					{
-						SourceFiles += FileSet.FromFile(Unreal.RootDirectory, TrimLine);
+						sourceFiles += FileSet.FromFile(Unreal.RootDirectory, trimLine);
 					}
 				}
 			}
-			SourceFiles -= DependencyListFiles;
+			sourceFiles -= dependencyListFiles;
 
 			// Clear the output directory
-			DirectoryReference OutputDirRef = new DirectoryReference(State.Get(OutputDir));
-			FileUtils.ForceDeleteDirectoryContents(OutputDirRef);
-			FileSet InstalledFiles = await SourceFiles.CopyToAsync(OutputDirRef);
+			DirectoryReference outputDirRef = new DirectoryReference(state.Get(outputDir));
+			FileUtils.ForceDeleteDirectoryContents(outputDirRef);
+			FileSet installedFiles = await sourceFiles.CopyToAsync(outputDirRef);
 
 			// Run the finalize command 
-			await State.CommandAsync("FinalizeInstalledBuild", Arguments: State.Get(FinalizeArgs));
+			await state.CommandAsync("FinalizeInstalledBuild", arguments: state.Get(finalizeArgs));
 
 			// Get the final list of files
-			InstalledFiles += FileSet.FromFile(OutputDirRef, "Engine/Build/InstalledEngine.txt");
+			installedFiles += FileSet.FromFile(outputDirRef, "Engine/Build/InstalledEngine.txt");
 
 			// Sanitize any receipts in the output directory
-			await InstalledFiles.Filter("*.target").SanitizeReceiptsAsync();
+			await installedFiles.Filter("*.target").SanitizeReceiptsAsync();
 
-			return InstalledFiles;
+			return installedFiles;
 		}
 
-		static void MapFilesToOutputDir(IEnumerable<FileReference> SourceFiles, DirectoryReference TargetDir, Dictionary<FileReference, FileReference> TargetFileToSourceFile)
+		static void MapFilesToOutputDir(IEnumerable<FileReference> sourceFiles, DirectoryReference targetDir, Dictionary<FileReference, FileReference> targetFileToSourceFile)
 		{
-			foreach (FileReference SourceFile in SourceFiles)
+			foreach (FileReference sourceFile in sourceFiles)
 			{
-				DirectoryReference BaseDir;
-				if (SourceFile.IsUnderDirectory(TempCsToolsDir))
+				DirectoryReference baseDir;
+				if (sourceFile.IsUnderDirectory(TempCsToolsDir))
 				{
-					BaseDir = TempCsToolsDir;
+					baseDir = TempCsToolsDir;
 				}
-				else if (SourceFile.IsUnderDirectory(TempMiscDir))
+				else if (sourceFile.IsUnderDirectory(TempMiscDir))
 				{
-					BaseDir = TempMiscDir;
+					baseDir = TempMiscDir;
 				}
-				else if (SourceFile.IsUnderDirectory(TempDdcDir))
+				else if (sourceFile.IsUnderDirectory(TempDdcDir))
 				{
-					BaseDir = TempDdcDir;
+					baseDir = TempDdcDir;
 				}
 				else
 				{
-					BaseDir = Unreal.RootDirectory;
+					baseDir = Unreal.RootDirectory;
 				}
 
-				FileReference TargetFile = FileReference.Combine(TargetDir, SourceFile.MakeRelativeTo(BaseDir));
-				TargetFileToSourceFile[TargetFile] = SourceFile;
+				FileReference targetFile = FileReference.Combine(targetDir, sourceFile.MakeRelativeTo(baseDir));
+				targetFileToSourceFile[targetFile] = sourceFile;
 			}
 		}
 	}
diff --git a/Engine/Source/Programs/AutomationTool/BuildGraph/Task.cs b/Engine/Source/Programs/AutomationTool/BuildGraph/Task.cs
--- a/Engine/Source/Programs/AutomationTool/BuildGraph/Task.cs
+++ b/Engine/Source/Programs/AutomationTool/BuildGraph/Task.cs
@@ -1,6 +1,5 @@
 // Copyright Epic Games, Inc. All Rights Reserved.
 
-using EpicGames.BuildGraph;
 using System;
 using System.Collections.Generic;
 using System.IO;
@@ -10,13 +9,9 @@ using System.Text;
 using System.Threading.Tasks;
 using System.Xml;
 using EpicGames.Core;
+using Microsoft.Extensions.Logging;
 using OpenTracing;
 using UnrealBuildBase;
-using UnrealBuildTool;
-using Microsoft.Extensions.Logging;
-using System.Runtime.InteropServices;
-
-using static AutomationTool.CommandUtils;
 
 #nullable enable
 
@@ -77,25 +72,25 @@ namespace AutomationTool
 
 	/// <summary>
 	/// Proxy to handle executing multiple tasks simultaneously (such as compile tasks). If a task supports simultaneous execution, it can return a separate
-	/// executor an executor instance from GetExecutor() callback. If not, it must implement Execute().
+	/// executor an executor instance from GetExecutor() callback. If not, it must implement ExecuteAsync().
 	/// </summary>
 	public interface ITaskExecutor
 	{
 		/// <summary>
 		/// Adds another task to this executor
 		/// </summary>
-		/// <param name="Task">Task to add</param>
+		/// <param name="task">Task to add</param>
 		/// <returns>True if the task could be added, false otherwise</returns>
-		bool Add(BgTaskImpl Task);
+		bool Add(BgTaskImpl task);
 
 		/// <summary>
-		/// Execute all the tasks added to this executor.
+		/// ExecuteAsync all the tasks added to this executor.
 		/// </summary>
-		/// <param name="Job">Information about the current job</param>
-		/// <param name="BuildProducts">Set of build products produced by this node.</param>
-		/// <param name="TagNameToFileSet">Mapping from tag names to the set of files they include</param>
+		/// <param name="job">Information about the current job</param>
+		/// <param name="buildProducts">Set of build products produced by this node.</param>
+		/// <param name="tagNameToFileSet">Mapping from tag names to the set of files they include</param>
 		/// <returns>Whether the task succeeded or not. Exiting with an exception will be caught and treated as a failure.</returns>
-		Task ExecuteAsync(JobContext Job, HashSet<FileReference> BuildProducts, Dictionary<string, HashSet<FileReference>> TagNameToFileSet);
+		Task ExecuteAsync(JobContext job, HashSet<FileReference> buildProducts, Dictionary<string, HashSet<FileReference>> tagNameToFileSet);
 	}
 
 	/// <summary>
@@ -114,13 +109,13 @@ namespace AutomationTool
 		public BgScriptLocation? SourceLocation { get; set; }
 
 		/// <summary>
-		/// Execute this node.
+		/// ExecuteAsync this node.
 		/// </summary>
-		/// <param name="Job">Information about the current job</param>
-		/// <param name="BuildProducts">Set of build products produced by this node.</param>
-		/// <param name="TagNameToFileSet">Mapping from tag names to the set of files they include</param>
+		/// <param name="job">Information about the current job</param>
+		/// <param name="buildProducts">Set of build products produced by this node.</param>
+		/// <param name="tagNameToFileSet">Mapping from tag names to the set of files they include</param>
 		/// <returns>Whether the task succeeded or not. Exiting with an exception will be caught and treated as a failure.</returns>
-		public abstract Task ExecuteAsync(JobContext Job, HashSet<FileReference> BuildProducts, Dictionary<string, HashSet<FileReference>> TagNameToFileSet);
+		public abstract Task ExecuteAsync(JobContext job, HashSet<FileReference> buildProducts, Dictionary<string, HashSet<FileReference>> tagNameToFileSet);
 
 		/// <summary>
 		/// Creates a proxy to execute this node.
@@ -134,47 +129,47 @@ namespace AutomationTool
 		/// <summary>
 		/// Output this task out to an XML writer.
 		/// </summary>
-		public abstract void Write(XmlWriter Writer);
+		public abstract void Write(XmlWriter writer);
 
 		/// <summary>
 		/// Writes this task to an XML writer, using the given parameters object.
 		/// </summary>
-		/// <param name="Writer">Writer for the XML schema</param>
-		/// <param name="Parameters">Parameters object that this task is constructed with</param>
-		protected void Write(XmlWriter Writer, object Parameters)
+		/// <param name="writer">Writer for the XML schema</param>
+		/// <param name="parameters">Parameters object that this task is constructed with</param>
+		protected void Write(XmlWriter writer, object parameters)
 		{
-			TaskElementAttribute Element = GetType().GetCustomAttribute<TaskElementAttribute>() ?? throw new InvalidOperationException();
-			Writer.WriteStartElement(Element.Name);
+			TaskElementAttribute element = GetType().GetCustomAttribute<TaskElementAttribute>() ?? throw new InvalidOperationException();
+			writer.WriteStartElement(element.Name);
 
-			foreach (FieldInfo Field in Parameters.GetType().GetFields())
+			foreach (FieldInfo field in parameters.GetType().GetFields())
 			{
-				if (Field.MemberType == MemberTypes.Field)
+				if (field.MemberType == MemberTypes.Field)
 				{
-					TaskParameterAttribute? ParameterAttribute = Field.GetCustomAttribute<TaskParameterAttribute>();
-					if (ParameterAttribute != null)
+					TaskParameterAttribute? parameterAttribute = field.GetCustomAttribute<TaskParameterAttribute>();
+					if (parameterAttribute != null)
 					{
-						object? Value = Field.GetValue(Parameters);
-						if (Value != null)
+						object? value = field.GetValue(parameters);
+						if (value != null)
 						{
-							Writer.WriteAttributeString(Field.Name, Value.ToString());
+							writer.WriteAttributeString(field.Name, value.ToString());
 						}
 					}
 				}
 			}
-			foreach (PropertyInfo Property in Parameters.GetType().GetProperties())
+			foreach (PropertyInfo property in parameters.GetType().GetProperties())
 			{
-				TaskParameterAttribute? ParameterAttribute = Property.GetCustomAttribute<TaskParameterAttribute>();
-				if (ParameterAttribute != null)
+				TaskParameterAttribute? parameterAttribute = property.GetCustomAttribute<TaskParameterAttribute>();
+				if (parameterAttribute != null)
 				{
-					object? Value = Property.GetValue(Parameters);
-					if (Value != null)
+					object? value = property.GetValue(parameters);
+					if (value != null)
 					{
-						Writer.WriteAttributeString(Property.Name, Value.ToString());
+						writer.WriteAttributeString(property.Name, value.ToString());
 					}
 				}
 			}
 
-			Writer.WriteEndElement();
+			writer.WriteEndElement();
 		}
 
 		/// <summary>
@@ -182,12 +177,12 @@ namespace AutomationTool
 		/// </summary>
 		public string GetTraceString()
 		{
-			StringBuilder Builder = new StringBuilder();
-			using (XmlWriter Writer = XmlWriter.Create(new StringWriter(Builder), new XmlWriterSettings() { OmitXmlDeclaration = true }))
+			StringBuilder builder = new StringBuilder();
+			using (XmlWriter writer = XmlWriter.Create(new StringWriter(builder), new XmlWriterSettings() { OmitXmlDeclaration = true }))
 			{
-				Write(Writer);
+				Write(writer);
 			}
-			return Builder.ToString();
+			return builder.ToString();
 		}
 
 		/// <summary>
@@ -196,35 +191,35 @@ namespace AutomationTool
 		/// <returns>The trace name</returns>
 		public virtual string GetTraceName()
 		{
-			TaskElementAttribute? TaskElement = GetType().GetCustomAttribute<TaskElementAttribute>();
-			return (TaskElement != null)? TaskElement.Name : "unknown";
+			TaskElementAttribute? taskElement = GetType().GetCustomAttribute<TaskElementAttribute>();
+			return (taskElement != null) ? taskElement.Name : "unknown";
 		}
 
 		/// <summary>
 		/// Get properties to include in tracing info
 		/// </summary>
-		/// <param name="Span">The scope to add properties to</param>
-		/// <param name="Prefix">Prefix for metadata entries</param>
-		public virtual void GetTraceMetadata(ITraceSpan Span, string Prefix)
+		/// <param name="span">The scope to add properties to</param>
+		/// <param name="prefix">Prefix for metadata entries</param>
+		public virtual void GetTraceMetadata(ITraceSpan span, string prefix)
 		{
 			if (SourceLocation != null)
 			{
-				Span.AddMetadata(Prefix + "source.file", SourceLocation.File.FullName);
-				Span.AddMetadata(Prefix + "source.line", SourceLocation.LineNumber.ToString());
+				span.AddMetadata(prefix + "source.file", SourceLocation.File.FullName);
+				span.AddMetadata(prefix + "source.line", SourceLocation.LineNumber.ToString());
 			}
 		}
-		
+
 		/// <summary>
 		/// Get properties to include in tracing info
 		/// </summary>
-		/// <param name="Span">The scope to add properties to</param>
-		/// <param name="Prefix">Prefix for metadata entries</param>
-		public virtual void GetTraceMetadata(ISpan Span, string Prefix)
+		/// <param name="span">The scope to add properties to</param>
+		/// <param name="prefix">Prefix for metadata entries</param>
+		public virtual void GetTraceMetadata(ISpan span, string prefix)
 		{
 			if (SourceLocation != null)
 			{
-				Span.SetTag(Prefix + "source.file", SourceLocation.File.FullName);
-				Span.SetTag(Prefix + "source.line", SourceLocation.LineNumber);
+				span.SetTag(prefix + "source.file", SourceLocation.File.FullName);
+				span.SetTag(prefix + "source.line", SourceLocation.LineNumber);
 			}
 		}
 
@@ -243,17 +238,17 @@ namespace AutomationTool
 		/// <summary>
 		/// Adds tag names from a filespec
 		/// </summary>
-		/// <param name="Filespec">A filespec, as can be passed to ResolveFilespec</param>
+		/// <param name="filespec">A filespec, as can be passed to ResolveFilespec</param>
 		/// <returns>Tag names from this filespec</returns>
-		protected static IEnumerable<string> FindTagNamesFromFilespec(string Filespec)
+		protected static IEnumerable<string> FindTagNamesFromFilespec(string filespec)
 		{
-			if(!String.IsNullOrEmpty(Filespec))
+			if (!String.IsNullOrEmpty(filespec))
 			{
-				foreach(string Pattern in SplitDelimitedList(Filespec))
+				foreach (string pattern in SplitDelimitedList(filespec))
 				{
-					if(Pattern.StartsWith("#", StringComparison.Ordinal))
+					if (pattern.StartsWith("#", StringComparison.Ordinal))
 					{
-						yield return Pattern;
+						yield return pattern;
 					}
 				}
 			}
@@ -262,15 +257,15 @@ namespace AutomationTool
 		/// <summary>
 		/// Enumerates tag names from a list
 		/// </summary>
-		/// <param name="TagList">List of tags separated by semicolons</param>
+		/// <param name="tagList">List of tags separated by semicolons</param>
 		/// <returns>Tag names from this filespec</returns>
-		protected static IEnumerable<string> FindTagNamesFromList(string? TagList)
+		protected static IEnumerable<string> FindTagNamesFromList(string? tagList)
 		{
-			if(!String.IsNullOrEmpty(TagList))
+			if (!String.IsNullOrEmpty(tagList))
 			{
-				foreach(string TagName in SplitDelimitedList(TagList))
+				foreach (string tagName in SplitDelimitedList(tagList))
 				{
-					yield return TagName;
+					yield return tagName;
 				}
 			}
 		}
@@ -278,77 +273,77 @@ namespace AutomationTool
 		/// <summary>
 		/// Resolves a single name to a file reference, resolving relative paths to the root of the current path.
 		/// </summary>
-		/// <param name="Name">Name of the file</param>
+		/// <param name="name">Name of the file</param>
 		/// <returns>Fully qualified file reference</returns>
-		public static FileReference ResolveFile(string Name)
+		public static FileReference ResolveFile(string name)
 		{
-			if(Path.IsPathRooted(Name))
+			if (Path.IsPathRooted(name))
 			{
-				return new FileReference(Name);
+				return new FileReference(name);
 			}
 			else
 			{
-				return new FileReference(Path.Combine(CommandUtils.CmdEnv.LocalRoot, Name));
+				return new FileReference(Path.Combine(CommandUtils.CmdEnv.LocalRoot, name));
 			}
 		}
 
 		/// <summary>
 		/// Resolves a directory reference from the given string. Assumes the root directory is the root of the current branch.
 		/// </summary>
-		/// <param name="Name">Name of the directory. May be null or empty.</param>
+		/// <param name="name">Name of the directory. May be null or empty.</param>
 		/// <returns>The resolved directory</returns>
-		public static DirectoryReference ResolveDirectory(string? Name)
+		public static DirectoryReference ResolveDirectory(string? name)
 		{
-			if(String.IsNullOrEmpty(Name))
+			if (String.IsNullOrEmpty(name))
 			{
 				return Unreal.RootDirectory;
 			}
-			else if(Path.IsPathRooted(Name))
+			else if (Path.IsPathRooted(name))
 			{
-				return new DirectoryReference(Name);
+				return new DirectoryReference(name);
 			}
 			else
 			{
-				return DirectoryReference.Combine(Unreal.RootDirectory, Name);
+				return DirectoryReference.Combine(Unreal.RootDirectory, name);
 			}
 		}
 
 		/// <summary>
 		/// Finds or adds a set containing files with the given tag
 		/// </summary>
-		/// <param name="TagNameToFileSet">Map of tag names to the set of files they contain</param>
-		/// <param name="TagName">The tag name to return a set for. A leading '#' character is required.</param>
+		/// <param name="tagNameToFileSet">Map of tag names to the set of files they contain</param>
+		/// <param name="tagName">The tag name to return a set for. A leading '#' character is required.</param>
 		/// <returns>Set of files</returns>
-		public static HashSet<FileReference> FindOrAddTagSet(Dictionary<string, HashSet<FileReference>> TagNameToFileSet, string TagName)
+		public static HashSet<FileReference> FindOrAddTagSet(Dictionary<string, HashSet<FileReference>> tagNameToFileSet, string tagName)
 		{
 			// Make sure the tag name contains a single leading hash
-			if (TagName.LastIndexOf('#') != 0)
+			if (tagName.LastIndexOf('#') != 0)
 			{
-				throw new AutomationException("Tag name '{0}' is not valid - should contain a single leading '#' character", TagName);
+				throw new AutomationException("Tag name '{0}' is not valid - should contain a single leading '#' character", tagName);
 			}
 
 			// Any spaces should be later than the second char - most likely to be a typo if directly after the # character
-			if (TagName.IndexOf(' ', StringComparison.Ordinal) == 1)
+			if (tagName.IndexOf(' ', StringComparison.Ordinal) == 1)
 			{
-				throw new AutomationException("Tag name '{0}' is not valid - spaces should only be used to separate words", TagName);
+				throw new AutomationException("Tag name '{0}' is not valid - spaces should only be used to separate words", tagName);
 			}
 
 			// Find the files which match this tag
-			HashSet<FileReference>? Files;
-			if(!TagNameToFileSet.TryGetValue(TagName, out Files))
+			HashSet<FileReference>? files;
+			if (!tagNameToFileSet.TryGetValue(tagName, out files))
 			{
-				Files = new HashSet<FileReference>();
-				TagNameToFileSet.Add(TagName, Files);
+				files = new HashSet<FileReference>();
+				tagNameToFileSet.Add(tagName, files);
 			}
 
 			// If we got a null reference, it's because the tag is not listed as an input for this node (see RunGraph.BuildSingleNode). Fill it in, but only with an error.
-			if(Files == null)
+			if (files == null)
 			{
-				Logger.LogError("Attempt to reference tag '{TagName}', which is not listed as a dependency of this node.", TagName);
-				Files = new HashSet<FileReference>();
-				TagNameToFileSet.Add(TagName, Files);
+				Logger.LogError("Attempt to reference tag '{TagName}', which is not listed as a dependency of this node.", tagName);
+				files = new HashSet<FileReference>();
+				tagNameToFileSet.Add(tagName, files);
 			}
-			return Files;
+			return files;
 		}
 
 		/// <summary>
@@ -359,99 +354,99 @@ namespace AutomationTool
 		///   d) A full directory wildcard (eg. Engine/...)
 		/// Note that wildcards may only match the last fragment in a pattern, so matches like "/*/Foo.txt" and "/.../Bar.txt" are illegal.
 		/// </summary>
-		/// <param name="DefaultDirectory">The default directory to resolve relative paths to</param>
-		/// <param name="DelimitedPatterns">List of files, tag names, or file specifications to include separated by semicolons.</param>
-		/// <param name="TagNameToFileSet">Mapping of tag name to fileset, as passed to the Execute() method</param>
+		/// <param name="defaultDirectory">The default directory to resolve relative paths to</param>
+		/// <param name="delimitedPatterns">List of files, tag names, or file specifications to include separated by semicolons.</param>
+		/// <param name="tagNameToFileSet">Mapping of tag name to fileset, as passed to the ExecuteAsync() method</param>
 		/// <returns>Set of matching files.</returns>
-		public static HashSet<FileReference> ResolveFilespec(DirectoryReference DefaultDirectory, string DelimitedPatterns, Dictionary<string, HashSet<FileReference>> TagNameToFileSet)
+		public static HashSet<FileReference> ResolveFilespec(DirectoryReference defaultDirectory, string delimitedPatterns, Dictionary<string, HashSet<FileReference>> tagNameToFileSet)
 		{
-			List<string> ExcludePatterns = new List<string>();
-			return ResolveFilespecWithExcludePatterns(DefaultDirectory, DelimitedPatterns, ExcludePatterns, TagNameToFileSet);
+			List<string> excludePatterns = new List<string>();
+			return ResolveFilespecWithExcludePatterns(defaultDirectory, delimitedPatterns, excludePatterns, tagNameToFileSet);
 		}
 
 		/// <summary>
 		/// Resolve a list of files, tag names or file specifications separated by semicolons as above, but preserves any directory references for further processing.
 		/// </summary>
-		/// <param name="DefaultDirectory">The default directory to resolve relative paths to</param>
-		/// <param name="DelimitedPatterns">List of files, tag names, or file specifications to include separated by semicolons.</param>
-		/// <param name="ExcludePatterns">Set of patterns to apply to directory searches. This can greatly speed up enumeration by earlying out of recursive directory searches if large directories are excluded (eg. .../Intermediate/...).</param>
-		/// <param name="TagNameToFileSet">Mapping of tag name to fileset, as passed to the Execute() method</param>
+		/// <param name="defaultDirectory">The default directory to resolve relative paths to</param>
+		/// <param name="delimitedPatterns">List of files, tag names, or file specifications to include separated by semicolons.</param>
+		/// <param name="excludePatterns">Set of patterns to apply to directory searches. This can greatly speed up enumeration by earlying out of recursive directory searches if large directories are excluded (eg. .../Intermediate/...).</param>
+		/// <param name="tagNameToFileSet">Mapping of tag name to fileset, as passed to the ExecuteAsync() method</param>
 		/// <returns>Set of matching files.</returns>
-		public static HashSet<FileReference> ResolveFilespecWithExcludePatterns(DirectoryReference DefaultDirectory, string DelimitedPatterns, List<string> ExcludePatterns, Dictionary<string, HashSet<FileReference>> TagNameToFileSet)
+		public static HashSet<FileReference> ResolveFilespecWithExcludePatterns(DirectoryReference defaultDirectory, string delimitedPatterns, List<string> excludePatterns, Dictionary<string, HashSet<FileReference>> tagNameToFileSet)
 		{
 			// Split the argument into a list of patterns
-			List<string> Patterns = SplitDelimitedList(DelimitedPatterns);
-			return ResolveFilespecWithExcludePatterns(DefaultDirectory, Patterns, ExcludePatterns, TagNameToFileSet);
+			List<string> patterns = SplitDelimitedList(delimitedPatterns);
+			return ResolveFilespecWithExcludePatterns(defaultDirectory, patterns, excludePatterns, tagNameToFileSet);
 		}
 
 		/// <summary>
 		/// Resolve a list of files, tag names or file specifications as above, but preserves any directory references for further processing.
 		/// </summary>
-		/// <param name="DefaultDirectory">The default directory to resolve relative paths to</param>
-		/// <param name="FilePatterns">List of files, tag names, or file specifications to include separated by semicolons.</param>
-		/// <param name="ExcludePatterns">Set of patterns to apply to directory searches. This can greatly speed up enumeration by earlying out of recursive directory searches if large directories are excluded (eg. .../Intermediate/...).</param>
-		/// <param name="TagNameToFileSet">Mapping of tag name to fileset, as passed to the Execute() method</param>
+		/// <param name="defaultDirectory">The default directory to resolve relative paths to</param>
+		/// <param name="filePatterns">List of files, tag names, or file specifications to include separated by semicolons.</param>
+		/// <param name="excludePatterns">Set of patterns to apply to directory searches. This can greatly speed up enumeration by earlying out of recursive directory searches if large directories are excluded (eg. .../Intermediate/...).</param>
+		/// <param name="tagNameToFileSet">Mapping of tag name to fileset, as passed to the ExecuteAsync() method</param>
 		/// <returns>Set of matching files.</returns>
-		public static HashSet<FileReference> ResolveFilespecWithExcludePatterns(DirectoryReference DefaultDirectory, List<string> FilePatterns, List<string> ExcludePatterns, Dictionary<string, HashSet<FileReference>> TagNameToFileSet)
+		public static HashSet<FileReference> ResolveFilespecWithExcludePatterns(DirectoryReference defaultDirectory, List<string> filePatterns, List<string> excludePatterns, Dictionary<string, HashSet<FileReference>> tagNameToFileSet)
 		{
 			// Parse each of the patterns, and add the results into the given sets
-			HashSet<FileReference> Files = new HashSet<FileReference>();
-			foreach(string Pattern in FilePatterns)
+			HashSet<FileReference> files = new HashSet<FileReference>();
+			foreach (string pattern in filePatterns)
 			{
 				// Check if it's a tag name
-				if(Pattern.StartsWith("#", StringComparison.Ordinal))
+				if (pattern.StartsWith("#", StringComparison.Ordinal))
 				{
-					Files.UnionWith(FindOrAddTagSet(TagNameToFileSet, Pattern));
+					files.UnionWith(FindOrAddTagSet(tagNameToFileSet, pattern));
 					continue;
 				}
 
 				// If it doesn't contain any wildcards, just add the pattern directly
-				int WildcardIdx = FileFilter.FindWildcardIndex(Pattern);
-				if(WildcardIdx == -1)
+				int wildcardIdx = FileFilter.FindWildcardIndex(pattern);
+				if (wildcardIdx == -1)
 				{
-					Files.Add(FileReference.Combine(DefaultDirectory, Pattern));
+					files.Add(FileReference.Combine(defaultDirectory, pattern));
 					continue;
 				}
 
 				// Find the base directory for the search. We construct this in a very deliberate way including the directory separator itself, so matches
 				// against the OS root directory will resolve correctly both on Mac (where / is the filesystem root) and Windows (where / refers to the current drive).
-				int LastDirectoryIdx = Pattern.LastIndexOfAny(new char[]{ Path.DirectorySeparatorChar, Path.AltDirectorySeparatorChar }, WildcardIdx);
-				DirectoryReference BaseDir = DirectoryReference.Combine(DefaultDirectory, Pattern.Substring(0, LastDirectoryIdx + 1));
+				int lastDirectoryIdx = pattern.LastIndexOfAny(new char[] { Path.DirectorySeparatorChar, Path.AltDirectorySeparatorChar }, wildcardIdx);
+				DirectoryReference baseDir = DirectoryReference.Combine(defaultDirectory, pattern.Substring(0, lastDirectoryIdx + 1));
 
 				// Construct the absolute include pattern to match against, re-inserting the resolved base directory to construct a canonical path.
-				string IncludePattern = BaseDir.FullName.TrimEnd(new char[]{ Path.DirectorySeparatorChar, Path.AltDirectorySeparatorChar }) + "/" + Pattern.Substring(LastDirectoryIdx + 1);
+				string includePattern = baseDir.FullName.TrimEnd(new char[] { Path.DirectorySeparatorChar, Path.AltDirectorySeparatorChar }) + "/" + pattern.Substring(lastDirectoryIdx + 1);
 
 				// Construct a filter and apply it to the directory
-				if(DirectoryReference.Exists(BaseDir))
+				if (DirectoryReference.Exists(baseDir))
 				{
-					FileFilter Filter = new FileFilter();
-					Filter.AddRule(IncludePattern, FileFilterType.Include);
-					if(ExcludePatterns != null && ExcludePatterns.Count > 0)
+					FileFilter filter = new FileFilter();
+					filter.AddRule(includePattern, FileFilterType.Include);
+					if (excludePatterns != null && excludePatterns.Count > 0)
 					{
-						Filter.AddRules(ExcludePatterns, FileFilterType.Exclude);
+						filter.AddRules(excludePatterns, FileFilterType.Exclude);
 					}
-					Files.UnionWith(Filter.ApplyToDirectory(BaseDir, BaseDir.FullName, true));
+					files.UnionWith(filter.ApplyToDirectory(baseDir, baseDir.FullName, true));
 				}
 			}
 
 			// If we have exclude rules, create and run a filter against all the output files to catch things that weren't added from an include
-			if(ExcludePatterns != null && ExcludePatterns.Count > 0)
+			if (excludePatterns != null && excludePatterns.Count > 0)
 			{
-				FileFilter Filter = new FileFilter(FileFilterType.Include);
-				Filter.AddRules(ExcludePatterns, FileFilterType.Exclude);
-				Files.RemoveWhere(x => !Filter.Matches(x.FullName));
+				FileFilter filter = new FileFilter(FileFilterType.Include);
+				filter.AddRules(excludePatterns, FileFilterType.Exclude);
+				files.RemoveWhere(x => !filter.Matches(x.FullName));
 			}
-			return Files;
+			return files;
 		}
 
 		/// <summary>
 		/// Splits a string separated by semicolons into a list, removing empty entries
 		/// </summary>
-		/// <param name="Text">The input string</param>
+		/// <param name="text">The input string</param>
 		/// <returns>Array of the parsed items</returns>
-		public static List<string> SplitDelimitedList(string Text)
+		public static List<string> SplitDelimitedList(string text)
 		{
-			return Text.Split(';').Select(x => x.Trim()).Where(x => x.Length > 0).ToList();
+			return text.Split(';').Select(x => x.Trim()).Where(x => x.Length > 0).ToList();
 		}
 
 		/// <summary>
@@ -467,15 +462,15 @@ namespace AutomationTool
 		/// <summary>
 		/// Add cleanup commands to run after the step completes
 		/// </summary>
-		/// <param name="NewLines">Lines to add to the cleanup script</param>
-		/// <param name="Lease">Whether to add the commands to run on lease termination</param>
-		public static async Task AddCleanupCommandsAsync(IEnumerable<string> NewLines, bool Lease = false)
+		/// <param name="newLines">Lines to add to the cleanup script</param>
+		/// <param name="lease">Whether to add the commands to run on lease termination</param>
+		public static async Task AddCleanupCommandsAsync(IEnumerable<string> newLines, bool lease = false)
 		{
-			string? CleanupScriptEnvVar = Environment.GetEnvironmentVariable(Lease? LeaseCleanupScriptEnvVarName : CleanupScriptEnvVarName);
-			if (!String.IsNullOrEmpty(CleanupScriptEnvVar))
+			string? cleanupScriptEnvVar = Environment.GetEnvironmentVariable(lease ? LeaseCleanupScriptEnvVarName : CleanupScriptEnvVarName);
+			if (!String.IsNullOrEmpty(cleanupScriptEnvVar))
 			{
-				FileReference CleanupScript = new FileReference(CleanupScriptEnvVar);
-				await FileReference.AppendAllLinesAsync(CleanupScript, NewLines);
+				FileReference cleanupScript = new FileReference(cleanupScriptEnvVar);
+				await FileReference.AppendAllLinesAsync(cleanupScript, newLines);
 			}
 		}
 
@@ -487,8 +482,8 @@ namespace AutomationTool
 		/// <summary>
 		/// Updates the graph currently used by Horde
 		/// </summary>
-		/// <param name="Job">Context for the current job that is being executed</param>
-		public static void UpdateGraphForHorde(JobContext Job)
+		/// <param name="job">Context for the current job that is being executed</param>
+		public static void UpdateGraphForHorde(JobContext job)
 		{
 			string? exportGraphFile = Environment.GetEnvironmentVariable(GraphUpdateEnvVarName);
 			if (String.IsNullOrEmpty(exportGraphFile))
@@ -498,7 +493,7 @@ namespace AutomationTool
 
 			List<string> newParams = new List<string>();
 			newParams.Add("BuildGraph");
-			newParams.AddRange(Job.OwnerCommand.Params.Select(x => $"-{x}"));
+			newParams.AddRange(job.OwnerCommand.Params.Select(x => $"-{x}"));
 			newParams.RemoveAll(x => x.StartsWith("-SingleNode=", StringComparison.OrdinalIgnoreCase));
 			newParams.Add($"-HordeExport={exportGraphFile}");
 			newParams.Add($"-ListOnly");
@@ -514,19 +509,19 @@ namespace AutomationTool
 	public abstract class CustomTask : BgTaskImpl
 	{
 		/// <inheritdoc/>
-		public sealed override Task ExecuteAsync(JobContext Job, HashSet<FileReference> BuildProducts, Dictionary<string, HashSet<FileReference>> TagNameToFileSet)
+		public sealed override Task ExecuteAsync(JobContext job, HashSet<FileReference> buildProducts, Dictionary<string, HashSet<FileReference>> tagNameToFileSet)
 		{
-			Execute(Job, BuildProducts, TagNameToFileSet);
+			Execute(job, buildProducts, tagNameToFileSet);
 			return Task.CompletedTask;
 		}
 
 		/// <summary>
-		/// Execute this node.
+		/// ExecuteAsync this node.
 		/// </summary>
-		/// <param name="Job">Information about the current job</param>
-		/// <param name="BuildProducts">Set of build products produced by this node.</param>
-		/// <param name="TagNameToFileSet">Mapping from tag names to the set of files they include</param>
+		/// <param name="job">Information about the current job</param>
+		/// <param name="buildProducts">Set of build products produced by this node.</param>
+		/// <param name="tagNameToFileSet">Mapping from tag names to the set of files they include</param>
 		/// <returns>Whether the task succeeded or not. Exiting with an exception will be caught and treated as a failure.</returns>
-		public abstract void Execute(JobContext Job, HashSet<FileReference> BuildProducts, Dictionary<string, HashSet<FileReference>> TagNameToFileSet);
+		public abstract void Execute(JobContext job, HashSet<FileReference> buildProducts, Dictionary<string, HashSet<FileReference>> tagNameToFileSet);
 	}
 }
diff --git a/Engine/Source/Programs/AutomationTool/BuildGraph/Tasks/AgeStoreTask.cs b/Engine/Source/Programs/AutomationTool/BuildGraph/Tasks/AgeStoreTask.cs
--- a/Engine/Source/Programs/AutomationTool/BuildGraph/Tasks/AgeStoreTask.cs
+++ b/Engine/Source/Programs/AutomationTool/BuildGraph/Tasks/AgeStoreTask.cs
@@ -1,17 +1,14 @@
 // Copyright Epic Games, Inc. All Rights Reserved.
 
-using AutomationTool;
 using System;
 using System.Collections.Generic;
 using System.IO;
 using System.Linq;
+using System.Threading.Tasks;
 using System.Xml;
 using EpicGames.Core;
-using UnrealBuildTool;
-using System.Threading.Tasks;
 using Microsoft.Extensions.Logging;
-
-using static AutomationTool.CommandUtils;
+using UnrealBuildTool;
 
 namespace AutomationTool.Tasks
 {
@@ -63,69 +60,69 @@ namespace AutomationTool.Tasks
 		/// <summary>
 		/// Parameters for this task
 		/// </summary>
-		AgeStoreTaskParameters Parameters;
+		readonly AgeStoreTaskParameters _parameters;
 
 		/// <summary>
 		/// Construct a spawn task
 		/// </summary>
-		/// <param name="InParameters">Parameters for the task</param>
-		public AgeStoreTask(AgeStoreTaskParameters InParameters)
+		/// <param name="parameters">Parameters for the task</param>
+		public AgeStoreTask(AgeStoreTaskParameters parameters)
 		{
-			Parameters = InParameters;
+			_parameters = parameters;
 		}
 
-		private static void TryDelete(DirectoryInfo Directory)
+		private static void TryDelete(DirectoryInfo directory)
 		{
 			try
 			{
-				Directory.Delete(true);
-                Logger.LogInformation("Removed '{Arg0}'", Directory.FullName);
+				directory.Delete(true);
+				Logger.LogInformation("Removed '{Arg0}'", directory.FullName);
 			}
 			catch
 			{
-				Logger.LogWarning("Couldn't delete '{Arg0}' - skipping", Directory.FullName);
+				Logger.LogWarning("Couldn't delete '{Arg0}' - skipping", directory.FullName);
 			}
 		}
 
-		private static void TryDelete(FileInfo File)
+		private static void TryDelete(FileInfo file)
 		{
 			try
 			{
-				File.Delete();
-				Logger.LogInformation("Removed '{Arg0}'", File.FullName);
+				file.Delete();
+				Logger.LogInformation("Removed '{Arg0}'", file.FullName);
 			}
 			catch
 			{
-				Logger.LogWarning("Couldn't delete '{Arg0}' - skipping", File.FullName);
+				Logger.LogWarning("Couldn't delete '{Arg0}' - skipping", file.FullName);
 			}
 		}
 
 		// Checks if an existing build has a version file, returns false to NOT delete if it exists
-		private static bool CheckCanDeleteFromVersionFile(HashSet<string> ExistingBuilds, DirectoryInfo Directory, FileInfo IndividualFile = null)
+		private static bool CheckCanDeleteFromVersionFile(HashSet<string> existingBuilds, DirectoryInfo directory, FileInfo individualFile = null)
 		{
 			// check for any existing version files
-			foreach (FileInfo BuildVersionFile in Directory.EnumerateFiles("*.version"))
+			foreach (FileInfo buildVersionFile in directory.EnumerateFiles("*.version"))
 			{
 				// If the buildversion matches one of the directories in build share provided, don't delete no matter the age.
-				string BuildVersion = Path.GetFileNameWithoutExtension(BuildVersionFile.Name);
-				if (ExistingBuilds.Contains(BuildVersion))
+				string buildVersion = Path.GetFileNameWithoutExtension(buildVersionFile.Name);
+				if (existingBuilds.Contains(buildVersion))
 				{
 					// if checking for an individual file, see if the filename matches what's in the .version file.
 					// these file names won't have extensions.
-					if (IndividualFile != null)
+					if (individualFile != null)
 					{
-						string IndividualFilePath = IndividualFile.FullName;
-						string FilePointerName = File.ReadAllText(BuildVersionFile.FullName).Trim();
-						if(FilePointerName == Path.GetFileNameWithoutExtension(IndividualFilePath))
+						string individualFilePath = individualFile.FullName;
+						string filePointerName = File.ReadAllText(buildVersionFile.FullName).Trim();
+						if (filePointerName == Path.GetFileNameWithoutExtension(individualFilePath))
 						{
-							Logger.LogInformation("Found existing build {BuildVersion} in the BuildDir with matching individual file {IndividualFilePath} - skipping.", BuildVersion, IndividualFilePath);
+							Logger.LogInformation("Found existing build {BuildVersion} in the BuildDir with matching individual file {IndividualFilePath} - skipping.", buildVersion, individualFilePath);
 							return false;
 						}
 					}
 					// otherwise it's okay to just mark the entire folder for delete
 					else
 					{
-						Logger.LogInformation("Found existing build {BuildVersion} in the BuildDir - skipping.", BuildVersion);
+						Logger.LogInformation("Found existing build {BuildVersion} in the BuildDir - skipping.", buildVersion);
 						return false;
 					}
 				}
@@ -133,96 +130,96 @@ namespace AutomationTool.Tasks
 			return true;
 		}
 
-		private static void RecurseDirectory(DateTime ExpireTimeUtc, DirectoryInfo CurrentDirectory, string[] DirectoryStructure, int Level, string Filter, HashSet<string> ExistingBuilds, bool bDeleteIndividualFiles)
+		private static void RecurseDirectory(DateTime expireTimeUtc, DirectoryInfo currentDirectory, string[] directoryStructure, int level, string filter, HashSet<string> existingBuilds, bool deleteIndividualFiles)
 		{
 			// Do a file search at the last level.
-			if (Level == DirectoryStructure.Length)
+			if (level == directoryStructure.Length)
 			{
-				if (bDeleteIndividualFiles)
+				if (deleteIndividualFiles)
 				{
 					// Delete any file in the directory that is out of date.
-					foreach (FileInfo OutdatedFile in CurrentDirectory.EnumerateFiles().Where(x => x.LastWriteTimeUtc < ExpireTimeUtc && x.Extension != ".version"))
+					foreach (FileInfo outdatedFile in currentDirectory.EnumerateFiles().Where(x => x.LastWriteTimeUtc < expireTimeUtc && x.Extension != ".version"))
 					{
 						// check to make sure this file is valid to delete
-						if (CheckCanDeleteFromVersionFile(ExistingBuilds, CurrentDirectory, OutdatedFile))
+						if (CheckCanDeleteFromVersionFile(existingBuilds, currentDirectory, outdatedFile))
 						{
-							TryDelete(OutdatedFile);
+							TryDelete(outdatedFile);
 						}
 					}
 				}
 				// If all files are out of date, delete the directory...
-				else if (CurrentDirectory.EnumerateFiles().Where(x => x.Extension != ".version").All(x => x.LastWriteTimeUtc < ExpireTimeUtc) && CheckCanDeleteFromVersionFile(ExistingBuilds, CurrentDirectory))
+				else if (currentDirectory.EnumerateFiles().Where(x => x.Extension != ".version").All(x => x.LastWriteTimeUtc < expireTimeUtc) && CheckCanDeleteFromVersionFile(existingBuilds, currentDirectory))
 				{
-					TryDelete(CurrentDirectory);
+					TryDelete(currentDirectory);
 				}
 			}
 			else
 			{
-				string[] Patterns = DirectoryStructure[Level].Split(';');
-				foreach (var Pattern in Patterns)
+				string[] patterns = directoryStructure[level].Split(';');
+				foreach (var pattern in patterns)
 				{
-					string ReplacedPattern = string.Format(Pattern, Filter);
+					string replacedPattern = string.Format(pattern, filter);
 
-					foreach (var ChildDirectory in CurrentDirectory.GetDirectories(ReplacedPattern, SearchOption.TopDirectoryOnly))
+					foreach (var childDirectory in currentDirectory.GetDirectories(replacedPattern, SearchOption.TopDirectoryOnly))
 					{
-						RecurseDirectory(ExpireTimeUtc, ChildDirectory, DirectoryStructure, Level + 1, Filter, ExistingBuilds, bDeleteIndividualFiles);
+						RecurseDirectory(expireTimeUtc, childDirectory, directoryStructure, level + 1, filter, existingBuilds, deleteIndividualFiles);
 					}
 				}
 
 				// Delete this directory if it is empty, and it is not the root directory.
-				if (Level > 0 && !CurrentDirectory.EnumerateFileSystemInfos().Any())
-					TryDelete(CurrentDirectory);
+				if (level > 0 && !currentDirectory.EnumerateFileSystemInfos().Any())
+					TryDelete(currentDirectory);
 			}
 		}
 
 		/// <summary>
-		/// Execute the task.
+		/// ExecuteAsync the task.
 		/// </summary>
-		/// <param name="Job">Information about the current job</param>
-		/// <param name="BuildProducts">Set of build products produced by this node.</param>
-		/// <param name="TagNameToFileSet">Mapping from tag names to the set of files they include</param>
-		public override Task ExecuteAsync(JobContext Job, HashSet<FileReference> BuildProducts, Dictionary<string, HashSet<FileReference>> TagNameToFileSet)
+		/// <param name="job">Information about the current job</param>
+		/// <param name="buildProducts">Set of build products produced by this node.</param>
+		/// <param name="tagNameToFileSet">Mapping from tag names to the set of files they include</param>
+		public override Task ExecuteAsync(JobContext job, HashSet<FileReference> buildProducts, Dictionary<string, HashSet<FileReference>> tagNameToFileSet)
 		{
 			// Get the list of symbol file name patterns from the platform.
-			Platform TargetPlatform = Platform.GetPlatform(Parameters.Platform);
-			string[] DirectoryStructure = TargetPlatform.SymbolServerDirectoryStructure;
-			if (DirectoryStructure == null)
+			Platform targetPlatform = Platform.GetPlatform(_parameters.Platform);
+			string[] directoryStructure = targetPlatform.SymbolServerDirectoryStructure;
+			if (directoryStructure == null)
 			{
 				throw new AutomationException("Platform does not specify the symbol server structure. Cannot age the symbol server.");
 			}
 
-			string Filter = string.IsNullOrWhiteSpace(Parameters.Filter)
+			string filter = string.IsNullOrWhiteSpace(_parameters.Filter)
 				? string.Empty
-				: Parameters.Filter.Trim();
+				: _parameters.Filter.Trim();
 
 			// Eumerate the root directory of builds for buildversions to check against
 			// Folder names in the root directory should match the name of the .version files
-			HashSet<string> ExistingBuilds = new HashSet<string>(StringComparer.OrdinalIgnoreCase);
-			if(!string.IsNullOrWhiteSpace(Parameters.BuildDir))
+			HashSet<string> existingBuilds = new HashSet<string>(StringComparer.OrdinalIgnoreCase);
+			if (!string.IsNullOrWhiteSpace(_parameters.BuildDir))
 			{
-				DirectoryReference BuildDir = new DirectoryReference(Parameters.BuildDir);
-				if(DirectoryReference.Exists(BuildDir))
+				DirectoryReference buildDir = new DirectoryReference(_parameters.BuildDir);
+				if (DirectoryReference.Exists(buildDir))
 				{
-					foreach (string BuildName in DirectoryReference.EnumerateDirectories(BuildDir).Select(Build => Build.GetDirectoryName()))
+					foreach (string buildName in DirectoryReference.EnumerateDirectories(buildDir).Select(build => build.GetDirectoryName()))
 					{
-						ExistingBuilds.Add(BuildName);
+						existingBuilds.Add(buildName);
 					}
 				}
 				else
 				{
-					Logger.LogWarning("BuildDir of {Arg0} was provided but it doesn't exist! Will not check buildversions against it.", Parameters.BuildDir);
+					Logger.LogWarning("BuildDir of {Arg0} was provided but it doesn't exist! Will not check buildversions against it.", _parameters.BuildDir);
 				}
 			}
 
 			// Get the time at which to expire files
-			DateTime ExpireTimeUtc = DateTime.UtcNow - TimeSpan.FromDays(Parameters.Days);
-            Logger.LogInformation("Expiring all files before {ExpireTimeUtc}...", ExpireTimeUtc);
+			DateTime expireTimeUtc = DateTime.UtcNow - TimeSpan.FromDays(_parameters.Days);
+			Logger.LogInformation("Expiring all files before {ExpireTimeUtc}...", expireTimeUtc);
 
 			// Scan the store directory and delete old symbol files
-			DirectoryReference SymbolServerDirectory = ResolveDirectory(Parameters.StoreDir);
-			CommandUtils.OptionallyTakeLock(TargetPlatform.SymbolServerRequiresLock, SymbolServerDirectory, TimeSpan.FromMinutes(15), () =>
+			DirectoryReference symbolServerDirectory = ResolveDirectory(_parameters.StoreDir);
+			CommandUtils.OptionallyTakeLock(targetPlatform.SymbolServerRequiresLock, symbolServerDirectory, TimeSpan.FromMinutes(15), () =>
 			{
-				RecurseDirectory(ExpireTimeUtc, new DirectoryInfo(SymbolServerDirectory.FullName), DirectoryStructure, 0, Filter, ExistingBuilds, TargetPlatform.SymbolServerDeleteIndividualFiles);
+				RecurseDirectory(expireTimeUtc, new DirectoryInfo(symbolServerDirectory.FullName), directoryStructure, 0, filter, existingBuilds, targetPlatform.SymbolServerDeleteIndividualFiles);
 			});
 			return Task.CompletedTask;
 		}
@@ -230,9 +227,9 @@ namespace AutomationTool.Tasks
 		/// <summary>
 		/// Output this task out to an XML writer.
 		/// </summary>
-		public override void Write(XmlWriter Writer)
+		public override void Write(XmlWriter writer)
 		{
-			Write(Writer, Parameters);
+			Write(writer, _parameters);
 		}
 
 		/// <summary>
diff --git a/Engine/Source/Programs/AutomationTool/BuildGraph/Tasks/AwsAssumeRoleTask.cs b/Engine/Source/Programs/AutomationTool/BuildGraph/Tasks/AwsAssumeRoleTask.cs
--- a/Engine/Source/Programs/AutomationTool/BuildGraph/Tasks/AwsAssumeRoleTask.cs
+++ b/Engine/Source/Programs/AutomationTool/BuildGraph/Tasks/AwsAssumeRoleTask.cs
@@ -1,15 +1,12 @@
 // Copyright Epic Games, Inc. All Rights Reserved.
 
-using AutomationTool;
-using EpicGames.Core;
-using System;
 using System.Collections.Generic;
-using System.IO;
 using System.Linq;
 using System.Text;
 using System.Text.Json;
 using System.Threading.Tasks;
 using System.Xml;
+using EpicGames.Core;
 
 namespace AutomationTool.Tasks
 {
@@ -73,73 +70,70 @@ namespace AutomationTool.Tasks
 			public string SessionToken { get; set; }
 		}
 
-		/// <summary>
-		/// Parameters for this task
-		/// </summary>
-		AwsAssumeRoleTaskParameters Parameters;
+		readonly AwsAssumeRoleTaskParameters _parameters;
 
 		/// <summary>
 		/// Construct an AWS CLI task
 		/// </summary>
-		/// <param name="InParameters">Parameters for the task</param>
-		public AwsAssumeRoleTask(AwsAssumeRoleTaskParameters InParameters)
+		/// <param name="parameters">Parameters for the task</param>
+		public AwsAssumeRoleTask(AwsAssumeRoleTaskParameters parameters)
 		{
-			Parameters = InParameters;
+			_parameters = parameters;
 		}
 
 		/// <summary>
-		/// Execute the task.
+		/// ExecuteAsync the task.
 		/// </summary>
-		/// <param name="Job">Information about the current job</param>
-		/// <param name="BuildProducts">Set of build products produced by this node.</param>
-		/// <param name="TagNameToFileSet">Mapping from tag names to the set of files they include</param>
-		public override async Task ExecuteAsync(JobContext Job, HashSet<FileReference> BuildProducts, Dictionary<string, HashSet<FileReference>> TagNameToFileSet)
+		/// <param name="job">Information about the current job</param>
+		/// <param name="buildProducts">Set of build products produced by this node.</param>
+		/// <param name="tagNameToFileSet">Mapping from tag names to the set of files they include</param>
+		public override async Task ExecuteAsync(JobContext job, HashSet<FileReference> buildProducts, Dictionary<string, HashSet<FileReference>> tagNameToFileSet)
 		{
-			StringBuilder Arguments = new StringBuilder("sts assume-role");
-			if(Parameters.Arn != null)
+			StringBuilder arguments = new StringBuilder("sts assume-role");
+			if (_parameters.Arn != null)
 			{
-				Arguments.Append($" --role-arn {Parameters.Arn}");
+				arguments.Append($" --role-arn {_parameters.Arn}");
 			}
-			if (Parameters.Session != null)
+			if (_parameters.Session != null)
 			{
-				Arguments.Append($" --role-session-name {Parameters.Session}");
+				arguments.Append($" --role-session-name {_parameters.Session}");
 			}
-			Arguments.Append($" --duration-seconds {Parameters.Duration}");
+			arguments.Append($" --duration-seconds {_parameters.Duration}");
 
-			Dictionary<string, string> Environment = SpawnTaskBase.ParseEnvVars(Parameters.Environment, Parameters.EnvironmentFile);
-			IProcessResult Result = await SpawnTaskBase.ExecuteAsync("aws", Arguments.ToString(), EnvVars: Environment, LogOutput: false);
+			Dictionary<string, string> environment = SpawnTaskBase.ParseEnvVars(_parameters.Environment, _parameters.EnvironmentFile);
+			IProcessResult result = await SpawnTaskBase.ExecuteAsync("aws", arguments.ToString(), envVars: environment, logOutput: false);
 
-			JsonSerializerOptions Options = new JsonSerializerOptions();
-			Options.PropertyNameCaseInsensitive = true;
+			JsonSerializerOptions options = new JsonSerializerOptions();
+			options.PropertyNameCaseInsensitive = true;
 
-			AwsSettings Settings = JsonSerializer.Deserialize<AwsSettings>(Result.Output, Options);
-			if (Settings.Credentials != null)
+			AwsSettings settings = JsonSerializer.Deserialize<AwsSettings>(result.Output, options);
+			if (settings.Credentials != null)
 			{
-				if (Settings.Credentials.AccessKeyId != null)
+				if (settings.Credentials.AccessKeyId != null)
 				{
-					Environment["AWS_ACCESS_KEY_ID"] = Settings.Credentials.AccessKeyId;
+					environment["AWS_ACCESS_KEY_ID"] = settings.Credentials.AccessKeyId;
 				}
-				if (Settings.Credentials.SecretAccessKey != null)
+				if (settings.Credentials.SecretAccessKey != null)
 				{
-					Environment["AWS_SECRET_ACCESS_KEY"] = Settings.Credentials.SecretAccessKey;
+					environment["AWS_SECRET_ACCESS_KEY"] = settings.Credentials.SecretAccessKey;
 				}
-				if (Settings.Credentials.SessionToken != null)
+				if (settings.Credentials.SessionToken != null)
 				{
-					Environment["AWS_SESSION_TOKEN"] = Settings.Credentials.SessionToken;
+					environment["AWS_SESSION_TOKEN"] = settings.Credentials.SessionToken;
 				}
 			}
 
-			FileReference OutputFile = ResolveFile(Parameters.OutputFile);
-			DirectoryReference.CreateDirectory(OutputFile.Directory);
-			await FileReference.WriteAllLinesAsync(OutputFile, Environment.OrderBy(x => x.Key).Select(x => $"{x.Key}={x.Value}"));
+			FileReference outputFile = ResolveFile(_parameters.OutputFile);
+			DirectoryReference.CreateDirectory(outputFile.Directory);
+			await FileReference.WriteAllLinesAsync(outputFile, environment.OrderBy(x => x.Key).Select(x => $"{x.Key}={x.Value}"));
 		}
 
 		/// <summary>
 		/// Output this task out to an XML writer.
 		/// </summary>
-		public override void Write(XmlWriter Writer)
+		public override void Write(XmlWriter writer)
 		{
-			Write(Writer, Parameters);
+			Write(writer, _parameters);
 		}
 
 		/// <summary>
