EpicGames.BuildGraph: Apply standard naming conventions to BuildGraph.Automation.

diff --git a/Engine/Source/Programs/AutomationTool/BuildGraph/Tasks/AwsEcsDeployTask.cs b/Engine/Source/Programs/AutomationTool/BuildGraph/Tasks/AwsEcsDeployTask.cs
--- a/Engine/Source/Programs/AutomationTool/BuildGraph/Tasks/AwsEcsDeployTask.cs
+++ b/Engine/Source/Programs/AutomationTool/BuildGraph/Tasks/AwsEcsDeployTask.cs
@@ -1,7 +1,5 @@
 // Copyright Epic Games, Inc. All Rights Reserved.
 
-using AutomationTool;
-using EpicGames.Core;
 using System;
 using System.Collections.Generic;
 using System.IO;
@@ -9,8 +7,9 @@ using System.Text;
 using System.Text.Json;
 using System.Threading.Tasks;
 using System.Xml;
-using UnrealBuildBase;
+using EpicGames.Core;
 using Microsoft.Extensions.Logging;
+using UnrealBuildBase;
 
 namespace AutomationTool.Tasks
 {
@@ -77,54 +76,54 @@ namespace AutomationTool.Tasks
 		/// <summary>
 		/// Parameters for this task
 		/// </summary>
-		AwsEcsDeployTaskParameters Parameters;
+		readonly AwsEcsDeployTaskParameters _parameters;
 
 		/// <summary>
 		/// Construct an AWS ECS deploy task
 		/// </summary>
-		/// <param name="InParameters">Parameters for the task</param>
-		public AwsEcsDeployTask(AwsEcsDeployTaskParameters InParameters)
+		/// <param name="parameters">Parameters for the task</param>
+		public AwsEcsDeployTask(AwsEcsDeployTaskParameters parameters)
 		{
-			Parameters = InParameters;
+			_parameters = parameters;
 		}
 
 		/// <summary>
-		/// Execute the task.
+		/// ExecuteAsync the task.
 		/// </summary>
-		/// <param name="Job">Information about the current job</param>
-		/// <param name="BuildProducts">Set of build products produced by this node.</param>
-		/// <param name="TagNameToFileSet">Mapping from tag names to the set of files they include</param>
-		public override async Task ExecuteAsync(JobContext Job, HashSet<FileReference> BuildProducts, Dictionary<string, HashSet<FileReference>> TagNameToFileSet)
+		/// <param name="job">Information about the current job</param>
+		/// <param name="buildProducts">Set of build products produced by this node.</param>
+		/// <param name="tagNameToFileSet">Mapping from tag names to the set of files they include</param>
+		public override async Task ExecuteAsync(JobContext job, HashSet<FileReference> buildProducts, Dictionary<string, HashSet<FileReference>> tagNameToFileSet)
 		{
-			string TaskDefTemplate = await File.ReadAllTextAsync(ResolveFile(Parameters.TaskDefinitionFile).FullName);
-			string TaskDefRendered = TaskDefTemplate.Replace("%%DOCKER_IMAGE%%", Parameters.DockerImage, StringComparison.Ordinal);
-			if (Parameters.Version != null)
+			string taskDefTemplate = await File.ReadAllTextAsync(ResolveFile(_parameters.TaskDefinitionFile).FullName);
+			string taskDefRendered = taskDefTemplate.Replace("%%DOCKER_IMAGE%%", _parameters.DockerImage, StringComparison.Ordinal);
+			if (_parameters.Version != null)
 			{
-				TaskDefRendered = TaskDefRendered.Replace("%%VERSION%%", Parameters.Version, StringComparison.Ordinal);
+				taskDefRendered = taskDefRendered.Replace("%%VERSION%%", _parameters.Version, StringComparison.Ordinal);
 			}
 
-			FileReference TempTaskDefFile = FileReference.Combine(Unreal.RootDirectory, "Engine", "Intermediate", "Build", "AwsEcsDeployTaskTemp.json");
-			DirectoryReference.CreateDirectory(TempTaskDefFile.Directory);
-			await File.WriteAllTextAsync(TempTaskDefFile.FullName, TaskDefRendered, new UTF8Encoding(encoderShouldEmitUTF8Identifier: false));
-			
-			IProcessResult CreateTaskDefResult = await SpawnTaskBase.ExecuteAsync("aws", $"ecs register-task-definition --cli-input-json \"file://{TempTaskDefFile.FullName}\"", EnvVars: ParseEnvVars(Parameters.Environment, Parameters.EnvironmentFile), LogOutput: Parameters.LogOutput);
+			FileReference tempTaskDefFile = FileReference.Combine(Unreal.RootDirectory, "Engine", "Intermediate", "Build", "AwsEcsDeployTaskTemp.json");
+			DirectoryReference.CreateDirectory(tempTaskDefFile.Directory);
+			await File.WriteAllTextAsync(tempTaskDefFile.FullName, taskDefRendered, new UTF8Encoding(encoderShouldEmitUTF8Identifier: false));
+
+			IProcessResult createTaskDefResult = await SpawnTaskBase.ExecuteAsync("aws", $"ecs register-task-definition --cli-input-json \"file://{tempTaskDefFile.FullName}\"", envVars: ParseEnvVars(_parameters.Environment, _parameters.EnvironmentFile), logOutput: _parameters.LogOutput);
 
-			JsonDocument TaskDefJson = JsonDocument.Parse(CreateTaskDefResult.Output);
-			string TaskDefFamily = TaskDefJson.RootElement.GetProperty("taskDefinition").GetProperty("family").GetString();
-			string TaskDefRevision = TaskDefJson.RootElement.GetProperty("taskDefinition").GetProperty("revision").ToString();
+			JsonDocument taskDefJson = JsonDocument.Parse(createTaskDefResult.Output);
+			string taskDefFamily = taskDefJson.RootElement.GetProperty("taskDefinition").GetProperty("family").GetString();
+			string taskDefRevision = taskDefJson.RootElement.GetProperty("taskDefinition").GetProperty("revision").ToString();
 
-			string Params = $"ecs update-service --cluster {Parameters.Cluster} --service {Parameters.Service} --task-definition {TaskDefFamily}:{TaskDefRevision}";
-			await SpawnTaskBase.ExecuteAsync("aws", Params, EnvVars: ParseEnvVars(Parameters.Environment, Parameters.EnvironmentFile), LogOutput: Parameters.LogOutput);
+			string @params = $"ecs update-service --cluster {_parameters.Cluster} --service {_parameters.Service} --task-definition {taskDefFamily}:{taskDefRevision}";
+			await SpawnTaskBase.ExecuteAsync("aws", @params, envVars: ParseEnvVars(_parameters.Environment, _parameters.EnvironmentFile), logOutput: _parameters.LogOutput);
 
-			Logger.LogInformation("Service {Service} updated to use new task def {TaskDefFamily}:{TaskDefRevision}", Parameters.Service, TaskDefFamily, TaskDefRevision);
+			Logger.LogInformation("Service {Service} updated to use new task def {TaskDefFamily}:{TaskDefRevision}", _parameters.Service, taskDefFamily, taskDefRevision);
 		}
 
 		/// <summary>
 		/// Output this task out to an XML writer.
 		/// </summary>
-		public override void Write(XmlWriter Writer)
+		public override void Write(XmlWriter writer)
 		{
-			Write(Writer, Parameters);
+			Write(writer, _parameters);
 		}
 
 		/// <summary>
diff --git a/Engine/Source/Programs/AutomationTool/BuildGraph/Tasks/AwsTask.cs b/Engine/Source/Programs/AutomationTool/BuildGraph/Tasks/AwsTask.cs
--- a/Engine/Source/Programs/AutomationTool/BuildGraph/Tasks/AwsTask.cs
+++ b/Engine/Source/Programs/AutomationTool/BuildGraph/Tasks/AwsTask.cs
@@ -1,14 +1,9 @@
 // Copyright Epic Games, Inc. All Rights Reserved.
 
-using AutomationTool;
-using EpicGames.Core;
-using System;
 using System.Collections.Generic;
-using System.IO;
-using System.Text;
-using System.Text.Json;
 using System.Threading.Tasks;
 using System.Xml;
+using EpicGames.Core;
 
 namespace AutomationTool.Tasks
 {
@@ -48,37 +43,34 @@ namespace AutomationTool.Tasks
 	[TaskElement("Aws", typeof(AwsTaskParameters))]
 	public class AwsTask : SpawnTaskBase
 	{
-		/// <summary>
-		/// Parameters for this task
-		/// </summary>
-		AwsTaskParameters Parameters;
+		readonly AwsTaskParameters _parameters;
 
 		/// <summary>
 		/// Construct an AWS CLI task
 		/// </summary>
-		/// <param name="InParameters">Parameters for the task</param>
-		public AwsTask(AwsTaskParameters InParameters)
+		/// <param name="parameters">Parameters for the task</param>
+		public AwsTask(AwsTaskParameters parameters)
 		{
-			Parameters = InParameters;
+			_parameters = parameters;
 		}
 
 		/// <summary>
-		/// Execute the task.
+		/// ExecuteAsync the task.
 		/// </summary>
-		/// <param name="Job">Information about the current job</param>
-		/// <param name="BuildProducts">Set of build products produced by this node.</param>
-		/// <param name="TagNameToFileSet">Mapping from tag names to the set of files they include</param>
-		public override async Task ExecuteAsync(JobContext Job, HashSet<FileReference> BuildProducts, Dictionary<string, HashSet<FileReference>> TagNameToFileSet)
+		/// <param name="job">Information about the current job</param>
+		/// <param name="buildProducts">Set of build products produced by this node.</param>
+		/// <param name="tagNameToFileSet">Mapping from tag names to the set of files they include</param>
+		public override async Task ExecuteAsync(JobContext job, HashSet<FileReference> buildProducts, Dictionary<string, HashSet<FileReference>> tagNameToFileSet)
 		{
-			await SpawnTaskBase.ExecuteAsync("aws", Parameters.Arguments, EnvVars: ParseEnvVars(Parameters.Environment, Parameters.EnvironmentFile), LogOutput: Parameters.LogOutput);
+			await SpawnTaskBase.ExecuteAsync("aws", _parameters.Arguments, envVars: ParseEnvVars(_parameters.Environment, _parameters.EnvironmentFile), logOutput: _parameters.LogOutput);
 		}
 
 		/// <summary>
 		/// Output this task out to an XML writer.
 		/// </summary>
-		public override void Write(XmlWriter Writer)
+		public override void Write(XmlWriter writer)
 		{
-			Write(Writer, Parameters);
+			Write(writer, _parameters);
 		}
 
 		/// <summary>
diff --git a/Engine/Source/Programs/AutomationTool/BuildGraph/Tasks/CheckMarkdownTask.cs b/Engine/Source/Programs/AutomationTool/BuildGraph/Tasks/CheckMarkdownTask.cs
--- a/Engine/Source/Programs/AutomationTool/BuildGraph/Tasks/CheckMarkdownTask.cs
+++ b/Engine/Source/Programs/AutomationTool/BuildGraph/Tasks/CheckMarkdownTask.cs
@@ -1,14 +1,13 @@
 // Copyright Epic Games, Inc. All Rights Reserved.
 
-using EpicGames.Core;
-using Microsoft.Extensions.Logging;
 using System;
 using System.Collections.Generic;
 using System.Linq;
-using System.Text;
 using System.Text.RegularExpressions;
 using System.Threading.Tasks;
 using System.Xml;
+using EpicGames.Core;
+using Microsoft.Extensions.Logging;
 using UnrealBuildBase;
 
 namespace AutomationTool.Tasks
@@ -58,7 +57,7 @@ namespace AutomationTool.Tasks
 			}
 
 			// Find the anchors in any markdown files
-			foreach(FileReference file in markdownFiles)
+			foreach (FileReference file in markdownFiles)
 			{
 				string[] lines = await FileReference.ReadAllLinesAsync(file);
 				foreach (string line in lines)
@@ -134,8 +133,8 @@ namespace AutomationTool.Tasks
 		}
 
 		/// <inheritdoc/>
-		public override void Write(XmlWriter Writer)
-			=> Write(Writer, _parameters);
+		public override void Write(XmlWriter writer)
+			=> Write(writer, _parameters);
 
 		/// <inheritdoc/>
 		public override IEnumerable<string> FindConsumedTagNames()
diff --git a/Engine/Source/Programs/AutomationTool/BuildGraph/Tasks/CommandTask.cs b/Engine/Source/Programs/AutomationTool/BuildGraph/Tasks/CommandTask.cs
--- a/Engine/Source/Programs/AutomationTool/BuildGraph/Tasks/CommandTask.cs
+++ b/Engine/Source/Programs/AutomationTool/BuildGraph/Tasks/CommandTask.cs
@@ -2,28 +2,21 @@
 
 using System;
 using System.Collections.Generic;
-using System.Linq;
 using System.Text;
 using System.Threading.Tasks;
-using AutomationTool;
-using UnrealBuildTool;
 using System.Xml;
-using System.IO;
-using EpicGames.Core;
-using UnrealBuildBase;
 using EpicGames.BuildGraph;
-using AutomationTool.Tasks;
+using EpicGames.Core;
 using Microsoft.Extensions.Logging;
-
-using static AutomationTool.CommandUtils;
+using UnrealBuildBase;
 
 namespace AutomationTool.Tasks
 {
 	static class StringExtensions
 	{
-		public static bool CaseInsensitiveContains(this string Text, string Value)
+		public static bool CaseInsensitiveContains(this string text, string value)
 		{
-			return System.Globalization.CultureInfo.InvariantCulture.CompareInfo.IndexOf(Text, Value, System.Globalization.CompareOptions.IgnoreCase) >= 0;
+			return System.Globalization.CultureInfo.InvariantCulture.CompareInfo.IndexOf(text, value, System.Globalization.CompareOptions.IgnoreCase) >= 0;
 		}
 	}
 
@@ -60,83 +53,83 @@ namespace AutomationTool.Tasks
 		/// <summary>
 		/// Parameters for this task
 		/// </summary>
-		CommandTaskParameters Parameters;
+		readonly CommandTaskParameters _parameters;
 
 		/// <summary>
 		/// Construct a new CommandTask.
 		/// </summary>
-		/// <param name="InParameters">Parameters for this task</param>
-		public CommandTask(CommandTaskParameters InParameters)
+		/// <param name="parameters">Parameters for this task</param>
+		public CommandTask(CommandTaskParameters parameters)
 		{
-			Parameters = InParameters;
+			_parameters = parameters;
 		}
 
 		/// <summary>
-		/// Execute the task.
+		/// ExecuteAsync the task.
 		/// </summary>
-		/// <param name="Job">Information about the current job</param>
-		/// <param name="BuildProducts">Set of build products produced by this node.</param>
-		/// <param name="TagNameToFileSet">Mapping from tag names to the set of files they include</param>
-		public override Task ExecuteAsync(JobContext Job, HashSet<FileReference> BuildProducts, Dictionary<string, HashSet<FileReference>> TagNameToFileSet)
+		/// <param name="job">Information about the current job</param>
+		/// <param name="buildProducts">Set of build products produced by this node.</param>
+		/// <param name="tagNameToFileSet">Mapping from tag names to the set of files they include</param>
+		public override Task ExecuteAsync(JobContext job, HashSet<FileReference> buildProducts, Dictionary<string, HashSet<FileReference>> tagNameToFileSet)
 		{
 			// If we're merging telemetry from the child process, get a temp filename for it
-			FileReference TelemetryFile = null;
-			if (Parameters.MergeTelemetryWithPrefix != null)
+			FileReference telemetryFile = null;
+			if (_parameters.MergeTelemetryWithPrefix != null)
 			{
-				TelemetryFile = FileReference.Combine(Unreal.RootDirectory, "Engine", "Intermediate", "UAT", "Telemetry.json");
-				DirectoryReference.CreateDirectory(TelemetryFile.Directory);
+				telemetryFile = FileReference.Combine(Unreal.RootDirectory, "Engine", "Intermediate", "UAT", "Telemetry.json");
+				DirectoryReference.CreateDirectory(telemetryFile.Directory);
 			}
 
 			// Run the command
-			StringBuilder CommandLine = new StringBuilder();
-			if (Parameters.Arguments == null || (!Parameters.Arguments.CaseInsensitiveContains("-p4") && !Parameters.Arguments.CaseInsensitiveContains("-nop4")))
+			StringBuilder commandLine = new StringBuilder();
+			if (_parameters.Arguments == null || (!_parameters.Arguments.CaseInsensitiveContains("-p4") && !_parameters.Arguments.CaseInsensitiveContains("-nop4")))
 			{
-				CommandLine.AppendFormat("{0} ", CommandUtils.P4Enabled ? "-p4" : "-nop4");
+				commandLine.AppendFormat("{0} ", CommandUtils.P4Enabled ? "-p4" : "-nop4");
 			}
-			if (Parameters.Arguments == null || (!Parameters.Arguments.CaseInsensitiveContains("-submit") && !Parameters.Arguments.CaseInsensitiveContains("-nosubmit")))
+			if (_parameters.Arguments == null || (!_parameters.Arguments.CaseInsensitiveContains("-submit") && !_parameters.Arguments.CaseInsensitiveContains("-nosubmit")))
 			{
-				if(GlobalCommandLine.Submit)
+				if (GlobalCommandLine.Submit)
 				{
-					CommandLine.Append("-submit ");
+					commandLine.Append("-submit ");
 				}
-				if(GlobalCommandLine.NoSubmit)
+				if (GlobalCommandLine.NoSubmit)
 				{
-					CommandLine.Append("-nosubmit ");
+					commandLine.Append("-nosubmit ");
 				}
 			}
-			if (Parameters.Arguments == null || !Parameters.Arguments.CaseInsensitiveContains("-uselocalbuildstorage"))
+			if (_parameters.Arguments == null || !_parameters.Arguments.CaseInsensitiveContains("-uselocalbuildstorage"))
 			{
 				if (GlobalCommandLine.UseLocalBuildStorage)
 				{
-					CommandLine.Append("-uselocalbuildstorage ");
+					commandLine.Append("-uselocalbuildstorage ");
 				}
 			}
 
-			CommandLine.Append("-NoCompile ");
-			CommandLine.Append(Parameters.Name);
-			if (!String.IsNullOrEmpty(Parameters.Arguments))
+			commandLine.Append("-NoCompile ");
+			commandLine.Append(_parameters.Name);
+			if (!String.IsNullOrEmpty(_parameters.Arguments))
 			{
-				CommandLine.AppendFormat(" {0}", Parameters.Arguments);
+				commandLine.AppendFormat(" {0}", _parameters.Arguments);
 			}
-			if (TelemetryFile != null)
+			if (telemetryFile != null)
 			{
-				CommandLine.AppendFormat(" -Telemetry={0}", CommandUtils.MakePathSafeToUseWithCommandLine(TelemetryFile.FullName));
+				commandLine.AppendFormat(" -Telemetry={0}", CommandUtils.MakePathSafeToUseWithCommandLine(telemetryFile.FullName));
 			}
-			CommandUtils.RunUAT(CommandUtils.CmdEnv, CommandLine.ToString(), Identifier: Parameters.Name);
+			CommandUtils.RunUAT(CommandUtils.CmdEnv, commandLine.ToString(), Identifier: _parameters.Name);
 
 			// Merge in any new telemetry data that was produced
-			if (TelemetryFile != null && FileReference.Exists(TelemetryFile))
+			if (telemetryFile != null && FileReference.Exists(telemetryFile))
 			{
-				Logger.LogDebug("Merging telemetry from {TelemetryFile}", TelemetryFile);
+				Logger.LogDebug("Merging telemetry from {TelemetryFile}", telemetryFile);
 
-				TelemetryData NewTelemetry;
-				if (TelemetryData.TryRead(TelemetryFile, out NewTelemetry))
+				TelemetryData newTelemetry;
+				if (TelemetryData.TryRead(telemetryFile, out newTelemetry))
 				{
-					CommandUtils.Telemetry.Merge(Parameters.MergeTelemetryWithPrefix, NewTelemetry);
+					CommandUtils.Telemetry.Merge(_parameters.MergeTelemetryWithPrefix, newTelemetry);
 				}
 				else
 				{
-					Logger.LogWarning("Unable to read UAT telemetry file from {TelemetryFile}", TelemetryFile);
+					Logger.LogWarning("Unable to read UAT telemetry file from {TelemetryFile}", telemetryFile);
 				}
 			}
 			return Task.CompletedTask;
@@ -148,15 +141,15 @@ namespace AutomationTool.Tasks
 		/// <returns>Name of the trace</returns>
 		public override string GetTraceName()
 		{
-			return String.Format("{0}.{1}", base.GetTraceName(), Parameters.Name.ToLowerInvariant());
+			return String.Format("{0}.{1}", base.GetTraceName(), _parameters.Name.ToLowerInvariant());
 		}
 
 		/// <summary>
 		/// Output this task out to an XML writer.
 		/// </summary>
-		public override void Write(XmlWriter Writer)
+		public override void Write(XmlWriter writer)
 		{
-			Write(Writer, Parameters);
+			Write(writer, _parameters);
 		}
 
 		/// <summary>
@@ -183,18 +176,20 @@ namespace AutomationTool.Tasks
 		/// <summary>
 		/// Runs another UAT command
 		/// </summary>
-		/// <param name="State">The execution state</param>
-		/// <param name="Name">Name of the command to run</param>
-		/// <param name="Arguments">Arguments for the command</param>
-		/// <param name="MergeTelemetryWithPrefix">If non-null, instructs telemetry from the command to be merged into the telemetry for this UAT instance with the given prefix. May be an empty (non-null) string.</param>
-		public static async Task CommandAsync(this BgContext State, string Name, string Arguments = null, string MergeTelemetryWithPrefix = null)
+		/// <param name="state">The execution state</param>
+		/// <param name="name">Name of the command to run</param>
+		/// <param name="arguments">Arguments for the command</param>
+		/// <param name="mergeTelemetryWithPrefix">If non-null, instructs telemetry from the command to be merged into the telemetry for this UAT instance with the given prefix. May be an empty (non-null) string.</param>
+		public static async Task CommandAsync(this BgContext state, string name, string arguments = null, string mergeTelemetryWithPrefix = null)
 		{
-			CommandTaskParameters Parameters = new CommandTaskParameters();
-			Parameters.Name = Name;
-			Parameters.Arguments = Arguments ?? Parameters.Arguments;
-			Parameters.MergeTelemetryWithPrefix = MergeTelemetryWithPrefix ?? Parameters.MergeTelemetryWithPrefix;
+			_ = state;
+
+			CommandTaskParameters parameters = new CommandTaskParameters();
+			parameters.Name = name;
+			parameters.Arguments = arguments ?? parameters.Arguments;
+			parameters.MergeTelemetryWithPrefix = mergeTelemetryWithPrefix ?? parameters.MergeTelemetryWithPrefix;
 
-			await ExecuteAsync(new CommandTask(Parameters));
+			await ExecuteAsync(new CommandTask(parameters));
 		}
 	}
 }
diff --git a/Engine/Source/Programs/AutomationTool/BuildGraph/Tasks/CommandletTask.cs b/Engine/Source/Programs/AutomationTool/BuildGraph/Tasks/CommandletTask.cs
--- a/Engine/Source/Programs/AutomationTool/BuildGraph/Tasks/CommandletTask.cs
+++ b/Engine/Source/Programs/AutomationTool/BuildGraph/Tasks/CommandletTask.cs
@@ -1,15 +1,12 @@
 // Copyright Epic Games, Inc. All Rights Reserved.
 
-using EpicGames.BuildGraph;
 using System;
 using System.Collections.Generic;
-using System.IO;
 using System.Linq;
-using System.Text;
 using System.Threading.Tasks;
 using System.Xml;
+using EpicGames.BuildGraph;
 using EpicGames.Core;
-using UnrealBuildBase;
 using UnrealBuildTool;
 
 namespace AutomationTool.Tasks
@@ -56,65 +53,62 @@ namespace AutomationTool.Tasks
 	[TaskElement("Commandlet", typeof(CommandletTaskParameters))]
 	public class CommandletTask : BgTaskImpl
 	{
-		/// <summary>
-		/// Parameters for this task
-		/// </summary>
-		CommandletTaskParameters Parameters;
+		readonly CommandletTaskParameters _parameters;
 
 		/// <summary>
 		/// Construct a new CommandletTask.
 		/// </summary>
-		/// <param name="InParameters">Parameters for this task</param>
-		public CommandletTask(CommandletTaskParameters InParameters)
+		/// <param name="parameters">Parameters for this task</param>
+		public CommandletTask(CommandletTaskParameters parameters)
 		{
-			Parameters = InParameters;
+			_parameters = parameters;
 		}
 
 		/// <summary>
-		/// Execute the task.
+		/// ExecuteAsync the task.
 		/// </summary>
-		/// <param name="Job">Information about the current job</param>
-		/// <param name="BuildProducts">Set of build products produced by this node.</param>
-		/// <param name="TagNameToFileSet">Mapping from tag names to the set of files they include</param>
-		public override Task ExecuteAsync(JobContext Job, HashSet<FileReference> BuildProducts, Dictionary<string, HashSet<FileReference>> TagNameToFileSet)
+		/// <param name="job">Information about the current job</param>
+		/// <param name="buildProducts">Set of build products produced by this node.</param>
+		/// <param name="tagNameToFileSet">Mapping from tag names to the set of files they include</param>
+		public override Task ExecuteAsync(JobContext job, HashSet<FileReference> buildProducts, Dictionary<string, HashSet<FileReference>> tagNameToFileSet)
 		{
 			// Get the full path to the project file
-			FileReference ProjectFile = null;
-			if(!String.IsNullOrEmpty(Parameters.Project))
+			FileReference projectFile = null;
+			if (!String.IsNullOrEmpty(_parameters.Project))
 			{
-				if(Parameters.Project.EndsWith(".uproject", StringComparison.OrdinalIgnoreCase))
+				if (_parameters.Project.EndsWith(".uproject", StringComparison.OrdinalIgnoreCase))
 				{
-					ProjectFile = ResolveFile(Parameters.Project);
+					projectFile = ResolveFile(_parameters.Project);
 				}
 				else
 				{
-					ProjectFile = NativeProjects.EnumerateProjectFiles(Log.Logger).FirstOrDefault(x => x.GetFileNameWithoutExtension().Equals(Parameters.Project, StringComparison.OrdinalIgnoreCase));
+					projectFile = NativeProjects.EnumerateProjectFiles(Log.Logger).FirstOrDefault(x => x.GetFileNameWithoutExtension().Equals(_parameters.Project, StringComparison.OrdinalIgnoreCase));
 				}
 
-				if(ProjectFile == null || !FileReference.Exists(ProjectFile))
+				if (projectFile == null || !FileReference.Exists(projectFile))
 				{
-					throw new BuildException("Unable to resolve project '{0}'", Parameters.Project);
+					throw new BuildException("Unable to resolve project '{0}'", _parameters.Project);
 				}
 			}
 
 			// Get the path to the editor, and check it exists
-			FileSystemReference EditorExe = Parameters.EditorExe;
-			if (EditorExe == null)
+			FileSystemReference editorExe = _parameters.EditorExe;
+			if (editorExe == null)
 			{
-				EditorExe = ProjectUtils.GetEditorForProject(ProjectFile);
+				editorExe = ProjectUtils.GetEditorForProject(projectFile);
 			}
 
 			// Run the commandlet
-			CommandUtils.RunCommandlet(ProjectFile, EditorExe.FullName, Parameters.Name, Parameters.Arguments, Parameters.ErrorLevel);
+			CommandUtils.RunCommandlet(projectFile, editorExe.FullName, _parameters.Name, _parameters.Arguments, _parameters.ErrorLevel);
 			return Task.CompletedTask;
 		}
 
 		/// <summary>
 		/// Output this task out to an XML writer.
 		/// </summary>
-		public override void Write(XmlWriter Writer)
+		public override void Write(XmlWriter writer)
 		{
-			Write(Writer, Parameters);
+			Write(writer, _parameters);
 		}
 
 		/// <summary>
@@ -144,21 +138,23 @@ namespace AutomationTool.Tasks
 		/// <summary>
 		/// Task which runs a UE commandlet
 		/// </summary>
-		/// <param name="State"></param>
-		/// <param name="Name">The commandlet name to execute.</param>
-		/// <param name="Project">The project to run the editor with.</param>
-		/// <param name="Arguments">Arguments to be passed to the commandlet.</param>
-		/// <param name="EditorExe">The editor executable to use. Defaults to the development UnrealEditor executable for the current platform.</param>
-		/// <param name="ErrorLevel">The minimum exit code, which is treated as an error.</param>
-		public static async Task CommandletAsync(this BgContext State, string Name, FileReference Project = null, string Arguments = null, FileReference EditorExe = null, int ErrorLevel = 1)
+		/// <param name="state"></param>
+		/// <param name="name">The commandlet name to execute.</param>
+		/// <param name="project">The project to run the editor with.</param>
+		/// <param name="arguments">Arguments to be passed to the commandlet.</param>
+		/// <param name="editorExe">The editor executable to use. Defaults to the development UnrealEditor executable for the current platform.</param>
+		/// <param name="errorLevel">The minimum exit code, which is treated as an error.</param>
+		public static async Task CommandletAsync(this BgContext state, string name, FileReference project = null, string arguments = null, FileReference editorExe = null, int errorLevel = 1)
 		{
-			CommandletTaskParameters Parameters = new CommandletTaskParameters();
-			Parameters.Name = Name;
-			Parameters.Project = Project?.FullName;
-			Parameters.Arguments = Arguments;
-			Parameters.EditorExe = EditorExe;
-			Parameters.ErrorLevel = ErrorLevel;
-			await ExecuteAsync(new CommandletTask(Parameters));
+			_ = state;
+
+			CommandletTaskParameters parameters = new CommandletTaskParameters();
+			parameters.Name = name;
+			parameters.Project = project?.FullName;
+			parameters.Arguments = arguments;
+			parameters.EditorExe = editorExe;
+			parameters.ErrorLevel = errorLevel;
+			await ExecuteAsync(new CommandletTask(parameters));
 		}
 	}
 }
diff --git a/Engine/Source/Programs/AutomationTool/BuildGraph/Tasks/CompileTask.cs b/Engine/Source/Programs/AutomationTool/BuildGraph/Tasks/CompileTask.cs
--- a/Engine/Source/Programs/AutomationTool/BuildGraph/Tasks/CompileTask.cs
+++ b/Engine/Source/Programs/AutomationTool/BuildGraph/Tasks/CompileTask.cs
@@ -1,16 +1,13 @@
 // Copyright Epic Games, Inc. All Rights Reserved.
 
-using EpicGames.BuildGraph;
 using System;
 using System.Collections.Generic;
 using System.Linq;
-using System.Text;
 using System.Threading.Tasks;
-using UnrealBuildTool;
-using AutomationTool;
 using System.Xml;
 using EpicGames.Core;
 using OpenTracing;
+using UnrealBuildTool;
 
 namespace AutomationTool.Tasks
 {
@@ -22,7 +19,7 @@ namespace AutomationTool.Tasks
 		/// <summary>
 		/// The target to compile.
 		/// </summary>
-		[TaskParameter(Optional=true)]
+		[TaskParameter(Optional = true)]
 		public string Target { get; set; }
 
 		/// <summary>
@@ -95,143 +92,143 @@ namespace AutomationTool.Tasks
 		/// <summary>
 		/// List of targets to compile. As well as the target specifically added for this task, additional compile tasks may be merged with it.
 		/// </summary>
-		List<UnrealBuild.BuildTarget> Targets = new List<UnrealBuild.BuildTarget>();
+		readonly List<UnrealBuild.BuildTarget> _targets = new List<UnrealBuild.BuildTarget>();
 
 		/// <summary>
 		/// Mapping of receipt filename to its corresponding tag name
 		/// </summary>
-		Dictionary<UnrealBuild.BuildTarget, string> TargetToTagName = new Dictionary<UnrealBuild.BuildTarget,string>();
+		readonly Dictionary<UnrealBuild.BuildTarget, string> _targetToTagName = new Dictionary<UnrealBuild.BuildTarget, string>();
 
 		/// <summary>
 		/// Whether to allow using XGE for this job
 		/// </summary>
-		bool bAllowXGE = true;
+		bool _allowXge = true;
 
 		/// <summary>
 		/// Whether to allow using all available cores for this job, when bAllowXGE is false
 		/// </summary>
-		bool bAllowAllCores = false;
+		bool _allowAllCores = false;
 
 		/// <summary>
 		/// Should SkipBuild be passed to UBT so that only .target files are generated.
 		/// </summary>
-		bool bSkipBuild = false;
+		bool _skipBuild = false;
 
 		/// <summary>
 		/// Constructor
 		/// </summary>
-		/// <param name="Task">Initial task to execute</param>
-		public CompileTaskExecutor(CompileTask Task)
+		/// <param name="task">Initial task to execute</param>
+		public CompileTaskExecutor(CompileTask task)
 		{
-			Add(Task);
+			Add(task);
 		}
 
 		/// <summary>
 		/// Adds another task to this executor
 		/// </summary>
-		/// <param name="Task">Task to add</param>
+		/// <param name="task">Task to add</param>
 		/// <returns>True if the task could be added, false otherwise</returns>
-		public bool Add(BgTaskImpl Task)
+		public bool Add(BgTaskImpl task)
 		{
-			CompileTask CompileTask = Task as CompileTask;
-			if(CompileTask == null)
+			CompileTask compileTask = task as CompileTask;
+			if (compileTask == null)
 			{
 				return false;
 			}
 
-			CompileTaskParameters Parameters = CompileTask.Parameters;
-			if (Targets.Count > 0)
+			CompileTaskParameters parameters = compileTask.Parameters;
+			if (_targets.Count > 0)
 			{
-				if (bAllowXGE != Parameters.AllowXGE
-					|| bSkipBuild != Parameters.SkipBuild)
+				if (_allowXge != parameters.AllowXGE
+					|| _skipBuild != parameters.SkipBuild)
 				{
 					return false;
 				}
 			}
 			else
 			{
-				bAllowXGE = Parameters.AllowXGE;
-				bAllowAllCores = Parameters.AllowAllCores;
-				bSkipBuild = Parameters.SkipBuild;
+				_allowXge = parameters.AllowXGE;
+				_allowAllCores = parameters.AllowAllCores;
+				_skipBuild = parameters.SkipBuild;
 			}
 
-			bAllowXGE &= Parameters.AllowXGE;
-			bAllowAllCores &= Parameters.AllowAllCores;
+			_allowXge &= parameters.AllowXGE;
+			_allowAllCores &= parameters.AllowAllCores;
 
-			UnrealBuild.BuildTarget Target = new UnrealBuild.BuildTarget { TargetName = Parameters.Target, Platform = Parameters.Platform, Config = Parameters.Configuration, UprojectPath = CompileTask.FindProjectFile(), UBTArgs = (Parameters.Arguments ?? ""), Clean = Parameters.Clean };
-			if(!String.IsNullOrEmpty(Parameters.Tag))
+			UnrealBuild.BuildTarget target = new UnrealBuild.BuildTarget { TargetName = parameters.Target, Platform = parameters.Platform, Config = parameters.Configuration, UprojectPath = compileTask.FindProjectFile(), UBTArgs = (parameters.Arguments ?? ""), Clean = parameters.Clean };
+			if (!String.IsNullOrEmpty(parameters.Tag))
 			{
-				TargetToTagName.Add(Target, Parameters.Tag);
+				_targetToTagName.Add(target, parameters.Tag);
 			}
-			Targets.Add(Target);
+			_targets.Add(target);
 
 			return true;
 		}
 
 		/// <summary>
-		/// Execute all the tasks added to this executor.
+		/// ExecuteAsync all the tasks added to this executor.
 		/// </summary>
-		/// <param name="Job">Information about the current job</param>
-		/// <param name="BuildProducts">Set of build products produced by this node.</param>
-		/// <param name="TagNameToFileSet">Mapping from tag names to the set of files they include</param>
+		/// <param name="job">Information about the current job</param>
+		/// <param name="buildProducts">Set of build products produced by this node.</param>
+		/// <param name="tagNameToFileSet">Mapping from tag names to the set of files they include</param>
 		/// <returns>Whether the task succeeded or not. Exiting with an exception will be caught and treated as a failure.</returns>
-		public Task ExecuteAsync(JobContext Job, HashSet<FileReference> BuildProducts, Dictionary<string, HashSet<FileReference>> TagNameToFileSet)
+		public Task ExecuteAsync(JobContext job, HashSet<FileReference> buildProducts, Dictionary<string, HashSet<FileReference>> tagNameToFileSet)
 		{
 			// Create the agenda
-			UnrealBuild.BuildAgenda Agenda = new UnrealBuild.BuildAgenda();
-			Agenda.Targets.AddRange(Targets);
+			UnrealBuild.BuildAgenda agenda = new UnrealBuild.BuildAgenda();
+			agenda.Targets.AddRange(_targets);
 
 			// Build everything
-			Dictionary<UnrealBuild.BuildTarget, BuildManifest> TargetToManifest = new Dictionary<UnrealBuild.BuildTarget,BuildManifest>();
-			UnrealBuild Builder = new UnrealBuild(Job.OwnerCommand);
+			Dictionary<UnrealBuild.BuildTarget, BuildManifest> targetToManifest = new Dictionary<UnrealBuild.BuildTarget, BuildManifest>();
+			UnrealBuild builder = new UnrealBuild(job.OwnerCommand);
 
-			bool bAllCores = (CommandUtils.IsBuildMachine || bAllowAllCores);   // Enable using all cores if this is a build agent or the flag was passed in to the task and XGE is disabled.
-			Builder.Build(Agenda, InDeleteBuildProducts: null, InUpdateVersionFiles: false, InForceNoXGE: !bAllowXGE, InAllCores: bAllCores, InTargetToManifest: TargetToManifest, InSkipBuild: bSkipBuild);
+			bool allCores = (CommandUtils.IsBuildMachine || _allowAllCores);   // Enable using all cores if this is a build agent or the flag was passed in to the task and XGE is disabled.
+			builder.Build(agenda, InDeleteBuildProducts: null, InUpdateVersionFiles: false, InForceNoXGE: !_allowXge, InAllCores: allCores, InTargetToManifest: targetToManifest, InSkipBuild: _skipBuild);
 
-			UnrealBuild.CheckBuildProducts(Builder.BuildProductFiles);
+			UnrealBuild.CheckBuildProducts(builder.BuildProductFiles);
 
 			// Tag all the outputs
-			foreach(KeyValuePair<UnrealBuild.BuildTarget, string> TargetTagName in TargetToTagName)
+			foreach (KeyValuePair<UnrealBuild.BuildTarget, string> targetTagName in _targetToTagName)
 			{
-				BuildManifest Manifest;
-				if(!TargetToManifest.TryGetValue(TargetTagName.Key, out Manifest))
+				BuildManifest manifest;
+				if (!targetToManifest.TryGetValue(targetTagName.Key, out manifest))
 				{
-					throw new AutomationException("Missing manifest for target {0} {1} {2}", TargetTagName.Key.TargetName, TargetTagName.Key.Platform, TargetTagName.Key.Config);
+					throw new AutomationException("Missing manifest for target {0} {1} {2}", targetTagName.Key.TargetName, targetTagName.Key.Platform, targetTagName.Key.Config);
 				}
 
-				HashSet<FileReference> ManifestBuildProducts = Manifest.BuildProducts.Select(x => new FileReference(x)).ToHashSet();
+				HashSet<FileReference> manifestBuildProducts = manifest.BuildProducts.Select(x => new FileReference(x)).ToHashSet();
 
 				// when we make a Mac/IOS build, Xcode will finalize the .app directory, adding files that UBT has no idea about, so now we recursively add any files in the .app
 				// as BuildProducts. look for any .apps that we have any files as BuildProducts, and expand to include all files in the .app
 				if (BuildHostPlatform.Current.Platform == UnrealTargetPlatform.Mac)
 				{
-					HashSet<string> AppBundleLocations = new();
-					foreach (FileReference File in ManifestBuildProducts)
+					HashSet<string> appBundleLocations = new();
+					foreach (FileReference file in manifestBuildProducts)
 					{
 						// look for a ".app/" portion and chop off anything after it
-						int AppLocation = File.FullName.IndexOf(".app/", StringComparison.InvariantCultureIgnoreCase);
-						if (AppLocation > 0)
+						int appLocation = file.FullName.IndexOf(".app/", StringComparison.InvariantCultureIgnoreCase);
+						if (appLocation > 0)
 						{
-							AppBundleLocations.Add(File.FullName.Substring(0, AppLocation + 4));
+							appBundleLocations.Add(file.FullName.Substring(0, appLocation + 4));
 						}
 					}
 
 					// now with a unique set of app bundles, add all files in them
-					foreach (string AppBundleLocation in AppBundleLocations)
+					foreach (string appBundleLocation in appBundleLocations)
 					{
-						ManifestBuildProducts.UnionWith(DirectoryReference.EnumerateFiles(new DirectoryReference(AppBundleLocation), "*", System.IO.SearchOption.AllDirectories));
+						manifestBuildProducts.UnionWith(DirectoryReference.EnumerateFiles(new DirectoryReference(appBundleLocation), "*", System.IO.SearchOption.AllDirectories));
 					}
 				}
 
-				foreach (string TagName in CustomTask.SplitDelimitedList(TargetTagName.Value))
+				foreach (string tagName in CustomTask.SplitDelimitedList(targetTagName.Value))
 				{
-					HashSet<FileReference> FileSet = CustomTask.FindOrAddTagSet(TagNameToFileSet, TagName);
-					FileSet.UnionWith(ManifestBuildProducts);
+					HashSet<FileReference> fileSet = CustomTask.FindOrAddTagSet(tagNameToFileSet, tagName);
+					fileSet.UnionWith(manifestBuildProducts);
 				}
 			}
 
 			// Add everything to the list of build products
-			BuildProducts.UnionWith(Builder.BuildProductFiles.Select(x => new FileReference(x)));
+			buildProducts.UnionWith(builder.BuildProductFiles.Select(x => new FileReference(x)));
 			return Task.CompletedTask;
 		}
 	}
@@ -255,10 +252,10 @@ namespace AutomationTool.Tasks
 		/// <summary>
 		/// Construct a compile task
 		/// </summary>
-		/// <param name="Parameters">Parameters for this task</param>
-		public CompileTask(CompileTaskParameters Parameters)
+		/// <param name="parameters">Parameters for this task</param>
+		public CompileTask(CompileTaskParameters parameters)
 		{
-			this.Parameters = Parameters;
+			this.Parameters = parameters;
 		}
 
 		/// <summary>
@@ -266,42 +263,42 @@ namespace AutomationTool.Tasks
 		/// </summary>
 		public FileReference FindProjectFile()
 		{
-			FileReference ProjectFile = null;
+			FileReference projectFile = null;
 
 			// Resolve the full path to the project file
-			if(!String.IsNullOrEmpty(Parameters.Project))
+			if (!String.IsNullOrEmpty(Parameters.Project))
 			{
-				if(Parameters.Project.EndsWith(".uproject", StringComparison.OrdinalIgnoreCase))
+				if (Parameters.Project.EndsWith(".uproject", StringComparison.OrdinalIgnoreCase))
 				{
-					ProjectFile = CustomTask.ResolveFile(Parameters.Project);
+					projectFile = CustomTask.ResolveFile(Parameters.Project);
 				}
 				else
 				{
-					ProjectFile = NativeProjects.EnumerateProjectFiles(Log.Logger).FirstOrDefault(x => x.GetFileNameWithoutExtension().Equals(Parameters.Project, StringComparison.OrdinalIgnoreCase));
+					projectFile = NativeProjects.EnumerateProjectFiles(Log.Logger).FirstOrDefault(x => x.GetFileNameWithoutExtension().Equals(Parameters.Project, StringComparison.OrdinalIgnoreCase));
 				}
 
-				if(ProjectFile == null || !FileReference.Exists(ProjectFile))
+				if (projectFile == null || !FileReference.Exists(projectFile))
 				{
 					throw new BuildException("Unable to resolve project '{0}'", Parameters.Project);
 				}
 			}
 
-			return ProjectFile;
+			return projectFile;
 		}
 
 		/// <summary>
-		/// Execute the task.
+		/// ExecuteAsync the task.
 		/// </summary>
-		/// <param name="Job">Information about the current job</param>
-		/// <param name="BuildProducts">Set of build products produced by this node.</param>
-		/// <param name="TagNameToFileSet">Mapping from tag names to the set of files they include</param>
-		public override Task ExecuteAsync(JobContext Job, HashSet<FileReference> BuildProducts, Dictionary<string, HashSet<FileReference>> TagNameToFileSet)
+		/// <param name="job">Information about the current job</param>
+		/// <param name="buildProducts">Set of build products produced by this node.</param>
+		/// <param name="tagNameToFileSet">Mapping from tag names to the set of files they include</param>
+		public override Task ExecuteAsync(JobContext job, HashSet<FileReference> buildProducts, Dictionary<string, HashSet<FileReference>> tagNameToFileSet)
 		{
 			//
 			// Don't do any logic here. You have to do it in the ctor or a getter
 			//  otherwise you break the ITaskExecutor pathway, which doesn't call this function!
 			//
-			return GetExecutor().ExecuteAsync(Job, BuildProducts, TagNameToFileSet);
+			return GetExecutor().ExecuteAsync(job, buildProducts, tagNameToFileSet);
 		}
 
 		/// <summary>
@@ -316,47 +313,47 @@ namespace AutomationTool.Tasks
 		/// <summary>
 		/// Get properties to include in tracing info
 		/// </summary>
-		/// <param name="Span">The span to add metadata to</param>
-		/// <param name="Prefix">Prefix for all metadata keys</param>
-		public override void GetTraceMetadata(ITraceSpan Span, string Prefix)
+		/// <param name="span">The span to add metadata to</param>
+		/// <param name="prefix">Prefix for all metadata keys</param>
+		public override void GetTraceMetadata(ITraceSpan span, string prefix)
 		{
-			base.GetTraceMetadata(Span, Prefix);
+			base.GetTraceMetadata(span, prefix);
 
-			Span.AddMetadata(Prefix + "target.name", Parameters.Target);
-			Span.AddMetadata(Prefix + "target.config", Parameters.Configuration.ToString());
-			Span.AddMetadata(Prefix + "target.platform", Parameters.Platform.ToString());
+			span.AddMetadata(prefix + "target.name", Parameters.Target);
+			span.AddMetadata(prefix + "target.config", Parameters.Configuration.ToString());
+			span.AddMetadata(prefix + "target.platform", Parameters.Platform.ToString());
 
 			if (Parameters.Project != null)
 			{
-				Span.AddMetadata(Prefix + "target.project", Parameters.Project);
+				span.AddMetadata(prefix + "target.project", Parameters.Project);
 			}
 		}
-		
+
 		/// <summary>
 		/// Get properties to include in tracing info
 		/// </summary>
-		/// <param name="Span">The span to add metadata to</param>
-		/// <param name="Prefix">Prefix for all metadata keys</param>
-		public override void GetTraceMetadata(ISpan Span, string Prefix)
+		/// <param name="span">The span to add metadata to</param>
+		/// <param name="prefix">Prefix for all metadata keys</param>
+		public override void GetTraceMetadata(ISpan span, string prefix)
 		{
-			base.GetTraceMetadata(Span, Prefix);
+			base.GetTraceMetadata(span, prefix);
 
-			Span.SetTag(Prefix + "target.name", Parameters.Target);
-			Span.SetTag(Prefix + "target.config", Parameters.Configuration.ToString());
-			Span.SetTag(Prefix + "target.platform", Parameters.Platform.ToString());
+			span.SetTag(prefix + "target.name", Parameters.Target);
+			span.SetTag(prefix + "target.config", Parameters.Configuration.ToString());
+			span.SetTag(prefix + "target.platform", Parameters.Platform.ToString());
 
 			if (Parameters.Project != null)
 			{
-				Span.SetTag(Prefix + "target.project", Parameters.Project);
+				span.SetTag(prefix + "target.project", Parameters.Project);
 			}
 		}
 
 		/// <summary>
 		/// Output this task out to an XML writer.
 		/// </summary>
-		public override void Write(XmlWriter Writer)
+		public override void Write(XmlWriter writer)
 		{
-			Write(Writer, Parameters);
+			Write(writer, Parameters);
 		}
 
 		/// <summary>
@@ -383,25 +380,25 @@ namespace AutomationTool.Tasks
 		/// <summary>
 		/// Compiles a target
 		/// </summary>
-		/// <param name="Target">The target to compile</param>
-		/// <param name="Configuration">The configuration to compile</param>
-		/// <param name="Platform">The platform to compile for</param>
-		/// <param name="Project">The project to compile with</param>
-		/// <param name="Arguments">Additional arguments for UnrealBuildTool</param>
-		/// <param name="AllowXGE">Whether to allow using XGE for compilation</param>
-		/// <param name="Clean">Whether to allow cleaning this target. If unspecified, targets are cleaned if the -Clean argument is passed on the command line</param>
+		/// <param name="target">The target to compile</param>
+		/// <param name="configuration">The configuration to compile</param>
+		/// <param name="platform">The platform to compile for</param>
+		/// <param name="project">The project to compile with</param>
+		/// <param name="arguments">Additional arguments for UnrealBuildTool</param>
+		/// <param name="allowXge">Whether to allow using XGE for compilation</param>
+		/// <param name="clean">Whether to allow cleaning this target. If unspecified, targets are cleaned if the -Clean argument is passed on the command line</param>
 		/// <returns>Build products from the compile</returns>
-		public static async Task<FileSet> CompileAsync(string Target, UnrealTargetPlatform Platform, UnrealTargetConfiguration Configuration, FileReference Project = null, string Arguments = null, bool AllowXGE = true, bool? Clean = null)
+		public static async Task<FileSet> CompileAsync(string target, UnrealTargetPlatform platform, UnrealTargetConfiguration configuration, FileReference project = null, string arguments = null, bool allowXge = true, bool? clean = null)
 		{
-			CompileTaskParameters Parameters = new CompileTaskParameters();
-			Parameters.Target = Target;
-			Parameters.Platform = Platform;
-			Parameters.Configuration = Configuration;
-			Parameters.Project = Project?.FullName;
-			Parameters.Arguments = Arguments;
-			Parameters.AllowXGE = AllowXGE;
-			Parameters.Clean = Clean;
-			return await ExecuteAsync(new CompileTask(Parameters));
+			CompileTaskParameters parameters = new CompileTaskParameters();
+			parameters.Target = target;
+			parameters.Platform = platform;
+			parameters.Configuration = configuration;
+			parameters.Project = project?.FullName;
+			parameters.Arguments = arguments;
+			parameters.AllowXGE = allowXge;
+			parameters.Clean = clean;
+			return await ExecuteAsync(new CompileTask(parameters));
 		}
 	}
 }
diff --git a/Engine/Source/Programs/AutomationTool/BuildGraph/Tasks/CookTask.cs b/Engine/Source/Programs/AutomationTool/BuildGraph/Tasks/CookTask.cs
--- a/Engine/Source/Programs/AutomationTool/BuildGraph/Tasks/CookTask.cs
+++ b/Engine/Source/Programs/AutomationTool/BuildGraph/Tasks/CookTask.cs
@@ -1,8 +1,6 @@
 // Copyright Epic Games, Inc. All Rights Reserved.
 
-using System;
 using System.Collections.Generic;
-using System.IO;
 using System.Linq;
 using System.Threading.Tasks;
 using System.Xml;
@@ -75,91 +73,91 @@ namespace AutomationTool.Tasks
 		/// <summary>
 		/// Parameters for the task
 		/// </summary>
-		CookTaskParameters Parameters;
+		readonly CookTaskParameters _parameters;
 
 		/// <summary>
 		/// Constructor.
 		/// </summary>
-		/// <param name="InParameters">Parameters for this task</param>
-		public CookTask(CookTaskParameters InParameters)
+		/// <param name="parameters">Parameters for this task</param>
+		public CookTask(CookTaskParameters parameters)
 		{
-			Parameters = InParameters;
+			_parameters = parameters;
 		}
 
 		/// <summary>
-		/// Execute the task.
+		/// ExecuteAsync the task.
 		/// </summary>
-		/// <param name="Job">Information about the current job</param>
-		/// <param name="BuildProducts">Set of build products produced by this node.</param>
-		/// <param name="TagNameToFileSet">Mapping from tag names to the set of files they include</param>
-		public override Task ExecuteAsync(JobContext Job, HashSet<FileReference> BuildProducts, Dictionary<string, HashSet<FileReference>> TagNameToFileSet)
+		/// <param name="job">Information about the current job</param>
+		/// <param name="buildProducts">Set of build products produced by this node.</param>
+		/// <param name="tagNameToFileSet">Mapping from tag names to the set of files they include</param>
+		public override Task ExecuteAsync(JobContext job, HashSet<FileReference> buildProducts, Dictionary<string, HashSet<FileReference>> tagNameToFileSet)
 		{
 			// Figure out the project that this target belongs to
-			FileReference ProjectFile = null;
-			if(Parameters.Project != null)
+			FileReference projectFile = null;
+			if (_parameters.Project != null)
 			{
-				ProjectFile = new FileReference(Parameters.Project);
-				if(!FileReference.Exists(ProjectFile))
+				projectFile = new FileReference(_parameters.Project);
+				if (!FileReference.Exists(projectFile))
 				{
-					throw new AutomationException("Missing project file - {0}", ProjectFile.FullName);
+					throw new AutomationException("Missing project file - {0}", projectFile.FullName);
 				}
 			}
 
-			// Execute the cooker
-			using (IScope Scope = GlobalTracer.Instance.BuildSpan("Cook").StartActive())
+			// ExecuteAsync the cooker
+			using (IScope scope = GlobalTracer.Instance.BuildSpan("Cook").StartActive())
 			{
-				Scope.Span.SetTag("project", ProjectFile == null ? "UE4" : ProjectFile.GetFileNameWithoutExtension());
-				Scope.Span.SetTag("platform", Parameters.Platform);
-				string[] Maps = (Parameters.Maps == null)? null : Parameters.Maps.Split(new char[]{ '+' });
-				string Arguments = (Parameters.Versioned ? "" : "-Unversioned ") + "-LogCmds=\"LogSavePackage Warning\" " + Parameters.Arguments;
-				string EditorExe = (string.IsNullOrWhiteSpace(Parameters.EditorExe) ? ProjectUtils.GetEditorForProject(ProjectFile).FullName : Parameters.EditorExe);
-				CommandUtils.CookCommandlet(ProjectFile, EditorExe, Maps, null, null, null, Parameters.Platform, Arguments);
+				scope.Span.SetTag("project", projectFile == null ? "UE4" : projectFile.GetFileNameWithoutExtension());
+				scope.Span.SetTag("platform", _parameters.Platform);
+				string[] maps = (_parameters.Maps == null) ? null : _parameters.Maps.Split(new char[] { '+' });
+				string arguments = (_parameters.Versioned ? "" : "-Unversioned ") + "-LogCmds=\"LogSavePackage Warning\" " + _parameters.Arguments;
+				string editorExe = (string.IsNullOrWhiteSpace(_parameters.EditorExe) ? ProjectUtils.GetEditorForProject(projectFile).FullName : _parameters.EditorExe);
+				CommandUtils.CookCommandlet(projectFile, editorExe, maps, null, null, null, _parameters.Platform, arguments);
 			}
 
 			// Find all the cooked files
-			List<FileReference> CookedFiles = new List<FileReference>();
-			if (Parameters.TagOutput)
+			List<FileReference> cookedFiles = new List<FileReference>();
+			if (_parameters.TagOutput)
 			{
-				foreach (string Platform in Parameters.Platform.Split('+'))
+				foreach (string platform in _parameters.Platform.Split('+'))
 				{
-					DirectoryReference PlatformCookedDirectory = DirectoryReference.Combine(ProjectFile.Directory, "Saved", "Cooked", Platform);
-					if (!DirectoryReference.Exists(PlatformCookedDirectory))
+					DirectoryReference platformCookedDirectory = DirectoryReference.Combine(projectFile.Directory, "Saved", "Cooked", platform);
+					if (!DirectoryReference.Exists(platformCookedDirectory))
 					{
-						throw new AutomationException("Cook output directory not found ({0})", PlatformCookedDirectory.FullName);
+						throw new AutomationException("Cook output directory not found ({0})", platformCookedDirectory.FullName);
 					}
-					List<FileReference> PlatformCookedFiles = DirectoryReference.EnumerateFiles(PlatformCookedDirectory, "*", System.IO.SearchOption.AllDirectories).ToList();
-					if (PlatformCookedFiles.Count == 0)
+					List<FileReference> platformCookedFiles = DirectoryReference.EnumerateFiles(platformCookedDirectory, "*", System.IO.SearchOption.AllDirectories).ToList();
+					if (platformCookedFiles.Count == 0)
 					{
-						throw new AutomationException("Cooking did not produce any files in {0}", PlatformCookedDirectory.FullName);
+						throw new AutomationException("Cooking did not produce any files in {0}", platformCookedDirectory.FullName);
 					}
-					CookedFiles.AddRange(PlatformCookedFiles);
+					cookedFiles.AddRange(platformCookedFiles);
 
-					DirectoryReference PackagingFilesDirectory = DirectoryReference.Combine(ProjectFile.Directory, "Saved", "TmpPackaging", Platform);
-					if (DirectoryReference.Exists(PackagingFilesDirectory))
+					DirectoryReference packagingFilesDirectory = DirectoryReference.Combine(projectFile.Directory, "Saved", "TmpPackaging", platform);
+					if (DirectoryReference.Exists(packagingFilesDirectory))
 					{
-						List<FileReference> PackagingFiles = DirectoryReference.EnumerateFiles(PackagingFilesDirectory, "*", System.IO.SearchOption.AllDirectories).ToList();
-						CookedFiles.AddRange(PackagingFiles);
+						List<FileReference> packagingFiles = DirectoryReference.EnumerateFiles(packagingFilesDirectory, "*", System.IO.SearchOption.AllDirectories).ToList();
+						cookedFiles.AddRange(packagingFiles);
 					}
 				}
 			}
 
 			// Apply the optional tag to the build products
-			foreach(string TagName in FindTagNamesFromList(Parameters.Tag))
+			foreach (string tagName in FindTagNamesFromList(_parameters.Tag))
 			{
-				FindOrAddTagSet(TagNameToFileSet, TagName).UnionWith(CookedFiles);
+				FindOrAddTagSet(tagNameToFileSet, tagName).UnionWith(cookedFiles);
 			}
 
 			// Add them to the set of build products
-			BuildProducts.UnionWith(CookedFiles);
+			buildProducts.UnionWith(cookedFiles);
 			return Task.CompletedTask;
 		}
 
 		/// <summary>
 		/// Output this task out to an XML writer.
 		/// </summary>
-		public override void Write(XmlWriter Writer)
+		public override void Write(XmlWriter writer)
 		{
-			Write(Writer, Parameters);
+			Write(writer, _parameters);
 		}
 
 		/// <summary>
@@ -177,7 +175,7 @@ namespace AutomationTool.Tasks
 		/// <returns>The tag names which are modified by this task</returns>
 		public override IEnumerable<string> FindProducedTagNames()
 		{
-			return FindTagNamesFromList(Parameters.Tag);
+			return FindTagNamesFromList(_parameters.Tag);
 		}
 	}
 }
\ No newline at end of file
diff --git a/Engine/Source/Programs/AutomationTool/BuildGraph/Tasks/CopyTask.cs b/Engine/Source/Programs/AutomationTool/BuildGraph/Tasks/CopyTask.cs
--- a/Engine/Source/Programs/AutomationTool/BuildGraph/Tasks/CopyTask.cs
+++ b/Engine/Source/Programs/AutomationTool/BuildGraph/Tasks/CopyTask.cs
@@ -1,20 +1,13 @@
 // Copyright Epic Games, Inc. All Rights Reserved.
 
-using AutomationTool;
-using EpicGames.BuildGraph;
 using System;
 using System.Collections.Generic;
 using System.Linq;
-using System.Text;
-using System.Threading;
 using System.Threading.Tasks;
 using System.Xml;
 using EpicGames.Core;
-using UnrealBuildTool;
-using UnrealBuildBase;
 using Microsoft.Extensions.Logging;
-
-using static AutomationTool.CommandUtils;
+using UnrealBuildBase;
 
 namespace AutomationTool.Tasks
 {
@@ -66,151 +59,147 @@ namespace AutomationTool.Tasks
 	[TaskElement("Copy", typeof(CopyTaskParameters))]
 	public class CopyTask : BgTaskImpl
 	{
-		/// <summary>
-		/// Parameters for this task
-		/// </summary>
-		CopyTaskParameters Parameters;
+		readonly CopyTaskParameters _parameters;
 
 		/// <summary>
 		/// Constructor
 		/// </summary>
-		/// <param name="InParameters">Parameters for this task</param>
-		public CopyTask(CopyTaskParameters InParameters)
+		/// <param name="parameters">Parameters for this task</param>
+		public CopyTask(CopyTaskParameters parameters)
 		{
-			Parameters = InParameters;
+			_parameters = parameters;
 		}
 
 		/// <summary>
-		/// Execute the task.
+		/// ExecuteAsync the task.
 		/// </summary>
-		/// <param name="Job">Information about the current job</param>
-		/// <param name="BuildProducts">Set of build products produced by this node.</param>
-		/// <param name="TagNameToFileSet">Mapping from tag names to the set of files they include</param>
-		public override async Task ExecuteAsync(JobContext Job, HashSet<FileReference> BuildProducts, Dictionary<string, HashSet<FileReference>> TagNameToFileSet)
+		/// <param name="job">Information about the current job</param>
+		/// <param name="buildProducts">Set of build products produced by this node.</param>
+		/// <param name="tagNameToFileSet">Mapping from tag names to the set of files they include</param>
+		public override async Task ExecuteAsync(JobContext job, HashSet<FileReference> buildProducts, Dictionary<string, HashSet<FileReference>> tagNameToFileSet)
 		{
 			// Parse all the source patterns
-			FilePattern SourcePattern = new FilePattern(Unreal.RootDirectory, Parameters.From);
+			FilePattern sourcePattern = new FilePattern(Unreal.RootDirectory, _parameters.From);
 
 			// Parse the target pattern
-			FilePattern TargetPattern = new FilePattern(Unreal.RootDirectory, Parameters.To);
+			FilePattern targetPattern = new FilePattern(Unreal.RootDirectory, _parameters.To);
 
 			// Apply the filter to the source files
-			HashSet<FileReference> Files = null;
-			if (!String.IsNullOrEmpty(Parameters.Files))
+			HashSet<FileReference> files = null;
+			if (!String.IsNullOrEmpty(_parameters.Files))
 			{
-				SourcePattern = SourcePattern.AsDirectoryPattern();
-				Files = ResolveFilespec(SourcePattern.BaseDirectory, Parameters.Files, TagNameToFileSet);
+				sourcePattern = sourcePattern.AsDirectoryPattern();
+				files = ResolveFilespec(sourcePattern.BaseDirectory, _parameters.Files, tagNameToFileSet);
 			}
 
 			try
 			{
 				// Build the file mapping
-				Dictionary<FileReference, FileReference> TargetFileToSourceFile = FilePattern.CreateMapping(Files, ref SourcePattern, ref TargetPattern);
+				Dictionary<FileReference, FileReference> targetFileToSourceFile = FilePattern.CreateMapping(files, ref sourcePattern, ref targetPattern);
 
 				// Check we got some files
-				if (TargetFileToSourceFile.Count == 0)
+				if (targetFileToSourceFile.Count == 0)
 				{
-					if (Parameters.ErrorIfNotFound)
+					if (_parameters.ErrorIfNotFound)
 					{
-						Logger.LogError("No files found matching '{SourcePattern}'", SourcePattern);
+						Logger.LogError("No files found matching '{SourcePattern}'", sourcePattern);
 					}
 					else
 					{
-						Logger.LogInformation("No files found matching '{SourcePattern}'", SourcePattern);
+						Logger.LogInformation("No files found matching '{SourcePattern}'", sourcePattern);
 					}
 					return;
 				}
 
 				// Run the copy
-				Logger.LogInformation("Copying {Arg0} file{Arg1} from {Arg2} to {Arg3}...", TargetFileToSourceFile.Count, (TargetFileToSourceFile.Count == 1) ? "" : "s", SourcePattern.BaseDirectory, TargetPattern.BaseDirectory);
-				await ExecuteAsync(TargetFileToSourceFile, Parameters.Overwrite);
+				Logger.LogInformation("Copying {Arg0} file{Arg1} from {Arg2} to {Arg3}...", targetFileToSourceFile.Count, (targetFileToSourceFile.Count == 1) ? "" : "s", sourcePattern.BaseDirectory, targetPattern.BaseDirectory);
+				await ExecuteAsync(targetFileToSourceFile, _parameters.Overwrite);
 
 				// Update the list of build products
-				BuildProducts.UnionWith(TargetFileToSourceFile.Keys);
+				buildProducts.UnionWith(targetFileToSourceFile.Keys);
 
 				// Apply the optional output tag to them
-				foreach (string TagName in FindTagNamesFromList(Parameters.Tag))
+				foreach (string tagName in FindTagNamesFromList(_parameters.Tag))
 				{
-					FindOrAddTagSet(TagNameToFileSet, TagName).UnionWith(TargetFileToSourceFile.Keys);
+					FindOrAddTagSet(tagNameToFileSet, tagName).UnionWith(targetFileToSourceFile.Keys);
 				}
 			}
-			catch (FilePatternSourceFileMissingException Ex)
+			catch (FilePatternSourceFileMissingException ex)
 			{
-				if (Parameters.ErrorIfNotFound)
+				if (_parameters.ErrorIfNotFound)
 				{
-					Logger.LogError(Ex, "Error while trying to create file pattern match for '{SourcePattern}', error {ExceptionString}", SourcePattern, Ex.Message);
+					Logger.LogError(ex, "Error while trying to create file pattern match for '{SourcePattern}', error {ExceptionString}", sourcePattern, ex.Message);
 				}
 				else
 				{
-					Logger.LogInformation(Ex, "Error while trying to create file pattern match for '{SourcePattern}', error {ExceptionString}", SourcePattern, Ex.Message);
+					Logger.LogInformation(ex, "Error while trying to create file pattern match for '{SourcePattern}', error {ExceptionString}", sourcePattern, ex.Message);
 				}
 			}
-
 		}
 
 		/// <summary>
 		/// 
 		/// </summary>
-		/// <param name="TargetFileToSourceFile"></param>
-		/// <param name="Overwrite"></param>
+		/// <param name="targetFileToSourceFile"></param>
+		/// <param name="overwrite"></param>
 		/// <returns></returns>
-		public static async Task ExecuteAsync(Dictionary<FileReference, FileReference> TargetFileToSourceFile, bool Overwrite)
+		public static async Task ExecuteAsync(Dictionary<FileReference, FileReference> targetFileToSourceFile, bool overwrite)
 		{
 			//  If we're not overwriting, remove any files where the destination file already exists.
-			if (!Overwrite)
+			if (!overwrite)
 			{
-				Dictionary<FileReference, FileReference> FilteredTargetToSourceFile = new Dictionary<FileReference, FileReference>();
-				foreach (KeyValuePair<FileReference, FileReference> File in TargetFileToSourceFile)
+				Dictionary<FileReference, FileReference> filteredTargetToSourceFile = new Dictionary<FileReference, FileReference>();
+				foreach (KeyValuePair<FileReference, FileReference> file in targetFileToSourceFile)
 				{
-					if (FileReference.Exists(File.Key))
+					if (FileReference.Exists(file.Key))
 					{
-						Logger.LogInformation("Not copying existing file {Arg0}", File.Key);
+						Logger.LogInformation("Not copying existing file {Arg0}", file.Key);
 						continue;
 					}
-					FilteredTargetToSourceFile.Add(File.Key, File.Value);
+					filteredTargetToSourceFile.Add(file.Key, file.Value);
 				}
-				if(FilteredTargetToSourceFile.Count == 0)
+				if (filteredTargetToSourceFile.Count == 0)
 				{
 					Logger.LogWarning("All files already exist, exiting early.");
 					return;
 				}
-				TargetFileToSourceFile = FilteredTargetToSourceFile;
+				targetFileToSourceFile = filteredTargetToSourceFile;
 			}
 
 			// If the target is on a network share, retry creating the first directory until it succeeds
-			DirectoryReference FirstTargetDirectory = TargetFileToSourceFile.First().Key.Directory;
-			if(!DirectoryReference.Exists(FirstTargetDirectory))
+			DirectoryReference firstTargetDirectory = targetFileToSourceFile.First().Key.Directory;
+			if (!DirectoryReference.Exists(firstTargetDirectory))
 			{
 				const int MaxNumRetries = 15;
-				for(int NumRetries = 0;;NumRetries++)
+				for (int numRetries = 0; ; numRetries++)
 				{
 					try
 					{
-						DirectoryReference.CreateDirectory(FirstTargetDirectory);
-						if(NumRetries == 1)
+						DirectoryReference.CreateDirectory(firstTargetDirectory);
+						if (numRetries == 1)
 						{
-							Logger.LogInformation("Created target directory {FirstTargetDirectory} after 1 retry.", FirstTargetDirectory);
+							Logger.LogInformation("Created target directory {FirstTargetDirectory} after 1 retry.", firstTargetDirectory);
 						}
-						else if(NumRetries > 1)
+						else if (numRetries > 1)
 						{
-							Logger.LogInformation("Created target directory {FirstTargetDirectory} after {NumRetries} retries.", FirstTargetDirectory, NumRetries);
+							Logger.LogInformation("Created target directory {FirstTargetDirectory} after {NumRetries} retries.", firstTargetDirectory, numRetries);
 						}
 						break;
 					}
-					catch(Exception Ex)
+					catch (Exception ex)
 					{
 #pragma warning disable CA1508 // False positive about NumRetries alwayws being zero
-						if(NumRetries == 0)
+						if (numRetries == 0)
 						{
-							Logger.LogInformation("Unable to create directory '{FirstTargetDirectory}' on first attempt. Retrying {MaxNumRetries} times...", FirstTargetDirectory, MaxNumRetries);
+							Logger.LogInformation("Unable to create directory '{FirstTargetDirectory}' on first attempt. Retrying {MaxNumRetries} times...", firstTargetDirectory, MaxNumRetries);
 						}
 #pragma warning restore CA1508
 
-						Logger.LogDebug("  {Ex}", Ex);
+						Logger.LogDebug("  {Ex}", ex);
 
-						if(NumRetries >= 15)
+						if (numRetries >= 15)
 						{
-							throw new AutomationException(Ex, "Unable to create target directory '{0}' after {1} retries.", FirstTargetDirectory, NumRetries);
+							throw new AutomationException(ex, "Unable to create target directory '{0}' after {1} retries.", firstTargetDirectory, numRetries);
 						}
 
 						await Task.Delay(2000);
@@ -219,20 +208,20 @@ namespace AutomationTool.Tasks
 			}
 
 			// Copy them all
-			KeyValuePair<FileReference, FileReference>[] FilePairs = TargetFileToSourceFile.ToArray();
-			foreach(KeyValuePair<FileReference, FileReference> FilePair in FilePairs)
+			KeyValuePair<FileReference, FileReference>[] filePairs = targetFileToSourceFile.ToArray();
+			foreach (KeyValuePair<FileReference, FileReference> filePair in filePairs)
 			{
-				Logger.LogDebug("  {Arg0} -> {Arg1}", FilePair.Value, FilePair.Key);
+				Logger.LogDebug("  {Arg0} -> {Arg1}", filePair.Value, filePair.Key);
 			}
-			CommandUtils.ThreadedCopyFiles(FilePairs.Select(x => x.Value.FullName).ToList(), FilePairs.Select(x => x.Key.FullName).ToList(), bQuiet: true, bRetry: true);
+			CommandUtils.ThreadedCopyFiles(filePairs.Select(x => x.Value.FullName).ToList(), filePairs.Select(x => x.Key.FullName).ToList(), bQuiet: true, bRetry: true);
 		}
 
 		/// <summary>
 		/// Output this task out to an XML writer.
 		/// </summary>
-		public override void Write(XmlWriter Writer)
+		public override void Write(XmlWriter writer)
 		{
-			Write(Writer, Parameters);
+			Write(writer, _parameters);
 		}
 
 		/// <summary>
@@ -241,9 +230,9 @@ namespace AutomationTool.Tasks
 		/// <returns>The tag names which are read by this task</returns>
 		public override IEnumerable<string> FindConsumedTagNames()
 		{
-			foreach(string TagName in FindTagNamesFromFilespec(Parameters.Files))
+			foreach (string tagName in FindTagNamesFromFilespec(_parameters.Files))
 			{
-				yield return TagName;
+				yield return tagName;
 			}
 		}
 
@@ -253,7 +242,7 @@ namespace AutomationTool.Tasks
 		/// <returns>The tag names which are modified by this task</returns>
 		public override IEnumerable<string> FindProducedTagNames()
 		{
-			return FindTagNamesFromList(Parameters.Tag);
+			return FindTagNamesFromList(_parameters.Tag);
 		}
 	}
 
@@ -265,21 +254,21 @@ namespace AutomationTool.Tasks
 		/// <summary>
 		/// Copy files from one location to another
 		/// </summary>
-		/// <param name="Files">The files to copy</param>
-		/// <param name="TargetDir"></param>
-		/// <param name="Overwrite">Whether or not to overwrite existing files.</param>
-		public static async Task<FileSet> CopyToAsync(this FileSet Files, DirectoryReference TargetDir, bool? Overwrite = null)
+		/// <param name="files">The files to copy</param>
+		/// <param name="targetDir"></param>
+		/// <param name="overwrite">Whether or not to overwrite existing files.</param>
+		public static async Task<FileSet> CopyToAsync(this FileSet files, DirectoryReference targetDir, bool? overwrite = null)
 		{
 			// Run the copy
-			Dictionary<FileReference, FileReference> TargetFileToSourceFile = Files.Flatten(TargetDir);
-			if (TargetFileToSourceFile.Count == 0)
+			Dictionary<FileReference, FileReference> targetFileToSourceFile = files.Flatten(targetDir);
+			if (targetFileToSourceFile.Count == 0)
 			{
 				return FileSet.Empty;
 			}
 
-			Log.Logger.LogInformation("Copying {NumFiles} file(s) to {TargetDir}...", TargetFileToSourceFile.Count, TargetDir);
-			await CopyTask.ExecuteAsync(TargetFileToSourceFile, Overwrite ?? true);
-			return FileSet.FromFiles(TargetFileToSourceFile.Keys.Select(x => (x.MakeRelativeTo(TargetDir), x)));
+			Log.Logger.LogInformation("Copying {NumFiles} file(s) to {TargetDir}...", targetFileToSourceFile.Count, targetDir);
+			await CopyTask.ExecuteAsync(targetFileToSourceFile, overwrite ?? true);
+			return FileSet.FromFiles(targetFileToSourceFile.Keys.Select(x => (x.MakeRelativeTo(targetDir), x)));
 		}
 	}
 }
diff --git a/Engine/Source/Programs/AutomationTool/BuildGraph/Tasks/CreateArtifactTask.cs b/Engine/Source/Programs/AutomationTool/BuildGraph/Tasks/CreateArtifactTask.cs
--- a/Engine/Source/Programs/AutomationTool/BuildGraph/Tasks/CreateArtifactTask.cs
+++ b/Engine/Source/Programs/AutomationTool/BuildGraph/Tasks/CreateArtifactTask.cs
@@ -1,30 +1,22 @@
 // Copyright Epic Games, Inc. All Rights Reserved.
 
-using AutomationTool;
-using EpicGames.Core;
 using System;
 using System.Collections.Generic;
-using System.Text.Json;
+using System.Data;
+using System.Diagnostics;
 using System.IO;
+using System.Linq;
 using System.Net.Http;
-using System.Net.Http.Headers;
-using System.Text;
 using System.Threading.Tasks;
 using System.Xml;
-using UnrealBuildBase;
-using Microsoft.Extensions.Logging;
-using EpicGames.Horde.Storage.Bundles;
-using EpicGames.Horde.Storage.Clients;
+using EpicGames.Core;
+using EpicGames.Horde;
+using EpicGames.Horde.Artifacts;
 using EpicGames.Horde.Storage;
+using EpicGames.Horde.Storage.Clients;
 using EpicGames.Horde.Storage.Nodes;
-using System.Threading;
-using System.Data;
-using EpicGames.Horde.Storage.Backends;
 using Microsoft.Extensions.DependencyInjection;
-using EpicGames.Horde;
-using EpicGames.Horde.Artifacts;
-using System.Linq;
-using System.Diagnostics;
+using Microsoft.Extensions.Logging;
 
 #nullable enable
 
@@ -87,27 +79,24 @@ namespace AutomationTool.Tasks
 			public void Dispose() { }
 		}
 
-		/// <summary>
-		/// Parameters for this task
-		/// </summary>
-		CreateArtifactTaskParameters Parameters;
+		readonly CreateArtifactTaskParameters _parameters;
 
 		/// <summary>
 		/// Construct a Helm task
 		/// </summary>
-		/// <param name="InParameters">Parameters for the task</param>
-		public CreateArtifactTask(CreateArtifactTaskParameters InParameters)
+		/// <param name="parameters">Parameters for the task</param>
+		public CreateArtifactTask(CreateArtifactTaskParameters parameters)
 		{
-			Parameters = InParameters;
+			_parameters = parameters;
 		}
 
 		/// <summary>
-		/// Execute the task.
+		/// ExecuteAsync the task.
 		/// </summary>
-		/// <param name="Job">Information about the current job</param>
-		/// <param name="BuildProducts">Set of build products produced by this node.</param>
-		/// <param name="TagNameToFileSet">Mapping from tag names to the set of files they include</param>
-		public override async Task ExecuteAsync(JobContext Job, HashSet<FileReference> BuildProducts, Dictionary<string, HashSet<FileReference>> TagNameToFileSet)
+		/// <param name="job">Information about the current job</param>
+		/// <param name="buildProducts">Set of build products produced by this node.</param>
+		/// <param name="tagNameToFileSet">Mapping from tag names to the set of files they include</param>
+		public override async Task ExecuteAsync(JobContext job, HashSet<FileReference> buildProducts, Dictionary<string, HashSet<FileReference>> tagNameToFileSet)
 		{
 			// Create a DI container for building the graph
 			ServiceCollection serviceCollection = new ServiceCollection();
@@ -118,12 +107,12 @@ namespace AutomationTool.Tasks
 
 			await using ServiceProvider serviceProvider = serviceCollection.BuildServiceProvider();
 
-			ArtifactName artifactName = new ArtifactName(Parameters.Name);
-			ArtifactType artifactType = new ArtifactType(Parameters.Type);
+			ArtifactName artifactName = new ArtifactName(_parameters.Name);
+			ArtifactType artifactType = new ArtifactType(_parameters.Type);
 
 			HordeHttpClient hordeHttpClient = serviceProvider.GetRequiredService<HordeHttpClient>();
-			int? change = (Parameters.Change == 0) ? (int?)null : Parameters.Change;
-			CreateArtifactResponse response = await hordeHttpClient.CreateArtifactAsync(artifactName, artifactType, Parameters.Description, change: change);
+			int? change = (_parameters.Change == 0) ? (int?)null : _parameters.Change;
+			CreateArtifactResponse response = await hordeHttpClient.CreateArtifactAsync(artifactName, artifactType, _parameters.Description, change: change);
 			Logger.LogInformation("Creating artifact {ArtifactId} '{ArtifactName}' ({ArtifactType}) (ns: {NamespaceId}, ref: {RefName})", response.ArtifactId, artifactName, artifactType, response.NamespaceId, response.RefName);
 
 			Stopwatch timer = Stopwatch.StartNew();
@@ -133,8 +122,8 @@ namespace AutomationTool.Tasks
 			{
 				await using (IBlobWriter writer = client.CreateBlobWriter(response.RefName))
 				{
-					DirectoryReference baseDir = ResolveDirectory(Parameters.BaseDir);
-					List<FileInfo> files = ResolveFilespec(baseDir, Parameters.Files, TagNameToFileSet).Select(x => x.ToFileInfo()).ToList();
+					DirectoryReference baseDir = ResolveDirectory(_parameters.BaseDir);
+					List<FileInfo> files = ResolveFilespec(baseDir, _parameters.Files, tagNameToFileSet).Select(x => x.ToFileInfo()).ToList();
 
 					int totalCount = files.Count;
 					long totalSize = files.Sum(x => x.Length);
@@ -152,9 +141,9 @@ namespace AutomationTool.Tasks
 		/// <summary>
 		/// Output this task out to an XML writer.
 		/// </summary>
-		public override void Write(XmlWriter Writer)
+		public override void Write(XmlWriter writer)
 		{
-			Write(Writer, Parameters);
+			Write(writer, _parameters);
 		}
 
 		/// <summary>
@@ -163,7 +152,7 @@ namespace AutomationTool.Tasks
 		/// <returns>The tag names which are read by this task</returns>
 		public override IEnumerable<string> FindConsumedTagNames()
 		{
-			return FindTagNamesFromList(Parameters.Files);
+			return FindTagNamesFromList(_parameters.Files);
 		}
 
 		/// <summary>
diff --git a/Engine/Source/Programs/AutomationTool/BuildGraph/Tasks/CsCompileTask.cs b/Engine/Source/Programs/AutomationTool/BuildGraph/Tasks/CsCompileTask.cs
--- a/Engine/Source/Programs/AutomationTool/BuildGraph/Tasks/CsCompileTask.cs
+++ b/Engine/Source/Programs/AutomationTool/BuildGraph/Tasks/CsCompileTask.cs
@@ -1,16 +1,12 @@
 ﻿// Copyright Epic Games, Inc. All Rights Reserved.
 
-using EpicGames.BuildGraph;
 using System;
 using System.Collections.Generic;
-using System.Linq;
-using System.Text;
 using System.Threading.Tasks;
 using System.Xml;
 using EpicGames.Core;
-using UnrealBuildBase;
-using UnrealBuildTool;
 using Microsoft.Extensions.Logging;
+using UnrealBuildBase;
 
 namespace AutomationTool.Tasks
 {
@@ -87,144 +83,141 @@ namespace AutomationTool.Tasks
 	[TaskElement("CsCompile", typeof(CsCompileTaskParameters))]
 	public class CsCompileTask : BgTaskImpl
 	{
-		/// <summary>
-		/// Parameters for the task
-		/// </summary>
-		CsCompileTaskParameters Parameters;
+		readonly CsCompileTaskParameters _parameters;
 
 		/// <summary>
 		/// Constructor.
 		/// </summary>
-		/// <param name="InParameters">Parameters for this task</param>
-		public CsCompileTask(CsCompileTaskParameters InParameters)
+		/// <param name="parameters">Parameters for this task</param>
+		public CsCompileTask(CsCompileTaskParameters parameters)
 		{
-			Parameters = InParameters;
+			_parameters = parameters;
 		}
 
 		/// <summary>
-		/// Execute the task.
+		/// ExecuteAsync the task.
 		/// </summary>
-		/// <param name="Job">Information about the current job</param>
-		/// <param name="BuildProducts">Set of build products produced by this node.</param>
-		/// <param name="TagNameToFileSet">Mapping from tag names to the set of files they include</param>
-		public override Task ExecuteAsync(JobContext Job, HashSet<FileReference> BuildProducts, Dictionary<string, HashSet<FileReference>> TagNameToFileSet)
+		/// <param name="job">Information about the current job</param>
+		/// <param name="buildProducts">Set of build products produced by this node.</param>
+		/// <param name="tagNameToFileSet">Mapping from tag names to the set of files they include</param>
+		public override Task ExecuteAsync(JobContext job, HashSet<FileReference> buildProducts, Dictionary<string, HashSet<FileReference>> tagNameToFileSet)
 		{
 			// Get the project file
-			HashSet<FileReference> ProjectFiles = ResolveFilespec(Unreal.RootDirectory, Parameters.Project, TagNameToFileSet);
-			foreach(FileReference ProjectFile in ProjectFiles)
+			HashSet<FileReference> projectFiles = ResolveFilespec(Unreal.RootDirectory, _parameters.Project, tagNameToFileSet);
+			foreach (FileReference projectFile in projectFiles)
 			{
-				if(!FileReference.Exists(ProjectFile))
+				if (!FileReference.Exists(projectFile))
 				{
-					throw new AutomationException("Couldn't find project file '{0}'", ProjectFile.FullName);
+					throw new AutomationException("Couldn't find project file '{0}'", projectFile.FullName);
 				}
-				if(!ProjectFile.HasExtension(".csproj"))
+				if (!projectFile.HasExtension(".csproj"))
 				{
-					throw new AutomationException("File '{0}' is not a C# project", ProjectFile.FullName);
+					throw new AutomationException("File '{0}' is not a C# project", projectFile.FullName);
 				}
 			}
 
 			// Get the default properties
-			Dictionary<string, string> Properties = new Dictionary<string, string>(StringComparer.InvariantCultureIgnoreCase);
-			if(!String.IsNullOrEmpty(Parameters.Platform))
+			Dictionary<string, string> properties = new Dictionary<string, string>(StringComparer.InvariantCultureIgnoreCase);
+			if (!String.IsNullOrEmpty(_parameters.Platform))
 			{
-				Properties["Platform"] = Parameters.Platform;
+				properties["Platform"] = _parameters.Platform;
 			}
-			if(!String.IsNullOrEmpty(Parameters.Configuration))
+			if (!String.IsNullOrEmpty(_parameters.Configuration))
 			{
-				Properties["Configuration"] = Parameters.Configuration;
+				properties["Configuration"] = _parameters.Configuration;
 			}
-			if (!String.IsNullOrEmpty(Parameters.Properties))
+			if (!String.IsNullOrEmpty(_parameters.Properties))
 			{
-				foreach (string Property in Parameters.Properties.Split(';'))
+				foreach (string property in _parameters.Properties.Split(';'))
 				{
-					if (!String.IsNullOrWhiteSpace(Property))
+					if (!String.IsNullOrWhiteSpace(property))
 					{
-						int EqualsIdx = Property.IndexOf('=', StringComparison.Ordinal);
-						if (EqualsIdx == -1)
+						int equalsIdx = property.IndexOf('=', StringComparison.Ordinal);
+						if (equalsIdx == -1)
 						{
 							Logger.LogWarning("Missing '=' in property assignment");
 						}
 						else
 						{
-							Properties[Property.Substring(0, EqualsIdx).Trim()] = Property.Substring(EqualsIdx + 1).Trim();
+							properties[property.Substring(0, equalsIdx).Trim()] = property.Substring(equalsIdx + 1).Trim();
 						}
 					}
 				}
 			}
 
 			// Build the arguments and run the build
-			if(!Parameters.EnumerateOnly)
+			if (!_parameters.EnumerateOnly)
 			{
-				List<string> Arguments = new List<string>();
-				foreach(KeyValuePair<string, string> PropertyPair in Properties)
+				List<string> arguments = new List<string>();
+				foreach (KeyValuePair<string, string> propertyPair in properties)
 				{
-					Arguments.Add(String.Format("/property:{0}={1}", CommandUtils.MakePathSafeToUseWithCommandLine(PropertyPair.Key), CommandUtils.MakePathSafeToUseWithCommandLine(PropertyPair.Value)));
+					arguments.Add(String.Format("/property:{0}={1}", CommandUtils.MakePathSafeToUseWithCommandLine(propertyPair.Key), CommandUtils.MakePathSafeToUseWithCommandLine(propertyPair.Value)));
 				}
-				if(!String.IsNullOrEmpty(Parameters.Arguments))
+				if (!String.IsNullOrEmpty(_parameters.Arguments))
 				{
-					Arguments.Add(Parameters.Arguments);
+					arguments.Add(_parameters.Arguments);
 				}
-				if(!String.IsNullOrEmpty(Parameters.Target))
+				if (!String.IsNullOrEmpty(_parameters.Target))
 				{
-					Arguments.Add(String.Format("/target:{0}", CommandUtils.MakePathSafeToUseWithCommandLine(Parameters.Target)));
+					arguments.Add(String.Format("/target:{0}", CommandUtils.MakePathSafeToUseWithCommandLine(_parameters.Target)));
 				}
 
-				Arguments.Add("/restore");
-				Arguments.Add("/verbosity:minimal");
-				Arguments.Add("/nologo");
+				arguments.Add("/restore");
+				arguments.Add("/verbosity:minimal");
+				arguments.Add("/nologo");
 
-				string JoinedArguments = String.Join(" ", Arguments);
+				string joinedArguments = String.Join(" ", arguments);
 
-				foreach(FileReference ProjectFile in ProjectFiles)
+				foreach (FileReference projectFile in projectFiles)
 				{
-					if (!FileReference.Exists(ProjectFile))
+					if (!FileReference.Exists(projectFile))
 					{
-						throw new AutomationException("Project {0} does not exist!", ProjectFile);
+						throw new AutomationException("Project {0} does not exist!", projectFile);
 					}
 
-					if (Parameters.UseSystemCompiler)
+					if (_parameters.UseSystemCompiler)
 					{
-						CommandUtils.MsBuild(CommandUtils.CmdEnv, ProjectFile.FullName, JoinedArguments, null);
+						CommandUtils.MsBuild(CommandUtils.CmdEnv, projectFile.FullName, joinedArguments, null);
 					}
 					else
 					{
-						CommandUtils.RunAndLog(CommandUtils.CmdEnv, CommandUtils.CmdEnv.DotnetMsbuildPath, $"msbuild {CommandUtils.MakePathSafeToUseWithCommandLine(ProjectFile.FullName)} {JoinedArguments}");
+						CommandUtils.RunAndLog(CommandUtils.CmdEnv, CommandUtils.CmdEnv.DotnetMsbuildPath, $"msbuild {CommandUtils.MakePathSafeToUseWithCommandLine(projectFile.FullName)} {joinedArguments}");
 					}
 				}
 			}
 
 			// Try to figure out the output files
-			HashSet<FileReference> ProjectBuildProducts;
-			HashSet<FileReference> ProjectReferences;
-			FindBuildProductsAndReferences(ProjectFiles, Properties, out ProjectBuildProducts, out ProjectReferences);
+			HashSet<FileReference> projectBuildProducts;
+			HashSet<FileReference> projectReferences;
+			FindBuildProductsAndReferences(projectFiles, properties, out projectBuildProducts, out projectReferences);
 
 			// Apply the optional tag to the produced archive
-			foreach(string TagName in FindTagNamesFromList(Parameters.Tag))
+			foreach (string tagName in FindTagNamesFromList(_parameters.Tag))
 			{
-				FindOrAddTagSet(TagNameToFileSet, TagName).UnionWith(ProjectBuildProducts);
+				FindOrAddTagSet(tagNameToFileSet, tagName).UnionWith(projectBuildProducts);
 			}
 
 			// Apply the optional tag to any references
-			if (!String.IsNullOrEmpty(Parameters.TagReferences))
+			if (!String.IsNullOrEmpty(_parameters.TagReferences))
 			{
-				foreach (string TagName in FindTagNamesFromList(Parameters.TagReferences))
+				foreach (string tagName in FindTagNamesFromList(_parameters.TagReferences))
 				{
-					FindOrAddTagSet(TagNameToFileSet, TagName).UnionWith(ProjectReferences);
+					FindOrAddTagSet(tagNameToFileSet, tagName).UnionWith(projectReferences);
 				}
 			}
 
 			// Merge them into the standard set of build products
-			BuildProducts.UnionWith(ProjectBuildProducts);
-			BuildProducts.UnionWith(ProjectReferences);
+			buildProducts.UnionWith(projectBuildProducts);
+			buildProducts.UnionWith(projectReferences);
 			return Task.CompletedTask;
 		}
 
 		/// <summary>
 		/// Output this task out to an XML writer.
 		/// </summary>
-		public override void Write(XmlWriter Writer)
+		public override void Write(XmlWriter writer)
 		{
-			Write(Writer, Parameters);
+			Write(writer, _parameters);
 		}
 
 		/// <summary>
@@ -233,7 +226,7 @@ namespace AutomationTool.Tasks
 		/// <returns>The tag names which are read by this task</returns>
 		public override IEnumerable<string> FindConsumedTagNames()
 		{
-			return FindTagNamesFromFilespec(Parameters.Project);
+			return FindTagNamesFromFilespec(_parameters.Project);
 		}
 
 		/// <summary>
@@ -242,86 +235,86 @@ namespace AutomationTool.Tasks
 		/// <returns>The tag names which are modified by this task</returns>
 		public override IEnumerable<string> FindProducedTagNames()
 		{
-			foreach (string TagName in FindTagNamesFromList(Parameters.Tag))
+			foreach (string tagName in FindTagNamesFromList(_parameters.Tag))
 			{
-				yield return TagName;
+				yield return tagName;
 			}
 
-			foreach (string TagName in FindTagNamesFromList(Parameters.TagReferences))
+			foreach (string tagName in FindTagNamesFromList(_parameters.TagReferences))
 			{
-				yield return TagName;
+				yield return tagName;
 			}
 		}
 
 		/// <summary>
 		/// Find all the build products created by compiling the given project file
 		/// </summary>
-		/// <param name="ProjectFiles">Initial project file to read. All referenced projects will also be read.</param>
-		/// <param name="InitialProperties">Mapping of property name to value</param>
-		/// <param name="OutBuildProducts">Receives a set of build products on success</param>
-		/// <param name="OutReferences">Receives a set of non-private references on success</param>
-		static void FindBuildProductsAndReferences(HashSet<FileReference> ProjectFiles, Dictionary<string, string> InitialProperties, out HashSet<FileReference> OutBuildProducts, out HashSet<FileReference> OutReferences)
+		/// <param name="projectFiles">Initial project file to read. All referenced projects will also be read.</param>
+		/// <param name="initialProperties">Mapping of property name to value</param>
+		/// <param name="outBuildProducts">Receives a set of build products on success</param>
+		/// <param name="outReferences">Receives a set of non-private references on success</param>
+		static void FindBuildProductsAndReferences(HashSet<FileReference> projectFiles, Dictionary<string, string> initialProperties, out HashSet<FileReference> outBuildProducts, out HashSet<FileReference> outReferences)
 		{
 			// Find all the build products and references
-			OutBuildProducts = new HashSet<FileReference>();
-			OutReferences = new HashSet<FileReference>();
+			outBuildProducts = new HashSet<FileReference>();
+			outReferences = new HashSet<FileReference>();
 
 			// Read all the project information into a dictionary
-			Dictionary<FileReference, CsProjectInfo> FileToProjectInfo = new Dictionary<FileReference, CsProjectInfo>();
-			foreach(FileReference ProjectFile in ProjectFiles)
+			Dictionary<FileReference, CsProjectInfo> fileToProjectInfo = new Dictionary<FileReference, CsProjectInfo>();
+			foreach (FileReference projectFile in projectFiles)
 			{
 				// Read all the projects
-				ReadProjectsRecursively(ProjectFile, InitialProperties, FileToProjectInfo);
+				ReadProjectsRecursively(projectFile, initialProperties, fileToProjectInfo);
 
 				// Find all the outputs for each project
-				foreach(KeyValuePair<FileReference, CsProjectInfo> Pair in FileToProjectInfo)
+				foreach (KeyValuePair<FileReference, CsProjectInfo> pair in fileToProjectInfo)
 				{
-					CsProjectInfo ProjectInfo = Pair.Value;
+					CsProjectInfo projectInfo = pair.Value;
 
 					// Add all the build projects from this project
-					DirectoryReference OutputDir = ProjectInfo.GetOutputDir(Pair.Key.Directory);
-					ProjectInfo.FindBuildProducts(OutputDir, OutBuildProducts, FileToProjectInfo);
+					DirectoryReference outputDir = projectInfo.GetOutputDir(pair.Key.Directory);
+					projectInfo.FindBuildProducts(outputDir, outBuildProducts, fileToProjectInfo);
 
 					// Add any files which are only referenced
-					foreach (KeyValuePair<FileReference, bool> Reference in ProjectInfo.References)
+					foreach (KeyValuePair<FileReference, bool> reference in projectInfo.References)
 					{
-						CsProjectInfo.AddReferencedAssemblyAndSupportFiles(Reference.Key, OutReferences);
+						CsProjectInfo.AddReferencedAssemblyAndSupportFiles(reference.Key, outReferences);
 					}
 				}
 			}
 
-			OutBuildProducts.RemoveWhere(x => !FileReference.Exists(x));
-			OutReferences.RemoveWhere(x => !FileReference.Exists(x));
+			outBuildProducts.RemoveWhere(x => !FileReference.Exists(x));
+			outReferences.RemoveWhere(x => !FileReference.Exists(x));
 		}
 
 		/// <summary>
 		/// Read a project file, plus all the project files it references.
 		/// </summary>
-		/// <param name="File">Project file to read</param>
-		/// <param name="InitialProperties">Mapping of property name to value for the initial project</param>
-		/// <param name="FileToProjectInfo"></param>
+		/// <param name="file">Project file to read</param>
+		/// <param name="initialProperties">Mapping of property name to value for the initial project</param>
+		/// <param name="fileToProjectInfo"></param>
 		/// <returns>True if the projects were read correctly, false (and prints an error to the log) if not</returns>
-		static void ReadProjectsRecursively(FileReference File, Dictionary<string, string> InitialProperties, Dictionary<FileReference, CsProjectInfo> FileToProjectInfo)
+		static void ReadProjectsRecursively(FileReference file, Dictionary<string, string> initialProperties, Dictionary<FileReference, CsProjectInfo> fileToProjectInfo)
 		{
 			// Early out if we've already read this project
-			if (!FileToProjectInfo.ContainsKey(File))
+			if (!fileToProjectInfo.ContainsKey(file))
 			{
 				// Try to read this project
-				CsProjectInfo ProjectInfo;
-				if (!CsProjectInfo.TryRead(File, InitialProperties, out ProjectInfo))
+				CsProjectInfo projectInfo;
+				if (!CsProjectInfo.TryRead(file, initialProperties, out projectInfo))
 				{
-					throw new AutomationException("Couldn't read project '{0}'", File.FullName);
+					throw new AutomationException("Couldn't read project '{0}'", file.FullName);
 				}
 
 				// Add it to the project lookup, and try to read all the projects it references
-				FileToProjectInfo.Add(File, ProjectInfo);
-				foreach(FileReference ProjectReference in ProjectInfo.ProjectReferences.Keys)
+				fileToProjectInfo.Add(file, projectInfo);
+				foreach (FileReference projectReference in projectInfo.ProjectReferences.Keys)
 				{
-					if(!FileReference.Exists(ProjectReference))
+					if (!FileReference.Exists(projectReference))
 					{
-						throw new AutomationException("Unable to find project '{0}' referenced by '{1}'", ProjectReference, File);
+						throw new AutomationException("Unable to find project '{0}' referenced by '{1}'", projectReference, file);
 					}
-					ReadProjectsRecursively(ProjectReference, InitialProperties, FileToProjectInfo);
+					ReadProjectsRecursively(projectReference, initialProperties, fileToProjectInfo);
 				}
 			}
 		}
@@ -350,10 +343,10 @@ namespace AutomationTool.Tasks
 		/// <summary>
 		/// Constructor
 		/// </summary>
-		public CsCompileOutput(FileSet Binaries, FileSet References)
+		public CsCompileOutput(FileSet binaries, FileSet references)
 		{
-			this.Binaries = Binaries;
-			this.References = References;
+			this.Binaries = binaries;
+			this.References = references;
 		}
 
 		/// <summary>
@@ -368,9 +361,9 @@ namespace AutomationTool.Tasks
 		/// <summary>
 		/// Merges two outputs together
 		/// </summary>
-		public static CsCompileOutput operator +(CsCompileOutput Lhs, CsCompileOutput Rhs)
+		public static CsCompileOutput operator +(CsCompileOutput lhs, CsCompileOutput rhs)
 		{
-			return new CsCompileOutput(Lhs.Binaries + Rhs.Binaries, Lhs.References + Rhs.References);
+			return new CsCompileOutput(lhs.Binaries + rhs.Binaries, lhs.References + rhs.References);
 		}
 	}
 
@@ -382,11 +375,11 @@ namespace AutomationTool.Tasks
 		/// <summary>
 		/// 
 		/// </summary>
-		/// <param name="Task"></param>
+		/// <param name="task"></param>
 		/// <returns></returns>
-		public static async Task<FileSet> MergeAsync(this Task<CsCompileOutput> Task)
+		public static async Task<FileSet> MergeAsync(this Task<CsCompileOutput> task)
 		{
-			return (await Task).Merge();
+			return (await task).Merge();
 		}
 	}
 
@@ -395,42 +388,42 @@ namespace AutomationTool.Tasks
 		/// <summary>
 		/// Compile a C# project
 		/// </summary>
-		/// <param name="Project">The C# project files to compile.</param>
-		/// <param name="Platform">The platform to compile.</param>
-		/// <param name="Configuration">The configuration to compile.</param>
-		/// <param name="Target">The target to build.</param>
-		/// <param name="Properties">Properties for the command.</param>
-		/// <param name="Arguments">Additional options to pass to the compiler.</param>
-		/// <param name="EnumerateOnly">Only enumerate build products -- do not actually compile the projects.</param>
-		public static async Task<CsCompileOutput> CsCompileAsync(FileReference Project, string Platform = null, string Configuration = null, string Target = null, string Properties = null, string Arguments = null, bool? EnumerateOnly = null)
+		/// <param name="project">The C# project files to compile.</param>
+		/// <param name="platform">The platform to compile.</param>
+		/// <param name="configuration">The configuration to compile.</param>
+		/// <param name="target">The target to build.</param>
+		/// <param name="properties">Properties for the command.</param>
+		/// <param name="arguments">Additional options to pass to the compiler.</param>
+		/// <param name="enumerateOnly">Only enumerate build products -- do not actually compile the projects.</param>
+		public static async Task<CsCompileOutput> CsCompileAsync(FileReference project, string platform = null, string configuration = null, string target = null, string properties = null, string arguments = null, bool? enumerateOnly = null)
 		{
-			CsCompileTaskParameters Parameters = new CsCompileTaskParameters();
-			Parameters.Project = Project.FullName;
-			Parameters.Platform = Platform;
-			Parameters.Configuration = Configuration;
-			Parameters.Target = Target;
-			Parameters.Properties = Properties;
-			Parameters.Arguments = Arguments;
-			Parameters.EnumerateOnly = EnumerateOnly ?? Parameters.EnumerateOnly;
-			Parameters.Tag = "#Out";
-			Parameters.TagReferences = "#Refs";
-
-			HashSet<FileReference> BuildProducts = new HashSet<FileReference>();
-			Dictionary<string, HashSet<FileReference>> TagNameToFileSet = new Dictionary<string, HashSet<FileReference>>();
-			await new CsCompileTask(Parameters).ExecuteAsync(new JobContext(null!, null!), BuildProducts, TagNameToFileSet);
-
-			FileSet Binaries = FileSet.Empty;
-			FileSet References = FileSet.Empty;
-			if (TagNameToFileSet.TryGetValue(Parameters.Tag, out HashSet<FileReference> BinaryFiles))
+			CsCompileTaskParameters parameters = new CsCompileTaskParameters();
+			parameters.Project = project.FullName;
+			parameters.Platform = platform;
+			parameters.Configuration = configuration;
+			parameters.Target = target;
+			parameters.Properties = properties;
+			parameters.Arguments = arguments;
+			parameters.EnumerateOnly = enumerateOnly ?? parameters.EnumerateOnly;
+			parameters.Tag = "#Out";
+			parameters.TagReferences = "#Refs";
+
+			HashSet<FileReference> buildProducts = new HashSet<FileReference>();
+			Dictionary<string, HashSet<FileReference>> tagNameToFileSet = new Dictionary<string, HashSet<FileReference>>();
+			await new CsCompileTask(parameters).ExecuteAsync(new JobContext(null!, null!), buildProducts, tagNameToFileSet);
+
+			FileSet binaries = FileSet.Empty;
+			FileSet references = FileSet.Empty;
+			if (tagNameToFileSet.TryGetValue(parameters.Tag, out HashSet<FileReference> binaryFiles))
 			{
-				Binaries = FileSet.FromFiles(Unreal.RootDirectory, BinaryFiles);
+				binaries = FileSet.FromFiles(Unreal.RootDirectory, binaryFiles);
 			}
-			if (TagNameToFileSet.TryGetValue(Parameters.TagReferences, out HashSet<FileReference> ReferenceFiles))
+			if (tagNameToFileSet.TryGetValue(parameters.TagReferences, out HashSet<FileReference> referenceFiles))
 			{
-				References = FileSet.FromFiles(Unreal.RootDirectory, ReferenceFiles);
+				references = FileSet.FromFiles(Unreal.RootDirectory, referenceFiles);
 			}
 
-			return new CsCompileOutput(Binaries, References);
+			return new CsCompileOutput(binaries, references);
 		}
 	}
 }
