EpicGames.BuildGraph: Apply standard naming conventions to BuildGraph.Automation.

diff --git a/Engine/Source/Programs/AutomationTool/BuildGraph/Tasks/DeleteTask.cs b/Engine/Source/Programs/AutomationTool/BuildGraph/Tasks/DeleteTask.cs
--- a/Engine/Source/Programs/AutomationTool/BuildGraph/Tasks/DeleteTask.cs
+++ b/Engine/Source/Programs/AutomationTool/BuildGraph/Tasks/DeleteTask.cs
@@ -1,20 +1,14 @@
 // Copyright Epic Games, Inc. All Rights Reserved.
 
-using AutomationTool;
-using EpicGames.BuildGraph;
 using System;
 using System.Collections.Generic;
 using System.IO;
 using System.Linq;
-using System.Text;
 using System.Threading.Tasks;
 using System.Xml;
 using EpicGames.Core;
-using UnrealBuildTool;
-using UnrealBuildBase;
 using Microsoft.Extensions.Logging;
-
-using static AutomationTool.CommandUtils;
+using UnrealBuildBase;
 
 namespace AutomationTool.Tasks
 {
@@ -54,60 +48,57 @@ namespace AutomationTool.Tasks
 	[TaskElement("Delete", typeof(DeleteTaskParameters))]
 	public class DeleteTask : BgTaskImpl
 	{
-		/// <summary>
-		/// Parameters for this task
-		/// </summary>
-		DeleteTaskParameters Parameters;
+		readonly DeleteTaskParameters _parameters;
 
 		/// <summary>
 		/// Constructor
 		/// </summary>
-		/// <param name="InParameters">Parameters for this task</param>
-		public DeleteTask(DeleteTaskParameters InParameters)
+		/// <param name="parameters">Parameters for this task</param>
+		public DeleteTask(DeleteTaskParameters parameters)
 		{
-			Parameters = InParameters;
+			_parameters = parameters;
 		}
 
 		/// <summary>
-		/// Execute the task.
+		/// ExecuteAsync the task.
 		/// </summary>
-		/// <param name="Job">Information about the current job</param>
-		/// <param name="BuildProducts">Set of build products produced by this node.</param>
-		/// <param name="TagNameToFileSet">Mapping from tag names to the set of files they include</param>
-		public override Task ExecuteAsync(JobContext Job, HashSet<FileReference> BuildProducts, Dictionary<string, HashSet<FileReference>> TagNameToFileSet)
+		/// <param name="job">Information about the current job</param>
+		/// <param name="buildProducts">Set of build products produced by this node.</param>
+		/// <param name="tagNameToFileSet">Mapping from tag names to the set of files they include</param>
+		public override Task ExecuteAsync(JobContext job, HashSet<FileReference> buildProducts, Dictionary<string, HashSet<FileReference>> tagNameToFileSet)
 		{
-			if (Parameters.Files != null)
+			if (_parameters.Files != null)
 			{
 				// Find all the referenced files and delete them
-				HashSet<FileReference> Files = ResolveFilespec(Unreal.RootDirectory, Parameters.Files, TagNameToFileSet);
-				foreach (FileReference File in Files)
+				HashSet<FileReference> files = ResolveFilespec(Unreal.RootDirectory, _parameters.Files, tagNameToFileSet);
+				foreach (FileReference file in files)
 				{
-					if (Parameters.Verbose)
+					if (_parameters.Verbose)
 					{
-						Logger.LogInformation("Deleting {File}", File.FullName);
+						Logger.LogInformation("Deleting {File}", file.FullName);
 					}
-					if (!InternalUtils.SafeDeleteFile(File.FullName))
+					if (!InternalUtils.SafeDeleteFile(file.FullName))
 					{
-						Logger.LogWarning("Couldn't delete file {Arg0}", File.FullName);
+						Logger.LogWarning("Couldn't delete file {Arg0}", file.FullName);
 					}
 				}
 
 				// Try to delete all the parent directories. Keep track of the directories we've already deleted to avoid hitting the disk.
-				if (Parameters.DeleteEmptyDirectories)
+				if (_parameters.DeleteEmptyDirectories)
 				{
 					// Find all the directories that we're touching
-					HashSet<DirectoryReference> ParentDirectories = new HashSet<DirectoryReference>();
-					foreach (FileReference File in Files)
+					HashSet<DirectoryReference> parentDirectories = new HashSet<DirectoryReference>();
+					foreach (FileReference file in files)
 					{
-						ParentDirectories.Add(File.Directory);
+						parentDirectories.Add(file.Directory);
 					}
 
 					// Recurse back up from each of those directories to the root folder
-					foreach (DirectoryReference ParentDirectory in ParentDirectories)
+					foreach (DirectoryReference parentDirectory in parentDirectories)
 					{
-						for (DirectoryReference CurrentDirectory = ParentDirectory; CurrentDirectory != Unreal.RootDirectory; CurrentDirectory = CurrentDirectory.ParentDirectory)
+						for (DirectoryReference currentDirectory = parentDirectory; currentDirectory != Unreal.RootDirectory; currentDirectory = currentDirectory.ParentDirectory)
 						{
-							if (!TryDeleteEmptyDirectory(CurrentDirectory))
+							if (!TryDeleteEmptyDirectory(currentDirectory))
 							{
 								break;
 							}
@@ -115,20 +106,20 @@ namespace AutomationTool.Tasks
 					}
 				}
 			}
-			if (Parameters.Directories != null)
+			if (_parameters.Directories != null)
 			{
-				foreach (string Directory in Parameters.Directories.Split(';'))
+				foreach (string directory in _parameters.Directories.Split(';'))
 				{
-					if (!String.IsNullOrEmpty(Directory))
+					if (!String.IsNullOrEmpty(directory))
 					{
-						if (Parameters.Verbose)
+						if (_parameters.Verbose)
 						{
-							Logger.LogInformation("Deleting {Directory}", Directory);
+							Logger.LogInformation("Deleting {Directory}", directory);
 						}
-						DirectoryReference FullDir = new DirectoryReference(Directory);
-						if (DirectoryReference.Exists(FullDir))
+						DirectoryReference fullDir = new DirectoryReference(directory);
+						if (DirectoryReference.Exists(fullDir))
 						{
-							FileUtils.ForceDeleteDirectory(FullDir);
+							FileUtils.ForceDeleteDirectory(fullDir);
 						}
 					}
 				}
@@ -139,18 +130,18 @@ namespace AutomationTool.Tasks
 		/// <summary>
 		/// Deletes a directory, if it's empty
 		/// </summary>
-		/// <param name="CandidateDirectory">The directory to check</param>
+		/// <param name="candidateDirectory">The directory to check</param>
 		/// <returns>True if the directory was deleted, false if not</returns>
-		static bool TryDeleteEmptyDirectory(DirectoryReference CandidateDirectory)
+		static bool TryDeleteEmptyDirectory(DirectoryReference candidateDirectory)
 		{
 			// Make sure the directory exists
-			if(!DirectoryReference.Exists(CandidateDirectory))
+			if (!DirectoryReference.Exists(candidateDirectory))
 			{
 				return false;
 			}
 
 			// Check if there are any files in it. If there are, don't bother trying to delete it.
-			if(Directory.EnumerateFiles(CandidateDirectory.FullName).Any() || Directory.EnumerateDirectories(CandidateDirectory.FullName).Any())
+			if (Directory.EnumerateFiles(candidateDirectory.FullName).Any() || Directory.EnumerateDirectories(candidateDirectory.FullName).Any())
 			{
 				return false;
 			}
@@ -158,12 +149,12 @@ namespace AutomationTool.Tasks
 			// Try to delete the directory.
 			try
 			{
-				Directory.Delete(CandidateDirectory.FullName);
+				Directory.Delete(candidateDirectory.FullName);
 				return true;
 			}
-			catch(Exception Ex)
+			catch (Exception ex)
 			{
-				Logger.LogWarning("Couldn't delete directory {Arg0} ({Arg1})", CandidateDirectory.FullName, Ex.Message);
+				Logger.LogWarning("Couldn't delete directory {Arg0} ({Arg1})", candidateDirectory.FullName, ex.Message);
 				return false;
 			}
 		}
@@ -171,9 +162,9 @@ namespace AutomationTool.Tasks
 		/// <summary>
 		/// Output this task out to an XML writer.
 		/// </summary>
-		public override void Write(XmlWriter Writer)
+		public override void Write(XmlWriter writer)
 		{
-			Write(Writer, Parameters);
+			Write(writer, _parameters);
 		}
 
 		/// <summary>
@@ -182,7 +173,7 @@ namespace AutomationTool.Tasks
 		/// <returns>The tag names which are read by this task</returns>
 		public override IEnumerable<string> FindConsumedTagNames()
 		{
-			return FindTagNamesFromFilespec(Parameters.Files);
+			return FindTagNamesFromFilespec(_parameters.Files);
 		}
 
 		/// <summary>
diff --git a/Engine/Source/Programs/AutomationTool/BuildGraph/Tasks/DeployToolTask.cs b/Engine/Source/Programs/AutomationTool/BuildGraph/Tasks/DeployToolTask.cs
--- a/Engine/Source/Programs/AutomationTool/BuildGraph/Tasks/DeployToolTask.cs
+++ b/Engine/Source/Programs/AutomationTool/BuildGraph/Tasks/DeployToolTask.cs
@@ -1,29 +1,19 @@
 // Copyright Epic Games, Inc. All Rights Reserved.
 
-using AutomationTool;
-using EpicGames.Core;
 using System;
 using System.Collections.Generic;
-using System.Text.Json;
 using System.IO;
-using System.Net.Http;
-using System.Net.Http.Headers;
-using System.Text;
+using System.Text.Json;
 using System.Threading.Tasks;
 using System.Xml;
-using UnrealBuildBase;
-using Microsoft.Extensions.Logging;
-using EpicGames.Horde.Storage.Bundles;
-using EpicGames.Horde.Storage.Clients;
+using EpicGames.Core;
+using EpicGames.Horde;
+using EpicGames.Horde.Server;
 using EpicGames.Horde.Storage;
 using EpicGames.Horde.Storage.Nodes;
-using System.Threading;
-using System.Data;
-using EpicGames.Horde.Storage.Backends;
-using Microsoft.Extensions.DependencyInjection;
-using EpicGames.Horde;
 using EpicGames.Horde.Tools;
-using EpicGames.Horde.Server;
+using Microsoft.Extensions.DependencyInjection;
+using Microsoft.Extensions.Logging;
 
 #nullable enable
 
@@ -100,32 +90,29 @@ namespace AutomationTool.Tasks
 			public string? Node { get; set; }
 		}
 
-		/// <summary>
-		/// Parameters for this task
-		/// </summary>
-		DeployToolTaskParameters Parameters;
+		readonly DeployToolTaskParameters _parameters;
 
 		/// <summary>
 		/// Construct a Helm task
 		/// </summary>
-		/// <param name="InParameters">Parameters for the task</param>
-		public DeployToolTask(DeployToolTaskParameters InParameters)
+		/// <param name="parameters">Parameters for the task</param>
+		public DeployToolTask(DeployToolTaskParameters parameters)
 		{
-			Parameters = InParameters;
+			_parameters = parameters;
 		}
 
 		/// <summary>
-		/// Execute the task.
+		/// ExecuteAsync the task.
 		/// </summary>
-		/// <param name="Job">Information about the current job</param>
-		/// <param name="BuildProducts">Set of build products produced by this node.</param>
-		/// <param name="TagNameToFileSet">Mapping from tag names to the set of files they include</param>
-		public override async Task ExecuteAsync(JobContext Job, HashSet<FileReference> BuildProducts, Dictionary<string, HashSet<FileReference>> TagNameToFileSet)
+		/// <param name="job">Information about the current job</param>
+		/// <param name="buildProducts">Set of build products produced by this node.</param>
+		/// <param name="tagNameToFileSet">Mapping from tag names to the set of files they include</param>
+		public override async Task ExecuteAsync(JobContext job, HashSet<FileReference> buildProducts, Dictionary<string, HashSet<FileReference>> tagNameToFileSet)
 		{
 			DeploySettings? settings = null;
-			if (!String.IsNullOrEmpty(Parameters.Settings))
+			if (!String.IsNullOrEmpty(_parameters.Settings))
 			{
-				FileReference settingsFile = ResolveFile(Parameters.Settings);
+				FileReference settingsFile = ResolveFile(_parameters.Settings);
 				if (!FileReference.Exists(settingsFile))
 				{
 					throw new AutomationException($"Settings file '{settingsFile}' does not exist");
@@ -145,7 +132,7 @@ namespace AutomationTool.Tasks
 				}
 			}
 
-			ToolId toolId = new ToolId(Parameters.Id);
+			ToolId toolId = new ToolId(_parameters.Id);
 
 			ServiceCollection serviceCollection = new ServiceCollection();
 			serviceCollection.Configure<HordeOptions>(options =>
@@ -178,14 +165,14 @@ namespace AutomationTool.Tasks
 			await using (IBlobWriter blobWriter = storageClient.CreateBlobWriter(serializerOptions: serializerOptions))
 			{
 				DirectoryNode sandbox = new DirectoryNode();
-				if (Parameters.File != null)
+				if (_parameters.File != null)
 				{
-					using FileStream stream = FileReference.Open(ResolveFile(Parameters.File), FileMode.Open, FileAccess.Read);
+					using FileStream stream = FileReference.Open(ResolveFile(_parameters.File), FileMode.Open, FileAccess.Read);
 					await sandbox.CopyFromZipStreamAsync(stream, blobWriter, new ChunkingOptions());
 				}
-				else if (Parameters.Directory != null)
+				else if (_parameters.Directory != null)
 				{
-					DirectoryInfo directoryInfo = ResolveDirectory(Parameters.Directory).ToDirectoryInfo();
+					DirectoryInfo directoryInfo = ResolveDirectory(_parameters.Directory).ToDirectoryInfo();
 					await sandbox.AddFilesAsync(directoryInfo, blobWriter);
 				}
 				else
@@ -197,28 +184,28 @@ namespace AutomationTool.Tasks
 			}
 
 			double? duration = null;
-			if (Parameters.Duration != 0)
+			if (_parameters.Duration != 0)
 			{
-				duration = Parameters.Duration;
+				duration = _parameters.Duration;
 			}
 
 			bool? createPaused = null;
-			if (Parameters.Paused)
+			if (_parameters.Paused)
 			{
 				createPaused = true;
 			}
 
 			BlobRefValue locator = handle.GetRefValue();
-			ToolDeploymentId deploymentId = await hordeHttpClient.CreateToolDeploymentAsync(toolId, Parameters.Version, duration, createPaused, locator);
+			ToolDeploymentId deploymentId = await hordeHttpClient.CreateToolDeploymentAsync(toolId, _parameters.Version, duration, createPaused, locator);
 			Logger.LogInformation("Created {ToolId} deployment {DeploymentId}", toolId, deploymentId);
 		}
 
 		/// <summary>
 		/// Output this task out to an XML writer.
 		/// </summary>
-		public override void Write(XmlWriter Writer)
+		public override void Write(XmlWriter writer)
 		{
-			Write(Writer, Parameters);
+			Write(writer, _parameters);
 		}
 
 		/// <summary>
diff --git a/Engine/Source/Programs/AutomationTool/BuildGraph/Tasks/DockerBuildTask.cs b/Engine/Source/Programs/AutomationTool/BuildGraph/Tasks/DockerBuildTask.cs
--- a/Engine/Source/Programs/AutomationTool/BuildGraph/Tasks/DockerBuildTask.cs
+++ b/Engine/Source/Programs/AutomationTool/BuildGraph/Tasks/DockerBuildTask.cs
@@ -1,17 +1,15 @@
 // Copyright Epic Games, Inc. All Rights Reserved.
 
-using EpicGames.Core;
 using System;
 using System.Collections.Generic;
-using System.IO;
 using System.Linq;
 using System.Text;
-using System.Xml;
-using AutomationTool;
-using UnrealBuildBase;
-using System.Threading.Tasks;
 using System.Text.RegularExpressions;
+using System.Threading.Tasks;
+using System.Xml;
+using EpicGames.Core;
 using Microsoft.Extensions.Logging;
+using UnrealBuildBase;
 
 namespace AutomationTool.Tasks
 {
@@ -105,114 +103,111 @@ namespace AutomationTool.Tasks
 	[TaskElement("Docker-Build", typeof(DockerBuildTaskParameters))]
 	public class DockerBuildTask : SpawnTaskBase
 	{
-		/// <summary>
-		/// Parameters for this task
-		/// </summary>
-		DockerBuildTaskParameters Parameters;
+		readonly DockerBuildTaskParameters _parameters;
 
 		/// <summary>
 		/// Construct a Docker task
 		/// </summary>
-		/// <param name="InParameters">Parameters for the task</param>
-		public DockerBuildTask(DockerBuildTaskParameters InParameters)
+		/// <param name="parameters">Parameters for the task</param>
+		public DockerBuildTask(DockerBuildTaskParameters parameters)
 		{
-			Parameters = InParameters;
+			_parameters = parameters;
 		}
 
 		/// <summary>
-		/// Execute the task.
+		/// ExecuteAsync the task.
 		/// </summary>
-		/// <param name="Job">Information about the current job</param>
-		/// <param name="BuildProducts">Set of build products produced by this node.</param>
-		/// <param name="TagNameToFileSet">Mapping from tag names to the set of files they include</param>
-		public override async Task ExecuteAsync(JobContext Job, HashSet<FileReference> BuildProducts, Dictionary<string, HashSet<FileReference>> TagNameToFileSet)
+		/// <param name="job">Information about the current job</param>
+		/// <param name="buildProducts">Set of build products produced by this node.</param>
+		/// <param name="tagNameToFileSet">Mapping from tag names to the set of files they include</param>
+		public override async Task ExecuteAsync(JobContext job, HashSet<FileReference> buildProducts, Dictionary<string, HashSet<FileReference>> tagNameToFileSet)
 		{
 			Logger.LogInformation("Building Docker image");
-			using (LogIndentScope Scope = new LogIndentScope("  "))
+			using (LogIndentScope scope = new LogIndentScope("  "))
 			{
-				DirectoryReference BaseDir = ResolveDirectory(Parameters.BaseDir);
-				List<FileReference> SourceFiles = ResolveFilespec(BaseDir, Parameters.Files, TagNameToFileSet).ToList();
-				bool isStagingEnabled = SourceFiles.Count > 0;
+				DirectoryReference baseDir = ResolveDirectory(_parameters.BaseDir);
+				List<FileReference> sourceFiles = ResolveFilespec(baseDir, _parameters.Files, tagNameToFileSet).ToList();
+				bool isStagingEnabled = sourceFiles.Count > 0;
 
-				DirectoryReference StagingDir = DirectoryReference.Combine(Unreal.EngineDirectory, "Intermediate", "Docker");
-				FileUtils.ForceDeleteDirectoryContents(StagingDir);
+				DirectoryReference stagingDir = DirectoryReference.Combine(Unreal.EngineDirectory, "Intermediate", "Docker");
+				FileUtils.ForceDeleteDirectoryContents(stagingDir);
 
-				List<FileReference> TargetFiles = SourceFiles.ConvertAll(x => FileReference.Combine(StagingDir, x.MakeRelativeTo(BaseDir)));
-				CommandUtils.ThreadedCopyFiles(SourceFiles, BaseDir, StagingDir);
+				List<FileReference> targetFiles = sourceFiles.ConvertAll(x => FileReference.Combine(stagingDir, x.MakeRelativeTo(baseDir)));
+				CommandUtils.ThreadedCopyFiles(sourceFiles, baseDir, stagingDir);
 
-				FileReference DockerIgnoreFileInBaseDir = FileReference.Combine(BaseDir, ".dockerignore");
-				FileReference.Delete(DockerIgnoreFileInBaseDir);
+				FileReference dockerIgnoreFileInBaseDir = FileReference.Combine(baseDir, ".dockerignore");
+				FileReference.Delete(dockerIgnoreFileInBaseDir);
 
-				if (!String.IsNullOrEmpty(Parameters.OverlayDirs))
+				if (!String.IsNullOrEmpty(_parameters.OverlayDirs))
 				{
-					foreach (string OverlayDir in Parameters.OverlayDirs.Split(';'))
+					foreach (string overlayDir in _parameters.OverlayDirs.Split(';'))
 					{
-						CommandUtils.ThreadedCopyFiles(ResolveDirectory(OverlayDir), StagingDir);
+						CommandUtils.ThreadedCopyFiles(ResolveDirectory(overlayDir), stagingDir);
 					}
 				}
 
-				StringBuilder Arguments = new StringBuilder("build .");
-				if (Parameters.Tag != null)
+				StringBuilder arguments = new StringBuilder("build .");
+				if (_parameters.Tag != null)
 				{
-					Arguments.Append($" -t {Parameters.Tag}");
+					arguments.Append($" -t {_parameters.Tag}");
 				}
-				if (Parameters.Target != null)
+				if (_parameters.Target != null)
 				{
-					Arguments.Append($" --target {Parameters.Target}");
+					arguments.Append($" --target {_parameters.Target}");
 				}
-				if (Parameters.Output != null)
+				if (_parameters.Output != null)
 				{
-					if (!Parameters.UseBuildKit)
+					if (!_parameters.UseBuildKit)
 					{
-						throw new AutomationException($"{nameof(Parameters.UseBuildKit)} must be enabled to use '{nameof(Parameters.Output)}' parameter");
+						throw new AutomationException($"{nameof(_parameters.UseBuildKit)} must be enabled to use '{nameof(_parameters.Output)}' parameter");
 					}
-					Arguments.Append($" --output {Parameters.Output}");
+					arguments.Append($" --output {_parameters.Output}");
 				}
-				if (Parameters.DockerFile != null)
+				if (_parameters.DockerFile != null)
 				{
-					FileReference DockerFile = ResolveFile(Parameters.DockerFile);
-					if (!DockerFile.IsUnderDirectory(BaseDir))
+					FileReference dockerFile = ResolveFile(_parameters.DockerFile);
+					if (!dockerFile.IsUnderDirectory(baseDir))
 					{
-						throw new AutomationException($"Dockerfile '{DockerFile}' is not under base directory ({BaseDir})");
+						throw new AutomationException($"Dockerfile '{dockerFile}' is not under base directory ({baseDir})");
 					}
-					Arguments.Append($" -f {DockerFile.MakeRelativeTo(BaseDir).QuoteArgument()}");
+					arguments.Append($" -f {dockerFile.MakeRelativeTo(baseDir).QuoteArgument()}");
 				}
-				if (Parameters.DockerIgnoreFile != null)
+				if (_parameters.DockerIgnoreFile != null)
 				{
-					FileReference DockerIgnoreFile = ResolveFile(Parameters.DockerIgnoreFile);
-					FileReference.Copy(DockerIgnoreFile, DockerIgnoreFileInBaseDir);
+					FileReference dockerIgnoreFile = ResolveFile(_parameters.DockerIgnoreFile);
+					FileReference.Copy(dockerIgnoreFile, dockerIgnoreFileInBaseDir);
 				}
-				if (Parameters.ProgressOutput != null)
+				if (_parameters.ProgressOutput != null)
 				{
-					Arguments.Append($" --progress={Parameters.ProgressOutput}");
+					arguments.Append($" --progress={_parameters.ProgressOutput}");
 				}
-				if (Parameters.Arguments != null)
+				if (_parameters.Arguments != null)
 				{
-					Arguments.Append($" {Parameters.Arguments}");
+					arguments.Append($" {_parameters.Arguments}");
 				}
 
-				Dictionary<string, string> EnvVars = ParseEnvVars(Parameters.Environment, Parameters.EnvironmentFile);
-				if (Parameters.UseBuildKit)
+				Dictionary<string, string> envVars = ParseEnvVars(_parameters.Environment, _parameters.EnvironmentFile);
+				if (_parameters.UseBuildKit)
 				{
-					EnvVars["DOCKER_BUILDKIT"] = "1";
+					envVars["DOCKER_BUILDKIT"] = "1";
 				}
-				
-				string WorkingDir = isStagingEnabled ? StagingDir.FullName : BaseDir.FullName;
-				string Exe = DockerTask.GetDockerExecutablePath();
-				await SpawnTaskBase.ExecuteAsync(Exe, Arguments.ToString(), EnvVars: EnvVars, WorkingDir: WorkingDir, SpewFilterCallback: FilterOutput);
+
+				string workingDir = isStagingEnabled ? stagingDir.FullName : baseDir.FullName;
+				string exe = DockerTask.GetDockerExecutablePath();
+				await SpawnTaskBase.ExecuteAsync(exe, arguments.ToString(), envVars: envVars, workingDir: workingDir, spewFilterCallback: FilterOutput);
 			}
 		}
 
-		static Regex FilterOutputPattern = new Regex(@"^#\d+ (?:\d+\.\d+ )?");
+		static readonly Regex s_filterOutputPattern = new Regex(@"^#\d+ (?:\d+\.\d+ )?");
 
-		static string FilterOutput(string Line) => FilterOutputPattern.Replace(Line, "");
+		static string FilterOutput(string line) => s_filterOutputPattern.Replace(line, "");
 
 		/// <summary>
 		/// Output this task out to an XML writer.
 		/// </summary>
-		public override void Write(XmlWriter Writer)
+		public override void Write(XmlWriter writer)
 		{
-			Write(Writer, Parameters);
+			Write(writer, _parameters);
 		}
 
 		/// <summary>
@@ -221,10 +216,10 @@ namespace AutomationTool.Tasks
 		/// <returns>The tag names which are read by this task</returns>
 		public override IEnumerable<string> FindConsumedTagNames()
 		{
-			List<string> TagNames = new List<string>();
-			TagNames.AddRange(FindTagNamesFromFilespec(Parameters.DockerFile));
-			TagNames.AddRange(FindTagNamesFromFilespec(Parameters.Files));
-			return TagNames;
+			List<string> tagNames = new List<string>();
+			tagNames.AddRange(FindTagNamesFromFilespec(_parameters.DockerFile));
+			tagNames.AddRange(FindTagNamesFromFilespec(_parameters.Files));
+			return tagNames;
 		}
 
 		/// <summary>
diff --git a/Engine/Source/Programs/AutomationTool/BuildGraph/Tasks/DockerComposeDownTask.cs b/Engine/Source/Programs/AutomationTool/BuildGraph/Tasks/DockerComposeDownTask.cs
--- a/Engine/Source/Programs/AutomationTool/BuildGraph/Tasks/DockerComposeDownTask.cs
+++ b/Engine/Source/Programs/AutomationTool/BuildGraph/Tasks/DockerComposeDownTask.cs
@@ -1,15 +1,12 @@
 // Copyright Epic Games, Inc. All Rights Reserved.
 
-using EpicGames.Core;
 using System;
 using System.Collections.Generic;
-using System.IO;
 using System.Linq;
 using System.Text;
-using System.Xml;
-using AutomationTool;
-using UnrealBuildBase;
 using System.Threading.Tasks;
+using System.Xml;
+using EpicGames.Core;
 using Microsoft.Extensions.Logging;
 
 namespace AutomationTool.Tasks
@@ -41,49 +38,49 @@ namespace AutomationTool.Tasks
 		/// <summary>
 		/// Parameters for this task
 		/// </summary>
-		DockerComposeDownTaskParameters Parameters;
+		readonly DockerComposeDownTaskParameters _parameters;
 
 		/// <summary>
 		/// Construct a Docker-Compose task
 		/// </summary>
-		/// <param name="InParameters">Parameters for the task</param>
-		public DockerComposeDownTask(DockerComposeDownTaskParameters InParameters)
+		/// <param name="parameters">Parameters for the task</param>
+		public DockerComposeDownTask(DockerComposeDownTaskParameters parameters)
 		{
-			Parameters = InParameters;
+			_parameters = parameters;
 		}
 
 		/// <summary>
-		/// Execute the task.
+		/// ExecuteAsync the task.
 		/// </summary>
-		/// <param name="Job">Information about the current job</param>
-		/// <param name="BuildProducts">Set of build products produced by this node.</param>
-		/// <param name="TagNameToFileSet">Mapping from tag names to the set of files they include</param>
-		public override async Task ExecuteAsync(JobContext Job, HashSet<FileReference> BuildProducts, Dictionary<string, HashSet<FileReference>> TagNameToFileSet)
+		/// <param name="job">Information about the current job</param>
+		/// <param name="buildProducts">Set of build products produced by this node.</param>
+		/// <param name="tagNameToFileSet">Mapping from tag names to the set of files they include</param>
+		public override async Task ExecuteAsync(JobContext job, HashSet<FileReference> buildProducts, Dictionary<string, HashSet<FileReference>> tagNameToFileSet)
 		{
-			StringBuilder Arguments = new StringBuilder("--ansi never ");
-			if (!String.IsNullOrEmpty(Parameters.File))
+			StringBuilder arguments = new StringBuilder("--ansi never ");
+			if (!String.IsNullOrEmpty(_parameters.File))
 			{
-				Arguments.Append($"--file {Parameters.File.QuoteArgument()} ");
+				arguments.Append($"--file {_parameters.File.QuoteArgument()} ");
 			}
-			Arguments.Append("down");
-			if (!String.IsNullOrEmpty(Parameters.Arguments))
+			arguments.Append("down");
+			if (!String.IsNullOrEmpty(_parameters.Arguments))
 			{
-				Arguments.Append($" {Parameters.Arguments}");
+				arguments.Append($" {_parameters.Arguments}");
 			}
 
-			Logger.LogInformation("Running docker compose {Arguments}", Arguments.ToString());
-			using (LogIndentScope Scope = new LogIndentScope("  "))
+			Logger.LogInformation("Running docker compose {Arguments}", arguments.ToString());
+			using (LogIndentScope scope = new LogIndentScope("  "))
 			{
-				await SpawnTaskBase.ExecuteAsync("docker-compose", Arguments.ToString());
+				await SpawnTaskBase.ExecuteAsync("docker-compose", arguments.ToString());
 			}
 		}
 
 		/// <summary>
 		/// Output this task out to an XML writer.
 		/// </summary>
-		public override void Write(XmlWriter Writer)
+		public override void Write(XmlWriter writer)
 		{
-			Write(Writer, Parameters);
+			Write(writer, _parameters);
 		}
 
 		/// <summary>
diff --git a/Engine/Source/Programs/AutomationTool/BuildGraph/Tasks/DockerComposeUpTask.cs b/Engine/Source/Programs/AutomationTool/BuildGraph/Tasks/DockerComposeUpTask.cs
--- a/Engine/Source/Programs/AutomationTool/BuildGraph/Tasks/DockerComposeUpTask.cs
+++ b/Engine/Source/Programs/AutomationTool/BuildGraph/Tasks/DockerComposeUpTask.cs
@@ -1,16 +1,14 @@
 // Copyright Epic Games, Inc. All Rights Reserved.
 
-using EpicGames.Core;
 using System;
 using System.Collections.Generic;
-using System.IO;
 using System.Linq;
 using System.Text;
-using System.Xml;
-using AutomationTool;
-using UnrealBuildBase;
 using System.Threading.Tasks;
+using System.Xml;
+using EpicGames.Core;
 using Microsoft.Extensions.Logging;
+using UnrealBuildBase;
 
 namespace AutomationTool.Tasks
 {
@@ -38,60 +36,57 @@ namespace AutomationTool.Tasks
 	[TaskElement("Docker-Compose-Up", typeof(DockerComposeUpTaskParameters))]
 	public class DockerComposeUpTask : SpawnTaskBase
 	{
-		/// <summary>
-		/// Parameters for this task
-		/// </summary>
-		DockerComposeUpTaskParameters Parameters;
+		readonly DockerComposeUpTaskParameters _parameters;
 
 		/// <summary>
 		/// Construct a Docker-Compose task
 		/// </summary>
-		/// <param name="InParameters">Parameters for the task</param>
-		public DockerComposeUpTask(DockerComposeUpTaskParameters InParameters)
+		/// <param name="parameters">Parameters for the task</param>
+		public DockerComposeUpTask(DockerComposeUpTaskParameters parameters)
 		{
-			Parameters = InParameters;
+			_parameters = parameters;
 		}
 
 		/// <summary>
-		/// Execute the task.
+		/// ExecuteAsync the task.
 		/// </summary>
-		/// <param name="Job">Information about the current job</param>
-		/// <param name="BuildProducts">Set of build products produced by this node.</param>
-		/// <param name="TagNameToFileSet">Mapping from tag names to the set of files they include</param>
-		public override async Task ExecuteAsync(JobContext Job, HashSet<FileReference> BuildProducts, Dictionary<string, HashSet<FileReference>> TagNameToFileSet)
+		/// <param name="job">Information about the current job</param>
+		/// <param name="buildProducts">Set of build products produced by this node.</param>
+		/// <param name="tagNameToFileSet">Mapping from tag names to the set of files they include</param>
+		public override async Task ExecuteAsync(JobContext job, HashSet<FileReference> buildProducts, Dictionary<string, HashSet<FileReference>> tagNameToFileSet)
 		{
-			FileReference LogFile = FileReference.Combine(Unreal.RootDirectory, "Engine/Programs/AutomationTool/Saved/Logs/docker-compose-logs.txt");
-			string[] Lines =
+			FileReference logFile = FileReference.Combine(Unreal.RootDirectory, "Engine/Programs/AutomationTool/Saved/Logs/docker-compose-logs.txt");
+			string[] lines =
 			{
-				$"docker-compose -f {Parameters.File.QuoteArgument()} logs --no-color > {LogFile.FullName.QuoteArgument()}",
-				$"docker-compose -f {Parameters.File.QuoteArgument()} down"
+				$"docker-compose -f {_parameters.File.QuoteArgument()} logs --no-color > {logFile.FullName.QuoteArgument()}",
+				$"docker-compose -f {_parameters.File.QuoteArgument()} down"
 			};
-			await AddCleanupCommandsAsync(Lines);
+			await AddCleanupCommandsAsync(lines);
 
-			StringBuilder Arguments = new StringBuilder("--ansi never ");
-			if (!String.IsNullOrEmpty(Parameters.File))
+			StringBuilder arguments = new StringBuilder("--ansi never ");
+			if (!String.IsNullOrEmpty(_parameters.File))
 			{
-				Arguments.Append($"--file {Parameters.File.QuoteArgument()} ");
+				arguments.Append($"--file {_parameters.File.QuoteArgument()} ");
 			}
-			Arguments.Append("up --detach");
-			if (!String.IsNullOrEmpty(Parameters.Arguments))
+			arguments.Append("up --detach");
+			if (!String.IsNullOrEmpty(_parameters.Arguments))
 			{
-				Arguments.Append($" {Parameters.Arguments}");
+				arguments.Append($" {_parameters.Arguments}");
 			}
 
-			Logger.LogInformation("Running docker compose {Arguments}", Arguments.ToString());
-			using (LogIndentScope Scope = new LogIndentScope("  "))
+			Logger.LogInformation("Running docker compose {Arguments}", arguments.ToString());
+			using (LogIndentScope scope = new LogIndentScope("  "))
 			{
-				await SpawnTaskBase.ExecuteAsync("docker-compose", Arguments.ToString());
+				await SpawnTaskBase.ExecuteAsync("docker-compose", arguments.ToString());
 			}
 		}
 
 		/// <summary>
 		/// Output this task out to an XML writer.
 		/// </summary>
-		public override void Write(XmlWriter Writer)
+		public override void Write(XmlWriter writer)
 		{
-			Write(Writer, Parameters);
+			Write(writer, _parameters);
 		}
 
 		/// <summary>
diff --git a/Engine/Source/Programs/AutomationTool/BuildGraph/Tasks/DockerPushTask.cs b/Engine/Source/Programs/AutomationTool/BuildGraph/Tasks/DockerPushTask.cs
--- a/Engine/Source/Programs/AutomationTool/BuildGraph/Tasks/DockerPushTask.cs
+++ b/Engine/Source/Programs/AutomationTool/BuildGraph/Tasks/DockerPushTask.cs
@@ -1,15 +1,11 @@
 // Copyright Epic Games, Inc. All Rights Reserved.
 
-using AutomationTool;
-using EpicGames.Core;
 using System;
 using System.Collections.Generic;
-using System.IO;
-using System.Linq;
-using System.Text;
 using System.Text.Json;
 using System.Threading.Tasks;
 using System.Xml;
+using EpicGames.Core;
 using Microsoft.Extensions.Logging;
 
 namespace AutomationTool.Tasks
@@ -68,58 +64,55 @@ namespace AutomationTool.Tasks
 	[TaskElement("Docker-Push", typeof(DockerPushTaskParameters))]
 	public class DockerPushTask : SpawnTaskBase
 	{
-		/// <summary>
-		/// Parameters for this task
-		/// </summary>
-		DockerPushTaskParameters Parameters;
+		readonly DockerPushTaskParameters _parameters;
 
 		/// <summary>
 		/// Construct a Docker task
 		/// </summary>
-		/// <param name="InParameters">Parameters for the task</param>
-		public DockerPushTask(DockerPushTaskParameters InParameters)
+		/// <param name="parameters">Parameters for the task</param>
+		public DockerPushTask(DockerPushTaskParameters parameters)
 		{
-			Parameters = InParameters;
+			_parameters = parameters;
 		}
 
 		/// <summary>
-		/// Execute the task.
+		/// ExecuteAsync the task.
 		/// </summary>
-		/// <param name="Job">Information about the current job</param>
-		/// <param name="BuildProducts">Set of build products produced by this node.</param>
-		/// <param name="TagNameToFileSet">Mapping from tag names to the set of files they include</param>
-		public override async Task ExecuteAsync(JobContext Job, HashSet<FileReference> BuildProducts, Dictionary<string, HashSet<FileReference>> TagNameToFileSet)
+		/// <param name="job">Information about the current job</param>
+		/// <param name="buildProducts">Set of build products produced by this node.</param>
+		/// <param name="tagNameToFileSet">Mapping from tag names to the set of files they include</param>
+		public override async Task ExecuteAsync(JobContext job, HashSet<FileReference> buildProducts, Dictionary<string, HashSet<FileReference>> tagNameToFileSet)
 		{
 			Logger.LogInformation("Pushing Docker image");
-			using (LogIndentScope Scope = new LogIndentScope("  "))
+			using (LogIndentScope scope = new LogIndentScope("  "))
 			{
-				string Exe = DockerTask.GetDockerExecutablePath();
-				Dictionary<string, string> Environment = ParseEnvVars(Parameters.Environment, Parameters.EnvironmentFile);
+				string exe = DockerTask.GetDockerExecutablePath();
+				Dictionary<string, string> environment = ParseEnvVars(_parameters.Environment, _parameters.EnvironmentFile);
 
-				if (Parameters.AwsEcr)
+				if (_parameters.AwsEcr)
 				{
-					IProcessResult Result = await SpawnTaskBase.ExecuteAsync("aws", "ecr get-login-password", EnvVars: Environment, LogOutput: false);
-					await ExecuteAsync(Exe, $"login {Parameters.Repository} --username AWS --password-stdin", Input: Result.Output);
+					IProcessResult result = await SpawnTaskBase.ExecuteAsync("aws", "ecr get-login-password", envVars: environment, logOutput: false);
+					await ExecuteAsync(exe, $"login {_parameters.Repository} --username AWS --password-stdin", input: result.Output);
 				}
-				if (!String.IsNullOrEmpty(Parameters.RepositoryAuthFile))
+				if (!String.IsNullOrEmpty(_parameters.RepositoryAuthFile))
 				{
-					string RepositoryText = CommandUtils.ReadAllText(Parameters.RepositoryAuthFile);
-					Dictionary<string, string> AuthDict = JsonSerializer.Deserialize<Dictionary<string, string>>(RepositoryText);
-					await ExecuteAsync(Exe, $"login {Parameters.Repository} --username {AuthDict["Username"]} --password-stdin", Input: AuthDict["Token"]);
+					string repositoryText = CommandUtils.ReadAllText(_parameters.RepositoryAuthFile);
+					Dictionary<string, string> authDict = JsonSerializer.Deserialize<Dictionary<string, string>>(repositoryText);
+					await ExecuteAsync(exe, $"login {_parameters.Repository} --username {authDict["Username"]} --password-stdin", input: authDict["Token"]);
 				}
 
-				string TargetImage = Parameters.TargetImage ?? Parameters.Image;
-				await ExecuteAsync(Exe, $"tag {Parameters.Image} {Parameters.Repository}/{TargetImage}", EnvVars: Environment);
-				await ExecuteAsync(Exe, $"push {Parameters.Repository}/{TargetImage}", EnvVars: Environment);
+				string targetImage = _parameters.TargetImage ?? _parameters.Image;
+				await ExecuteAsync(exe, $"tag {_parameters.Image} {_parameters.Repository}/{targetImage}", envVars: environment);
+				await ExecuteAsync(exe, $"push {_parameters.Repository}/{targetImage}", envVars: environment);
 			}
 		}
 
 		/// <summary>
 		/// Output this task out to an XML writer.
 		/// </summary>
-		public override void Write(XmlWriter Writer)
+		public override void Write(XmlWriter writer)
 		{
-			Write(Writer, Parameters);
+			Write(writer, _parameters);
 		}
 
 		/// <summary>
diff --git a/Engine/Source/Programs/AutomationTool/BuildGraph/Tasks/DockerTask.cs b/Engine/Source/Programs/AutomationTool/BuildGraph/Tasks/DockerTask.cs
--- a/Engine/Source/Programs/AutomationTool/BuildGraph/Tasks/DockerTask.cs
+++ b/Engine/Source/Programs/AutomationTool/BuildGraph/Tasks/DockerTask.cs
@@ -1,13 +1,10 @@
 // Copyright Epic Games, Inc. All Rights Reserved.
 
-using AutomationTool;
-using EpicGames.Core;
 using System;
 using System.Collections.Generic;
-using System.IO;
-using System.Text;
 using System.Threading.Tasks;
 using System.Xml;
+using EpicGames.Core;
 
 namespace AutomationTool.Tasks
 {
@@ -50,34 +47,34 @@ namespace AutomationTool.Tasks
 		/// <summary>
 		/// Parameters for this task
 		/// </summary>
-		DockerTaskParameters Parameters;
+		readonly DockerTaskParameters _parameters;
 
 		/// <summary>
 		/// Construct a Docker task
 		/// </summary>
-		/// <param name="InParameters">Parameters for the task</param>
-		public DockerTask(DockerTaskParameters InParameters)
+		/// <param name="parameters">Parameters for the task</param>
+		public DockerTask(DockerTaskParameters parameters)
 		{
-			Parameters = InParameters;
+			_parameters = parameters;
 		}
 
 		/// <summary>
-		/// Execute the task.
+		/// ExecuteAsync the task.
 		/// </summary>
-		/// <param name="Job">Information about the current job</param>
-		/// <param name="BuildProducts">Set of build products produced by this node.</param>
-		/// <param name="TagNameToFileSet">Mapping from tag names to the set of files they include</param>
-		public override async Task ExecuteAsync(JobContext Job, HashSet<FileReference> BuildProducts, Dictionary<string, HashSet<FileReference>> TagNameToFileSet)
+		/// <param name="job">Information about the current job</param>
+		/// <param name="buildProducts">Set of build products produced by this node.</param>
+		/// <param name="tagNameToFileSet">Mapping from tag names to the set of files they include</param>
+		public override async Task ExecuteAsync(JobContext job, HashSet<FileReference> buildProducts, Dictionary<string, HashSet<FileReference>> tagNameToFileSet)
 		{
-			await ExecuteAsync(GetDockerExecutablePath(), Parameters.Arguments, EnvVars: ParseEnvVars(Parameters.Environment, Parameters.EnvironmentFile), WorkingDir: Parameters.WorkingDir);
+			await ExecuteAsync(GetDockerExecutablePath(), _parameters.Arguments, envVars: ParseEnvVars(_parameters.Environment, _parameters.EnvironmentFile), workingDir: _parameters.WorkingDir);
 		}
 
 		/// <summary>
 		/// Output this task out to an XML writer.
 		/// </summary>
-		public override void Write(XmlWriter Writer)
+		public override void Write(XmlWriter writer)
 		{
-			Write(Writer, Parameters);
+			Write(writer, _parameters);
 		}
 
 		/// <summary>
diff --git a/Engine/Source/Programs/AutomationTool/BuildGraph/Tasks/DotNetTask.cs b/Engine/Source/Programs/AutomationTool/BuildGraph/Tasks/DotNetTask.cs
--- a/Engine/Source/Programs/AutomationTool/BuildGraph/Tasks/DotNetTask.cs
+++ b/Engine/Source/Programs/AutomationTool/BuildGraph/Tasks/DotNetTask.cs
@@ -1,15 +1,9 @@
 // Copyright Epic Games, Inc. All Rights Reserved.
 
-using AutomationTool;
-using AutomationTool.Tasks;
-using EpicGames.BuildGraph;
-using EpicGames.Core;
-using System;
 using System.Collections.Generic;
-using System.IO;
-using System.Text;
 using System.Threading.Tasks;
 using System.Xml;
+using EpicGames.Core;
 using UnrealBuildBase;
 
 namespace AutomationTool.Tasks
@@ -62,47 +56,44 @@ namespace AutomationTool.Tasks
 	[TaskElement("DotNet", typeof(DotNetTaskParameters))]
 	public class DotNetTask : SpawnTaskBase
 	{
-		/// <summary>
-		/// Parameters for this task
-		/// </summary>
-		DotNetTaskParameters Parameters;
+		readonly DotNetTaskParameters _parameters;
 
 		/// <summary>
 		/// Construct a Docker task
 		/// </summary>
-		/// <param name="InParameters">Parameters for the task</param>
-		public DotNetTask(DotNetTaskParameters InParameters)
+		/// <param name="parameters">Parameters for the task</param>
+		public DotNetTask(DotNetTaskParameters parameters)
 		{
-			Parameters = InParameters;
+			_parameters = parameters;
 		}
 
 		/// <summary>
-		/// Execute the task.
+		/// ExecuteAsync the task.
 		/// </summary>
-		/// <param name="Job">Information about the current job</param>
-		/// <param name="BuildProducts">Set of build products produced by this node.</param>
-		/// <param name="TagNameToFileSet">Mapping from tag names to the set of files they include</param>
-		public override async Task ExecuteAsync(JobContext Job, HashSet<FileReference> BuildProducts, Dictionary<string, HashSet<FileReference>> TagNameToFileSet)
+		/// <param name="job">Information about the current job</param>
+		/// <param name="buildProducts">Set of build products produced by this node.</param>
+		/// <param name="tagNameToFileSet">Mapping from tag names to the set of files they include</param>
+		public override async Task ExecuteAsync(JobContext job, HashSet<FileReference> buildProducts, Dictionary<string, HashSet<FileReference>> tagNameToFileSet)
 		{
-			FileReference DotNetFile = Parameters.DotNetPath == null ? Unreal.DotnetPath : Parameters.DotNetPath;
-			if(!FileReference.Exists(DotNetFile))
+			FileReference dotNetFile = _parameters.DotNetPath == null ? Unreal.DotnetPath : _parameters.DotNetPath;
+			if (!FileReference.Exists(dotNetFile))
 			{
-				throw new AutomationException("DotNet is missing from {0}", DotNetFile);
+				throw new AutomationException("DotNet is missing from {0}", dotNetFile);
 			}
 
-			IProcessResult Result = await ExecuteAsync(DotNetFile.FullName, Parameters.Arguments, WorkingDir: Parameters.BaseDir, EnvVars: ParseEnvVars(Parameters.Environment, Parameters.EnvironmentFile));
-			if (Result.ExitCode < 0 || Result.ExitCode >= Parameters.ErrorLevel)
+			IProcessResult result = await ExecuteAsync(dotNetFile.FullName, _parameters.Arguments, workingDir: _parameters.BaseDir, envVars: ParseEnvVars(_parameters.Environment, _parameters.EnvironmentFile));
+			if (result.ExitCode < 0 || result.ExitCode >= _parameters.ErrorLevel)
 			{
-				throw new AutomationException("Docker terminated with an exit code indicating an error ({0})", Result.ExitCode);
+				throw new AutomationException("Docker terminated with an exit code indicating an error ({0})", result.ExitCode);
 			}
 		}
 
 		/// <summary>
 		/// Output this task out to an XML writer.
 		/// </summary>
-		public override void Write(XmlWriter Writer)
+		public override void Write(XmlWriter writer)
 		{
-			Write(Writer, Parameters);
+			Write(writer, _parameters);
 		}
 
 		/// <summary>
@@ -129,23 +120,23 @@ namespace AutomationTool.Tasks
 		/// <summary>
 		/// Runs a command using dotnet.
 		/// </summary>
-		/// <param name="Arguments">Command-line arguments.</param>
-		/// <param name="BaseDir">Base directory for running the command.</param>
-		/// <param name="Environment">Environment variables to set.</param>
-		/// <param name="EnvironmentFile">File to read environment variables from.</param>
-		/// <param name="ErrorLevel">The minimum exit code, which is treated as an error.</param>
-		/// <param name="DotNetPath">Override path to dotnet executable.</param>
-		public static async Task DotNetAsync(string Arguments = null, DirectoryReference BaseDir = null, string Environment = null, FileReference EnvironmentFile = null, int ErrorLevel = 1, FileReference DotNetPath = null)
+		/// <param name="arguments">Command-line arguments.</param>
+		/// <param name="baseDir">Base directory for running the command.</param>
+		/// <param name="environment">Environment variables to set.</param>
+		/// <param name="environmentFile">File to read environment variables from.</param>
+		/// <param name="errorLevel">The minimum exit code, which is treated as an error.</param>
+		/// <param name="dotNetPath">Override path to dotnet executable.</param>
+		public static async Task DotNetAsync(string arguments = null, DirectoryReference baseDir = null, string environment = null, FileReference environmentFile = null, int errorLevel = 1, FileReference dotNetPath = null)
 		{
-			DotNetTaskParameters Parameters = new DotNetTaskParameters();
-			Parameters.Arguments = Arguments;
-			Parameters.BaseDir = BaseDir?.FullName;
-			Parameters.Environment = Environment;
-			Parameters.EnvironmentFile = EnvironmentFile?.FullName;
-			Parameters.ErrorLevel = ErrorLevel;
-			Parameters.DotNetPath = DotNetPath;
-			
-			await ExecuteAsync(new DotNetTask(Parameters));
+			DotNetTaskParameters parameters = new DotNetTaskParameters();
+			parameters.Arguments = arguments;
+			parameters.BaseDir = baseDir?.FullName;
+			parameters.Environment = environment;
+			parameters.EnvironmentFile = environmentFile?.FullName;
+			parameters.ErrorLevel = errorLevel;
+			parameters.DotNetPath = dotNetPath;
+
+			await ExecuteAsync(new DotNetTask(parameters));
 		}
 	}
 }
diff --git a/Engine/Source/Programs/AutomationTool/BuildGraph/Tasks/FindModifiedFiles.cs b/Engine/Source/Programs/AutomationTool/BuildGraph/Tasks/FindModifiedFiles.cs
--- a/Engine/Source/Programs/AutomationTool/BuildGraph/Tasks/FindModifiedFiles.cs
+++ b/Engine/Source/Programs/AutomationTool/BuildGraph/Tasks/FindModifiedFiles.cs
@@ -1,8 +1,5 @@
 // Copyright Epic Games, Inc. All Rights Reserved.
 
-using EpicGames.BuildGraph;
-using EpicGames.Perforce;
-using Microsoft.Extensions.Logging.Abstractions;
 using System;
 using System.Collections.Generic;
 using System.Linq;
@@ -10,8 +7,7 @@ using System.Text;
 using System.Threading.Tasks;
 using System.Xml;
 using EpicGames.Core;
-using UnrealBuildBase;
-using UnrealBuildTool;
+using EpicGames.Perforce;
 using Microsoft.Extensions.Logging;
 
 namespace AutomationTool.Tasks
@@ -24,7 +20,7 @@ namespace AutomationTool.Tasks
 		/// <summary>
 		///  List of file specifications separated by semicolon (default is ...)
 		/// </summary>
-		[TaskParameter(Optional=true, ValidationType = TaskParameterValidationType.FileSpec)]
+		[TaskParameter(Optional = true, ValidationType = TaskParameterValidationType.FileSpec)]
 		public string Path { get; set; } = "...";
 
 		/// <summary>
@@ -61,44 +57,44 @@ namespace AutomationTool.Tasks
 		/// <summary>
 		/// Parameters for the task
 		/// </summary>
-		FindModifiedFilesTaskParameters Parameters;
+		readonly FindModifiedFilesTaskParameters _parameters;
 
 		/// <summary>
 		/// Constructor.
 		/// </summary>
-		/// <param name="InParameters">Parameters for this task</param>
-		public FindModifiedFilesTask(FindModifiedFilesTaskParameters InParameters)
+		/// <param name="parameters">Parameters for this task</param>
+		public FindModifiedFilesTask(FindModifiedFilesTaskParameters parameters)
 		{
-			Parameters = InParameters;
+			_parameters = parameters;
 		}
 
 		/// <summary>
-		/// Execute the task.
+		/// ExecuteAsync the task.
 		/// </summary>
-		/// <param name="Job">Information about the current job</param>
-		/// <param name="BuildProducts">Set of build products produced by this node.</param>
-		/// <param name="TagNameToFileSet">Mapping from tag names to the set of files they include</param>
-		public override async Task ExecuteAsync(JobContext Job, HashSet<FileReference> BuildProducts, Dictionary<string, HashSet<FileReference>> TagNameToFileSet)
+		/// <param name="job">Information about the current job</param>
+		/// <param name="buildProducts">Set of build products produced by this node.</param>
+		/// <param name="tagNameToFileSet">Mapping from tag names to the set of files they include</param>
+		public override async Task ExecuteAsync(JobContext job, HashSet<FileReference> buildProducts, Dictionary<string, HashSet<FileReference>> tagNameToFileSet)
 		{
-			using IPerforceConnection Connection = await PerforceConnection.CreateAsync(CommandUtils.P4Settings, Log.Logger);
+			using IPerforceConnection connection = await PerforceConnection.CreateAsync(CommandUtils.P4Settings, Log.Logger);
 
-			SortedSet<string> AllLocalFiles = new SortedSet<string>();
-			foreach (string Path in Parameters.Path.Split(";"))
+			SortedSet<string> allLocalFiles = new SortedSet<string>();
+			foreach (string path in _parameters.Path.Split(";"))
 			{
-				StringBuilder Filter = new StringBuilder($"//{Connection.Settings.ClientName}/{Path}");
-				if (Parameters.Change > 0)
+				StringBuilder filter = new StringBuilder($"//{connection.Settings.ClientName}/{path}");
+				if (_parameters.Change > 0)
 				{
-					Filter.Append($"@={Parameters.Change}");
+					filter.Append($"@={_parameters.Change}");
 				}
-				else if (Parameters.MinChange > 0)
+				else if (_parameters.MinChange > 0)
 				{
-					if (Parameters.MaxChange > 0)
+					if (_parameters.MaxChange > 0)
 					{
-						Filter.Append($"@{Parameters.MinChange},{Parameters.MaxChange}");
+						filter.Append($"@{_parameters.MinChange},{_parameters.MaxChange}");
 					}
 					else
 					{
-						Filter.Append($"@>={Parameters.MinChange}");
+						filter.Append($"@>={_parameters.MinChange}");
 					}
 				}
 				else
@@ -106,47 +102,47 @@ namespace AutomationTool.Tasks
 					throw new AutomationException("Change or MinChange must be specified to FindModifiedFiles task");
 				}
 
-				StreamRecord StreamRecord = await Connection.GetStreamAsync(CommandUtils.P4Env.Branch, true);
-				PerforceViewMap ViewMap = PerforceViewMap.Parse(StreamRecord.View);
+				StreamRecord streamRecord = await connection.GetStreamAsync(CommandUtils.P4Env.Branch, true);
+				PerforceViewMap viewMap = PerforceViewMap.Parse(streamRecord.View);
 
-				HashSet<string> LocalFiles = new HashSet<string>();
+				HashSet<string> localFiles = new HashSet<string>();
 
-				List<FilesRecord> Files = await Connection.FilesAsync(FilesOptions.None, Filter.ToString());
-				foreach (FilesRecord File in Files)
+				List<FilesRecord> files = await connection.FilesAsync(FilesOptions.None, filter.ToString());
+				foreach (FilesRecord file in files)
 				{
-					string LocalFile;
-					if (ViewMap.TryMapFile(File.DepotFile, StringComparison.OrdinalIgnoreCase, out LocalFile))
+					string localFile;
+					if (viewMap.TryMapFile(file.DepotFile, StringComparison.OrdinalIgnoreCase, out localFile))
 					{
-						LocalFiles.Add(LocalFile);
+						localFiles.Add(localFile);
 					}
 					else
 					{
-						Logger.LogInformation("Unable to map {DepotFile} to workspace; skipping.", File.DepotFile);
+						Logger.LogInformation("Unable to map {DepotFile} to workspace; skipping.", file.DepotFile);
 					}
 				}
 
-				Logger.LogInformation("Found {NumFiles} modified files matching {Filter}", LocalFiles.Count, Filter.ToString());
-				foreach (string LocalFile in LocalFiles)
+				Logger.LogInformation("Found {NumFiles} modified files matching {Filter}", localFiles.Count, filter.ToString());
+				foreach (string localFile in localFiles)
 				{
-					Logger.LogInformation("  {LocalFile}", LocalFile);
+					Logger.LogInformation("  {LocalFile}", localFile);
 				}
-				AllLocalFiles.UnionWith(LocalFiles);
+				allLocalFiles.UnionWith(localFiles);
 			}
 
-			Logger.LogInformation("Found {NumFiles} total modified files", AllLocalFiles.Count);
-			if (Parameters.Output != null)
+			Logger.LogInformation("Found {NumFiles} total modified files", allLocalFiles.Count);
+			if (_parameters.Output != null)
 			{
-				await FileReference.WriteAllLinesAsync(Parameters.Output, AllLocalFiles);
-				Logger.LogInformation("Written {OutputFile}", Parameters.Output);
+				await FileReference.WriteAllLinesAsync(_parameters.Output, allLocalFiles);
+				Logger.LogInformation("Written {OutputFile}", _parameters.Output);
 			}
 		}
 
 		/// <summary>
 		/// Output this task out to an XML writer.
 		/// </summary>
-		public override void Write(XmlWriter Writer)
+		public override void Write(XmlWriter writer)
 		{
-			Write(Writer, Parameters);
+			Write(writer, _parameters);
 		}
 
 		/// <summary>
@@ -166,15 +162,15 @@ namespace AutomationTool.Tasks
 		public override IEnumerable<string> FindProducedTagNames()
 		{
 			return Enumerable.Empty<string>();
-/*			foreach (string TagName in FindTagNamesFromList(Parameters.Tag))
-			{
-				yield return TagName;
-			}
-
-			foreach (string TagName in FindTagNamesFromList(Parameters.TagReferences))
-			{
-				yield return TagName;
-			}*/
+			/*			foreach (string TagName in FindTagNamesFromList(Parameters.Tag))
+						{
+							yield return TagName;
+						}
+
+						foreach (string TagName in FindTagNamesFromList(Parameters.TagReferences))
+						{
+							yield return TagName;
+						}*/
 		}
 	}
 }
diff --git a/Engine/Source/Programs/AutomationTool/BuildGraph/Tasks/GatherBuildProductsFromFileTask.cs b/Engine/Source/Programs/AutomationTool/BuildGraph/Tasks/GatherBuildProductsFromFileTask.cs
--- a/Engine/Source/Programs/AutomationTool/BuildGraph/Tasks/GatherBuildProductsFromFileTask.cs
+++ b/Engine/Source/Programs/AutomationTool/BuildGraph/Tasks/GatherBuildProductsFromFileTask.cs
@@ -3,18 +3,13 @@
 using System;
 using System.Collections.Generic;
 using System.IO;
-using System.Text;
 using System.Threading.Tasks;
 using System.Xml;
-
 using EpicGames.Core;
-using UnrealBuildTool;
 using Microsoft.Extensions.Logging;
 
-using static AutomationTool.CommandUtils;
-
 namespace AutomationTool.Tasks
-{	
+{
 	/// <summary>
 	/// Parameters for <see cref="GatherBuildProductsFromFileTask"/>
 	/// </summary>
@@ -30,37 +25,35 @@ namespace AutomationTool.Tasks
 	[TaskElement("GatherBuildProductsFromFile", typeof(GatherBuildProductsFromFileTaskParameters))]
 	class GatherBuildProductsFromFileTask : BgTaskImpl
 	{
-		public GatherBuildProductsFromFileTask(GatherBuildProductsFromFileTaskParameters InParameters)
+		public GatherBuildProductsFromFileTask(GatherBuildProductsFromFileTaskParameters parameters)
 		{
-			Parameters = InParameters;
+			_parameters = parameters;
 		}
 
-		public override async Task ExecuteAsync(JobContext Job, HashSet<FileReference> BuildProducts, Dictionary<string, HashSet<FileReference>> TagNameToFileSet)
+		public override async Task ExecuteAsync(JobContext job, HashSet<FileReference> buildProducts, Dictionary<string, HashSet<FileReference>> tagNameToFileSet)
 		{
-			List<FileReference> CleanupFiles = new List<FileReference>();
-
-			Logger.LogInformation("Gathering BuildProducts from {Arg0}...", Parameters.BuildProductsFile);
+			Logger.LogInformation("Gathering BuildProducts from {Arg0}...", _parameters.BuildProductsFile);
 
 			try
 			{
-				var FileBuildProducts = await File.ReadAllLinesAsync(Parameters.BuildProductsFile);
-				foreach(var BuildProduct in FileBuildProducts)
+				var fileBuildProducts = await File.ReadAllLinesAsync(_parameters.BuildProductsFile);
+				foreach (var buildProduct in fileBuildProducts)
 				{
-					Logger.LogInformation("Adding file to build products: {BuildProduct}", BuildProduct);
-					BuildProducts.Add(new FileReference(BuildProduct));
+					Logger.LogInformation("Adding file to build products: {BuildProduct}", buildProduct);
+					buildProducts.Add(new FileReference(buildProduct));
 				}
 			}
-			catch (Exception Ex)
+			catch (Exception ex)
 			{
-				Logger.LogInformation("Failed to gather build products: {Arg0}", Ex.Message);
+				Logger.LogInformation("Failed to gather build products: {Arg0}", ex.Message);
 			}
 		}
-		
-		public override void Write(XmlWriter Writer)
+
+		public override void Write(XmlWriter writer)
 		{
-			Write(Writer, Parameters);
+			Write(writer, _parameters);
 		}
-		
+
 		public override IEnumerable<string> FindConsumedTagNames()
 		{
 			yield break;
@@ -71,6 +64,6 @@ namespace AutomationTool.Tasks
 			yield break;
 		}
 
-		public GatherBuildProductsFromFileTaskParameters Parameters;
+		public GatherBuildProductsFromFileTaskParameters _parameters;
 	}
 }
diff --git a/Engine/Source/Programs/AutomationTool/BuildGraph/Tasks/GitCloneTask.cs b/Engine/Source/Programs/AutomationTool/BuildGraph/Tasks/GitCloneTask.cs
--- a/Engine/Source/Programs/AutomationTool/BuildGraph/Tasks/GitCloneTask.cs
+++ b/Engine/Source/Programs/AutomationTool/BuildGraph/Tasks/GitCloneTask.cs
@@ -1,13 +1,11 @@
 // Copyright Epic Games, Inc. All Rights Reserved.
 
-using EpicGames.Core;
-using System;
 using System.Collections.Generic;
-using System.Text;
 using System.Threading.Tasks;
 using System.Xml;
-using UnrealBuildBase;
+using EpicGames.Core;
 using Microsoft.Extensions.Logging;
+using UnrealBuildBase;
 
 namespace AutomationTool.Tasks
 {
@@ -50,69 +48,69 @@ namespace AutomationTool.Tasks
 		/// <summary>
 		/// Parameters for this task
 		/// </summary>
-		GitCloneTaskParameters Parameters;
+		readonly GitCloneTaskParameters _parameters;
 
 		/// <summary>
 		/// Construct a Git task
 		/// </summary>
-		/// <param name="InParameters">Parameters for the task</param>
-		public GitCloneTask(GitCloneTaskParameters InParameters)
+		/// <param name="parameters">Parameters for the task</param>
+		public GitCloneTask(GitCloneTaskParameters parameters)
 		{
-			Parameters = InParameters;
+			_parameters = parameters;
 		}
 
 		/// <summary>
-		/// Execute the task.
+		/// ExecuteAsync the task.
 		/// </summary>
-		/// <param name="Job">Information about the current job</param>
-		/// <param name="BuildProducts">Set of build products produced by this node.</param>
-		/// <param name="TagNameToFileSet">Mapping from tag names to the set of files they include</param>
-		public override async Task ExecuteAsync(JobContext Job, HashSet<FileReference> BuildProducts, Dictionary<string, HashSet<FileReference>> TagNameToFileSet)
+		/// <param name="job">Information about the current job</param>
+		/// <param name="buildProducts">Set of build products produced by this node.</param>
+		/// <param name="tagNameToFileSet">Mapping from tag names to the set of files they include</param>
+		public override async Task ExecuteAsync(JobContext job, HashSet<FileReference> buildProducts, Dictionary<string, HashSet<FileReference>> tagNameToFileSet)
 		{
-			FileReference GitExe = CommandUtils.FindToolInPath("git");
-			if(GitExe == null)
+			FileReference gitExe = CommandUtils.FindToolInPath("git");
+			if (gitExe == null)
 			{
 				throw new AutomationException("Unable to find path to Git. Check you have it installed, and it is on your PATH.");
 			}
 
-			DirectoryReference Dir = ResolveDirectory(Parameters.Dir);
-			Logger.LogInformation("Cloning Git repository into {Dir}", Parameters.Dir);
-			using (LogIndentScope Scope = new LogIndentScope("  "))
+			DirectoryReference dir = ResolveDirectory(_parameters.Dir);
+			Logger.LogInformation("Cloning Git repository into {Dir}", _parameters.Dir);
+			using (LogIndentScope scope = new LogIndentScope("  "))
 			{
-				DirectoryReference GitDir = DirectoryReference.Combine(Dir, ".git");
-				if (!FileReference.Exists(FileReference.Combine(GitDir, "HEAD")))
+				DirectoryReference gitDir = DirectoryReference.Combine(dir, ".git");
+				if (!FileReference.Exists(FileReference.Combine(gitDir, "HEAD")))
 				{
-					await RunGitAsync(GitExe, $"init \"{Dir}\"", Unreal.RootDirectory);
+					await RunGit(gitExe, $"init \"{dir}\"", Unreal.RootDirectory);
 				}
 
-				if (Parameters.ConfigFile != null)
+				if (_parameters.ConfigFile != null)
 				{
-					CommandUtils.CopyFile(Parameters.ConfigFile, FileReference.Combine(GitDir, "config").FullName);
+					CommandUtils.CopyFile(_parameters.ConfigFile, FileReference.Combine(gitDir, "config").FullName);
 				}
 
-				if (Parameters.Remote != null)
+				if (_parameters.Remote != null)
 				{
-					await RunGitAsync(GitExe, $"remote add origin {Parameters.Remote}", Dir);
+					await RunGit(gitExe, $"remote add origin {_parameters.Remote}", dir);
 				}
 
-				await RunGitAsync(GitExe, "clean -dxf", Dir);
-				await RunGitAsync(GitExe, "fetch --all", Dir);
-				await RunGitAsync(GitExe, $"reset --hard {Parameters.Branch}", Dir);
+				await RunGit(gitExe, "clean -dxf", dir);
+				await RunGit(gitExe, "fetch --all", dir);
+				await RunGit(gitExe, $"reset --hard {_parameters.Branch}", dir);
 			}
 		}
 
 		/// <summary>
 		/// Runs a git command
 		/// </summary>
-		/// <param name="ToolFile"></param>
-		/// <param name="Arguments"></param>
-		/// <param name="WorkingDir"></param>
-		static Task RunGitAsync(FileReference ToolFile, string Arguments, DirectoryReference WorkingDir)
+		/// <param name="toolFile"></param>
+		/// <param name="arguments"></param>
+		/// <param name="workingDir"></param>
+		static Task RunGit(FileReference toolFile, string arguments, DirectoryReference workingDir)
 		{
-			IProcessResult Result = CommandUtils.Run(ToolFile.FullName, Arguments, WorkingDir: WorkingDir.FullName);
-			if (Result.ExitCode != 0)
+			IProcessResult result = CommandUtils.Run(toolFile.FullName, arguments, WorkingDir: workingDir.FullName);
+			if (result.ExitCode != 0)
 			{
-				throw new AutomationException("Git terminated with an exit code indicating an error ({0})", Result.ExitCode);
+				throw new AutomationException("Git terminated with an exit code indicating an error ({0})", result.ExitCode);
 			}
 			return Task.CompletedTask;
 		}
@@ -120,9 +118,9 @@ namespace AutomationTool.Tasks
 		/// <summary>
 		/// Output this task out to an XML writer.
 		/// </summary>
-		public override void Write(XmlWriter Writer)
+		public override void Write(XmlWriter writer)
 		{
-			Write(Writer, Parameters);
+			Write(writer, _parameters);
 		}
 
 		/// <summary>
diff --git a/Engine/Source/Programs/AutomationTool/BuildGraph/Tasks/GitTask.cs b/Engine/Source/Programs/AutomationTool/BuildGraph/Tasks/GitTask.cs
--- a/Engine/Source/Programs/AutomationTool/BuildGraph/Tasks/GitTask.cs
+++ b/Engine/Source/Programs/AutomationTool/BuildGraph/Tasks/GitTask.cs
@@ -1,11 +1,9 @@
 // Copyright Epic Games, Inc. All Rights Reserved.
 
-using EpicGames.Core;
-using System;
 using System.Collections.Generic;
-using System.Text;
 using System.Threading.Tasks;
 using System.Xml;
+using EpicGames.Core;
 
 namespace AutomationTool.Tasks
 {
@@ -39,38 +37,35 @@ namespace AutomationTool.Tasks
 	[TaskElement("Git", typeof(GitTaskParameters))]
 	public class GitTask : BgTaskImpl
 	{
-		/// <summary>
-		/// Parameters for this task
-		/// </summary>
-		GitTaskParameters Parameters;
+		readonly GitTaskParameters _parameters;
 
 		/// <summary>
 		/// Construct a Git task
 		/// </summary>
-		/// <param name="InParameters">Parameters for the task</param>
-		public GitTask(GitTaskParameters InParameters)
+		/// <param name="parameters">Parameters for the task</param>
+		public GitTask(GitTaskParameters parameters)
 		{
-			Parameters = InParameters;
+			_parameters = parameters;
 		}
 
 		/// <summary>
-		/// Execute the task.
+		/// ExecuteAsync the task.
 		/// </summary>
-		/// <param name="Job">Information about the current job</param>
-		/// <param name="BuildProducts">Set of build products produced by this node.</param>
-		/// <param name="TagNameToFileSet">Mapping from tag names to the set of files they include</param>
-		public override Task ExecuteAsync(JobContext Job, HashSet<FileReference> BuildProducts, Dictionary<string, HashSet<FileReference>> TagNameToFileSet)
+		/// <param name="job">Information about the current job</param>
+		/// <param name="buildProducts">Set of build products produced by this node.</param>
+		/// <param name="tagNameToFileSet">Mapping from tag names to the set of files they include</param>
+		public override Task ExecuteAsync(JobContext job, HashSet<FileReference> buildProducts, Dictionary<string, HashSet<FileReference>> tagNameToFileSet)
 		{
-			FileReference ToolFile = CommandUtils.FindToolInPath("git");
-			if(ToolFile == null)
+			FileReference toolFile = CommandUtils.FindToolInPath("git");
+			if (toolFile == null)
 			{
 				throw new AutomationException("Unable to find path to Git. Check you have it installed, and it is on your PATH.");
 			}
 
-			IProcessResult Result = CommandUtils.Run(ToolFile.FullName, Parameters.Arguments, WorkingDir: Parameters.BaseDir);
-			if (Result.ExitCode < 0 || Result.ExitCode >= Parameters.ErrorLevel)
+			IProcessResult result = CommandUtils.Run(toolFile.FullName, _parameters.Arguments, WorkingDir: _parameters.BaseDir);
+			if (result.ExitCode < 0 || result.ExitCode >= _parameters.ErrorLevel)
 			{
-				throw new AutomationException("Git terminated with an exit code indicating an error ({0})", Result.ExitCode);
+				throw new AutomationException("Git terminated with an exit code indicating an error ({0})", result.ExitCode);
 			}
 
 			return Task.CompletedTask;
@@ -79,9 +74,9 @@ namespace AutomationTool.Tasks
 		/// <summary>
 		/// Output this task out to an XML writer.
 		/// </summary>
-		public override void Write(XmlWriter Writer)
+		public override void Write(XmlWriter writer)
 		{
-			Write(Writer, Parameters);
+			Write(writer, _parameters);
 		}
 
 		/// <summary>
diff --git a/Engine/Source/Programs/AutomationTool/BuildGraph/Tasks/HelmTask.cs b/Engine/Source/Programs/AutomationTool/BuildGraph/Tasks/HelmTask.cs
--- a/Engine/Source/Programs/AutomationTool/BuildGraph/Tasks/HelmTask.cs
+++ b/Engine/Source/Programs/AutomationTool/BuildGraph/Tasks/HelmTask.cs
@@ -1,13 +1,10 @@
 // Copyright Epic Games, Inc. All Rights Reserved.
 
-using AutomationTool;
-using EpicGames.Core;
 using System;
 using System.Collections.Generic;
-using System.IO;
-using System.Text;
 using System.Threading.Tasks;
 using System.Xml;
+using EpicGames.Core;
 using UnrealBuildBase;
 
 namespace AutomationTool.Tasks
@@ -90,78 +87,75 @@ namespace AutomationTool.Tasks
 	[TaskElement("Helm", typeof(HelmTaskParameters))]
 	public class HelmTask : SpawnTaskBase
 	{
-		/// <summary>
-		/// Parameters for this task
-		/// </summary>
-		HelmTaskParameters Parameters;
+		readonly HelmTaskParameters _parameters;
 
 		/// <summary>
 		/// Construct a Helm task
 		/// </summary>
-		/// <param name="InParameters">Parameters for the task</param>
-		public HelmTask(HelmTaskParameters InParameters)
+		/// <param name="parameters">Parameters for the task</param>
+		public HelmTask(HelmTaskParameters parameters)
 		{
-			Parameters = InParameters;
+			_parameters = parameters;
 		}
 
 		/// <summary>
-		/// Execute the task.
+		/// ExecuteAsync the task.
 		/// </summary>
-		/// <param name="Job">Information about the current job</param>
-		/// <param name="BuildProducts">Set of build products produced by this node.</param>
-		/// <param name="TagNameToFileSet">Mapping from tag names to the set of files they include</param>
-		public override async Task ExecuteAsync(JobContext Job, HashSet<FileReference> BuildProducts, Dictionary<string, HashSet<FileReference>> TagNameToFileSet)
+		/// <param name="job">Information about the current job</param>
+		/// <param name="buildProducts">Set of build products produced by this node.</param>
+		/// <param name="tagNameToFileSet">Mapping from tag names to the set of files they include</param>
+		public override async Task ExecuteAsync(JobContext job, HashSet<FileReference> buildProducts, Dictionary<string, HashSet<FileReference>> tagNameToFileSet)
 		{
 
 			// Build the argument list
-			List<string> Arguments = new List<string>();
-			Arguments.Add("upgrade");
-			Arguments.Add(Parameters.Deployment);
-			Arguments.Add(new FileReference(Parameters.Chart).FullName);
-			Arguments.Add("--install");
-			Arguments.Add("--reset-values");
-			if(Parameters.Namespace != null)
+			List<string> arguments = new List<string>();
+			arguments.Add("upgrade");
+			arguments.Add(_parameters.Deployment);
+			arguments.Add(new FileReference(_parameters.Chart).FullName);
+			arguments.Add("--install");
+			arguments.Add("--reset-values");
+			if (_parameters.Namespace != null)
 			{
-				Arguments.Add("--namespace");
-				Arguments.Add(Parameters.Namespace);
+				arguments.Add("--namespace");
+				arguments.Add(_parameters.Namespace);
 			}
-			if (Parameters.KubeContext != null)
+			if (_parameters.KubeContext != null)
 			{
-				Arguments.Add("--kube-context");
-				Arguments.Add(Parameters.KubeContext);
+				arguments.Add("--kube-context");
+				arguments.Add(_parameters.KubeContext);
 			}
-			if (Parameters.KubeConfig != null)
+			if (_parameters.KubeConfig != null)
 			{
-				Arguments.Add("--kubeconfig");
-				Arguments.Add(Parameters.KubeConfig);
+				arguments.Add("--kubeconfig");
+				arguments.Add(_parameters.KubeConfig);
 			}
-			if (!string.IsNullOrEmpty(Parameters.Values))
+			if (!string.IsNullOrEmpty(_parameters.Values))
 			{
-				foreach (string Value in SplitDelimitedList(Parameters.Values))
+				foreach (string value in SplitDelimitedList(_parameters.Values))
 				{
-					Arguments.Add("--set");
-					Arguments.Add(Value);
+					arguments.Add("--set");
+					arguments.Add(value);
 				}
 			}
-			if (!String.IsNullOrEmpty(Parameters.ValuesFile))
+			if (!String.IsNullOrEmpty(_parameters.ValuesFile))
 			{
-				foreach (FileReference ValuesFile in ResolveFilespec(Unreal.RootDirectory, Parameters.ValuesFile, TagNameToFileSet))
+				foreach (FileReference valuesFile in ResolveFilespec(Unreal.RootDirectory, _parameters.ValuesFile, tagNameToFileSet))
 				{
-					Arguments.Add("--values");
-					Arguments.Add(ValuesFile.FullName);
+					arguments.Add("--values");
+					arguments.Add(valuesFile.FullName);
 				}
 			}
 
-			string AdditionalArguments = String.IsNullOrEmpty(Parameters.Arguments) ? String.Empty : $" {Parameters.Arguments}";
-			await SpawnTaskBase.ExecuteAsync("helm", CommandLineArguments.Join(Arguments) + AdditionalArguments, WorkingDir: Parameters.WorkingDir, EnvVars: ParseEnvVars(Parameters.Environment, Parameters.EnvironmentFile));
+			string additionalArguments = String.IsNullOrEmpty(_parameters.Arguments) ? String.Empty : $" {_parameters.Arguments}";
+			await SpawnTaskBase.ExecuteAsync("helm", CommandLineArguments.Join(arguments) + additionalArguments, workingDir: _parameters.WorkingDir, envVars: ParseEnvVars(_parameters.Environment, _parameters.EnvironmentFile));
 		}
 
 		/// <summary>
 		/// Output this task out to an XML writer.
 		/// </summary>
-		public override void Write(XmlWriter Writer)
+		public override void Write(XmlWriter writer)
 		{
-			Write(Writer, Parameters);
+			Write(writer, _parameters);
 		}
 
 		/// <summary>
diff --git a/Engine/Source/Programs/AutomationTool/BuildGraph/Tasks/HordeCreateReportTask.cs b/Engine/Source/Programs/AutomationTool/BuildGraph/Tasks/HordeCreateReportTask.cs
--- a/Engine/Source/Programs/AutomationTool/BuildGraph/Tasks/HordeCreateReportTask.cs
+++ b/Engine/Source/Programs/AutomationTool/BuildGraph/Tasks/HordeCreateReportTask.cs
@@ -1,15 +1,13 @@
 // Copyright Epic Games, Inc. All Rights Reserved.
 
-using EpicGames.Core;
-using Microsoft.Extensions.Logging;
-using System;
 using System.Collections.Generic;
 using System.IO;
 using System.Linq;
-using System.Text;
 using System.Text.Json;
 using System.Threading.Tasks;
 using System.Xml;
+using EpicGames.Core;
+using Microsoft.Extensions.Logging;
 
 namespace AutomationTool.Tasks
 {
@@ -52,51 +50,51 @@ namespace AutomationTool.Tasks
 		/// <summary>
 		/// Parameters for this task.
 		/// </summary>
-		HordeCreateReportTaskParameters Parameters;
+		readonly HordeCreateReportTaskParameters _parameters;
 
 		/// <summary>
 		/// Constructor.
 		/// </summary>
-		/// <param name="InParameters">Parameters for this task.</param>
-		public HordeCreateReportTask(HordeCreateReportTaskParameters InParameters)
+		/// <param name="parameters">Parameters for this task.</param>
+		public HordeCreateReportTask(HordeCreateReportTaskParameters parameters)
 		{
-			Parameters = InParameters;
+			_parameters = parameters;
 		}
 
 		/// <summary>
-		/// Execute the task.
+		/// ExecuteAsync the task.
 		/// </summary>
-		/// <param name="Job">Information about the current job.</param>
-		/// <param name="BuildProducts">Set of build products produced by this node.</param>
-		/// <param name="TagNameToFileSet">Mapping from tag names to the set of files they include.</param>
-		public override async Task ExecuteAsync(JobContext Job, HashSet<FileReference> BuildProducts, Dictionary<string, HashSet<FileReference>> TagNameToFileSet)
+		/// <param name="job">Information about the current job.</param>
+		/// <param name="buildProducts">Set of build products produced by this node.</param>
+		/// <param name="tagNameToFileSet">Mapping from tag names to the set of files they include.</param>
+		public override async Task ExecuteAsync(JobContext job, HashSet<FileReference> buildProducts, Dictionary<string, HashSet<FileReference>> tagNameToFileSet)
 		{
-			FileReference ReportTextFile = FileReference.Combine(new DirectoryReference(CommandUtils.CmdEnv.LogFolder), $"{Parameters.Name}.md");
-			await FileReference.WriteAllTextAsync(ReportTextFile, Parameters.Text);
+			FileReference reportTextFile = FileReference.Combine(new DirectoryReference(CommandUtils.CmdEnv.LogFolder), $"{_parameters.Name}.md");
+			await FileReference.WriteAllTextAsync(reportTextFile, _parameters.Text);
 
-			FileReference ReportJsonFile = FileReference.Combine(new DirectoryReference(CommandUtils.CmdEnv.LogFolder), $"{Parameters.Name}.report.json");
-			using (FileStream ReportJsonStream = FileReference.Open(ReportJsonFile, FileMode.Create, FileAccess.Write, FileShare.Read))
+			FileReference reportJsonFile = FileReference.Combine(new DirectoryReference(CommandUtils.CmdEnv.LogFolder), $"{_parameters.Name}.report.json");
+			using (FileStream reportJsonStream = FileReference.Open(reportJsonFile, FileMode.Create, FileAccess.Write, FileShare.Read))
 			{
-				using (Utf8JsonWriter Writer = new Utf8JsonWriter(ReportJsonStream))
+				using (Utf8JsonWriter writer = new Utf8JsonWriter(reportJsonStream))
 				{
-					Writer.WriteStartObject();
-					Writer.WriteString("scope", Parameters.Scope);
-					Writer.WriteString("name", Parameters.Name);
-					Writer.WriteString("placement", Parameters.Placement);
-					Writer.WriteString("fileName", ReportTextFile.GetFileName());
-					Writer.WriteEndObject();
+					writer.WriteStartObject();
+					writer.WriteString("scope", _parameters.Scope);
+					writer.WriteString("name", _parameters.Name);
+					writer.WriteString("placement", _parameters.Placement);
+					writer.WriteString("fileName", reportTextFile.GetFileName());
+					writer.WriteEndObject();
 				}
 			}
 
-			Logger.LogInformation("Written report to {TextFile} and {JsonFile}: \"{Text}\"", ReportTextFile, ReportJsonFile, Parameters.Text);
+			Logger.LogInformation("Written report to {TextFile} and {JsonFile}: \"{Text}\"", reportTextFile, reportJsonFile, _parameters.Text);
 		}
 
 		/// <summary>
 		/// Output this task out to an XML writer.
 		/// </summary>
-		public override void Write(XmlWriter Writer)
+		public override void Write(XmlWriter writer)
 		{
-			Write(Writer, Parameters);
+			Write(writer, _parameters);
 		}
 
 		/// <summary>
diff --git a/Engine/Source/Programs/AutomationTool/BuildGraph/Tasks/HordeGetSecretsTask.cs b/Engine/Source/Programs/AutomationTool/BuildGraph/Tasks/HordeGetSecretsTask.cs
--- a/Engine/Source/Programs/AutomationTool/BuildGraph/Tasks/HordeGetSecretsTask.cs
+++ b/Engine/Source/Programs/AutomationTool/BuildGraph/Tasks/HordeGetSecretsTask.cs
@@ -1,18 +1,15 @@
 // Copyright Epic Games, Inc. All Rights Reserved.
 
-using EpicGames.Core;
-using EpicGames.Horde;
-using EpicGames.Horde.Secrets;
-using Microsoft.Extensions.DependencyInjection;
-using Microsoft.Extensions.Logging;
 using System;
 using System.Collections.Generic;
-using System.IO;
 using System.Linq;
-using System.Text;
-using System.Text.Json;
 using System.Threading.Tasks;
 using System.Xml;
+using EpicGames.Core;
+using EpicGames.Horde;
+using EpicGames.Horde.Secrets;
+using Microsoft.Extensions.DependencyInjection;
+using Microsoft.Extensions.Logging;
 
 #nullable enable
 
@@ -56,12 +53,12 @@ namespace AutomationTool.Tasks
 		}
 
 		/// <summary>
-		/// Execute the task.
+		/// ExecuteAsync the task.
 		/// </summary>
-		/// <param name="Job">Information about the current job.</param>
-		/// <param name="BuildProducts">Set of build products produced by this node.</param>
-		/// <param name="TagNameToFileSet">Mapping from tag names to the set of files they include.</param>
-		public override async Task ExecuteAsync(JobContext Job, HashSet<FileReference> BuildProducts, Dictionary<string, HashSet<FileReference>> TagNameToFileSet)
+		/// <param name="job">Information about the current job.</param>
+		/// <param name="buildProducts">Set of build products produced by this node.</param>
+		/// <param name="tagNameToFileSet">Mapping from tag names to the set of files they include.</param>
+		public override async Task ExecuteAsync(JobContext job, HashSet<FileReference> buildProducts, Dictionary<string, HashSet<FileReference>> tagNameToFileSet)
 		{
 			// Read the input text
 			FileReference file = ResolveFile(_parameters.File);
diff --git a/Engine/Source/Programs/AutomationTool/BuildGraph/Tasks/KubectlTask.cs b/Engine/Source/Programs/AutomationTool/BuildGraph/Tasks/KubectlTask.cs
--- a/Engine/Source/Programs/AutomationTool/BuildGraph/Tasks/KubectlTask.cs
+++ b/Engine/Source/Programs/AutomationTool/BuildGraph/Tasks/KubectlTask.cs
@@ -1,12 +1,9 @@
 // Copyright Epic Games, Inc. All Rights Reserved.
 
-using EpicGames.Core;
-using System;
 using System.Collections.Generic;
-using System.IO;
-using System.Text;
 using System.Threading.Tasks;
 using System.Xml;
+using EpicGames.Core;
 
 namespace AutomationTool.Tasks
 {
@@ -34,38 +31,35 @@ namespace AutomationTool.Tasks
 	[TaskElement("Kubectl", typeof(KubectlTaskParameters))]
 	public class KubectlTask : BgTaskImpl
 	{
-		/// <summary>
-		/// Parameters for this task
-		/// </summary>
-		KubectlTaskParameters Parameters;
+		readonly KubectlTaskParameters _parameters;
 
 		/// <summary>
 		/// Construct a Kubectl task
 		/// </summary>
-		/// <param name="InParameters">Parameters for the task</param>
-		public KubectlTask(KubectlTaskParameters InParameters)
+		/// <param name="parameters">Parameters for the task</param>
+		public KubectlTask(KubectlTaskParameters parameters)
 		{
-			Parameters = InParameters;
+			_parameters = parameters;
 		}
 
 		/// <summary>
-		/// Execute the task.
+		/// ExecuteAsync the task.
 		/// </summary>
-		/// <param name="Job">Information about the current job</param>
-		/// <param name="BuildProducts">Set of build products produced by this node.</param>
-		/// <param name="TagNameToFileSet">Mapping from tag names to the set of files they include</param>
-		public override Task ExecuteAsync(JobContext Job, HashSet<FileReference> BuildProducts, Dictionary<string, HashSet<FileReference>> TagNameToFileSet)
+		/// <param name="job">Information about the current job</param>
+		/// <param name="buildProducts">Set of build products produced by this node.</param>
+		/// <param name="tagNameToFileSet">Mapping from tag names to the set of files they include</param>
+		public override Task ExecuteAsync(JobContext job, HashSet<FileReference> buildProducts, Dictionary<string, HashSet<FileReference>> tagNameToFileSet)
 		{
-			FileReference KubectlExe = CommandUtils.FindToolInPath("kubectl");
-			if (KubectlExe == null)
+			FileReference kubectlExe = CommandUtils.FindToolInPath("kubectl");
+			if (kubectlExe == null)
 			{
 				throw new AutomationException("Unable to find path to Kubectl. Check you have it installed, and it is on your PATH.");
 			}
 
-			IProcessResult Result = CommandUtils.Run(KubectlExe.FullName, Parameters.Arguments, null, WorkingDir: Parameters.BaseDir);
-			if (Result.ExitCode != 0)
+			IProcessResult result = CommandUtils.Run(kubectlExe.FullName, _parameters.Arguments, null, WorkingDir: _parameters.BaseDir);
+			if (result.ExitCode != 0)
 			{
-				throw new AutomationException("Kubectl terminated with an exit code indicating an error ({0})", Result.ExitCode);
+				throw new AutomationException("Kubectl terminated with an exit code indicating an error ({0})", result.ExitCode);
 			}
 
 			return Task.CompletedTask;
@@ -74,9 +68,9 @@ namespace AutomationTool.Tasks
 		/// <summary>
 		/// Output this task out to an XML writer.
 		/// </summary>
-		public override void Write(XmlWriter Writer)
+		public override void Write(XmlWriter writer)
 		{
-			Write(Writer, Parameters);
+			Write(writer, _parameters);
 		}
 
 		/// <summary>
diff --git a/Engine/Source/Programs/AutomationTool/BuildGraph/Tasks/LogTask.cs b/Engine/Source/Programs/AutomationTool/BuildGraph/Tasks/LogTask.cs
--- a/Engine/Source/Programs/AutomationTool/BuildGraph/Tasks/LogTask.cs
+++ b/Engine/Source/Programs/AutomationTool/BuildGraph/Tasks/LogTask.cs
@@ -1,20 +1,13 @@
 // Copyright Epic Games, Inc. All Rights Reserved.
 
-using AutomationTool;
-using EpicGames.BuildGraph;
 using System;
 using System.Collections.Generic;
-using System.IO;
 using System.Linq;
-using System.Text;
 using System.Threading.Tasks;
 using System.Xml;
 using EpicGames.Core;
-using UnrealBuildTool;
-using UnrealBuildBase;
 using Microsoft.Extensions.Logging;
-
-using static AutomationTool.CommandUtils;
+using UnrealBuildBase;
 
 namespace AutomationTool.Tasks
 {
@@ -51,44 +44,44 @@ namespace AutomationTool.Tasks
 		/// <summary>
 		/// Parameters for the task
 		/// </summary>
-		LogTaskParameters Parameters;
+		readonly LogTaskParameters _parameters;
 
 		/// <summary>
 		/// Constructor.
 		/// </summary>
-		/// <param name="InParameters">Parameters for the task</param>
-		public LogTask(LogTaskParameters InParameters)
+		/// <param name="parameters">Parameters for the task</param>
+		public LogTask(LogTaskParameters parameters)
 		{
-			Parameters = InParameters;
+			_parameters = parameters;
 		}
 
 		/// <summary>
-		/// Execute the task.
+		/// ExecuteAsync the task.
 		/// </summary>
-		/// <param name="Job">Information about the current job</param>
-		/// <param name="BuildProducts">Set of build products produced by this node.</param>
-		/// <param name="TagNameToFileSet">Mapping from tag names to the set of files they include</param>
-		public override async Task ExecuteAsync(JobContext Job, HashSet<FileReference> BuildProducts, Dictionary<string, HashSet<FileReference>> TagNameToFileSet)
+		/// <param name="job">Information about the current job</param>
+		/// <param name="buildProducts">Set of build products produced by this node.</param>
+		/// <param name="tagNameToFileSet">Mapping from tag names to the set of files they include</param>
+		public override async Task ExecuteAsync(JobContext job, HashSet<FileReference> buildProducts, Dictionary<string, HashSet<FileReference>> tagNameToFileSet)
 		{
 			// Print the message
-			if(!String.IsNullOrEmpty(Parameters.Message))
+			if (!String.IsNullOrEmpty(_parameters.Message))
 			{
-				Logger.LogInformation("{Text}", Parameters.Message);
+				Logger.LogInformation("{Text}", _parameters.Message);
 			}
 
 			// Print the contents of the given tag, if specified
-			if(!String.IsNullOrEmpty(Parameters.Files))
+			if (!String.IsNullOrEmpty(_parameters.Files))
 			{
-				HashSet<FileReference> Files = ResolveFilespec(Unreal.RootDirectory, Parameters.Files, TagNameToFileSet);
-				foreach(FileReference File in Files.OrderBy(x => x.FullName))
+				HashSet<FileReference> files = ResolveFilespec(Unreal.RootDirectory, _parameters.Files, tagNameToFileSet);
+				foreach (FileReference file in files.OrderBy(x => x.FullName))
 				{
-					Logger.LogInformation("  {Arg0}", File.FullName);
-					if(Parameters.IncludeContents)
+					Logger.LogInformation("  {Arg0}", file.FullName);
+					if (_parameters.IncludeContents)
 					{
-						string[] lines = await System.IO.File.ReadAllLinesAsync(File.FullName);
-						foreach (string Line in lines)
+						string[] lines = await System.IO.File.ReadAllLinesAsync(file.FullName);
+						foreach (string line in lines)
 						{
-							Logger.LogInformation("    {Line}", Line);
+							Logger.LogInformation("    {Line}", line);
 						}
 					}
 				}
@@ -98,9 +91,9 @@ namespace AutomationTool.Tasks
 		/// <summary>
 		/// Output this task out to an XML writer.
 		/// </summary>
-		public override void Write(XmlWriter Writer)
+		public override void Write(XmlWriter writer)
 		{
-			Write(Writer, Parameters);
+			Write(writer, _parameters);
 		}
 
 		/// <summary>
@@ -109,7 +102,7 @@ namespace AutomationTool.Tasks
 		/// <returns>The tag names which are read by this task</returns>
 		public override IEnumerable<string> FindConsumedTagNames()
 		{
-			return FindTagNamesFromFilespec(Parameters.Files);
+			return FindTagNamesFromFilespec(_parameters.Files);
 		}
 
 		/// <summary>
diff --git a/Engine/Source/Programs/AutomationTool/BuildGraph/Tasks/ModifyConfigTask.cs b/Engine/Source/Programs/AutomationTool/BuildGraph/Tasks/ModifyConfigTask.cs
--- a/Engine/Source/Programs/AutomationTool/BuildGraph/Tasks/ModifyConfigTask.cs
+++ b/Engine/Source/Programs/AutomationTool/BuildGraph/Tasks/ModifyConfigTask.cs
@@ -1,11 +1,7 @@
 // Copyright Epic Games, Inc. All Rights Reserved.
 
-using AutomationTool;
-using EpicGames.BuildGraph;
 using System;
 using System.Collections.Generic;
-using System.Linq;
-using System.Text;
 using System.Threading.Tasks;
 using System.Xml;
 using EpicGames.Core;
@@ -55,64 +51,61 @@ namespace AutomationTool.Tasks
 	[TaskElement("ModifyConfig", typeof(ModifyConfigTaskParameters))]
 	public class ModifyConfigTask : BgTaskImpl
 	{
-		/// <summary>
-		/// Parameters for this task
-		/// </summary>
-		ModifyConfigTaskParameters Parameters;
+		readonly ModifyConfigTaskParameters _parameters;
 
 		/// <summary>
 		/// Constructor
 		/// </summary>
-		/// <param name="InParameters">Parameters for this task</param>
-		public ModifyConfigTask(ModifyConfigTaskParameters InParameters)
+		/// <param name="parameters">Parameters for this task</param>
+		public ModifyConfigTask(ModifyConfigTaskParameters parameters)
 		{
-			Parameters = InParameters;
+			_parameters = parameters;
 		}
 
 		/// <summary>
-		/// Execute the task.
+		/// ExecuteAsync the task.
 		/// </summary>
-		/// <param name="Job">Information about the current job</param>
-		/// <param name="BuildProducts">Set of build products produced by this node.</param>
-		/// <param name="TagNameToFileSet">Mapping from tag names to the set of files they include</param>
-		public override Task ExecuteAsync(JobContext Job, HashSet<FileReference> BuildProducts, Dictionary<string, HashSet<FileReference>> TagNameToFileSet)
+		/// <param name="job">Information about the current job</param>
+		/// <param name="buildProducts">Set of build products produced by this node.</param>
+		/// <param name="tagNameToFileSet">Mapping from tag names to the set of files they include</param>
+		public override Task ExecuteAsync(JobContext job, HashSet<FileReference> buildProducts, Dictionary<string, HashSet<FileReference>> tagNameToFileSet)
 		{
-			FileReference ConfigFileLocation = ResolveFile(Parameters.File);
+			FileReference configFileLocation = ResolveFile(_parameters.File);
 
-			ConfigFile ConfigFile;
-			if(FileReference.Exists(ConfigFileLocation))
+			ConfigFile configFile;
+			if (FileReference.Exists(configFileLocation))
 			{
-				ConfigFile = new ConfigFile(ConfigFileLocation);
+				configFile = new ConfigFile(configFileLocation);
 			}
 			else
 			{
-				ConfigFile = new ConfigFile();
+				configFile = new ConfigFile();
 			}
 
-			ConfigFileSection Section = ConfigFile.FindOrAddSection(Parameters.Section);
-			Section.Lines.RemoveAll(x => String.Equals(x.Key, Parameters.Key, StringComparison.OrdinalIgnoreCase));
-			Section.Lines.Add(new ConfigLine(ConfigLineAction.Set, Parameters.Key, Parameters.Value));
+			ConfigFileSection section = configFile.FindOrAddSection(_parameters.Section);
+			section.Lines.RemoveAll(x => String.Equals(x.Key, _parameters.Key, StringComparison.OrdinalIgnoreCase));
+			section.Lines.Add(new ConfigLine(ConfigLineAction.Set, _parameters.Key, _parameters.Value));
 
-			FileReference.MakeWriteable(ConfigFileLocation);
-			ConfigFile.Write(ConfigFileLocation);
+			FileReference.MakeWriteable(configFileLocation);
+			configFile.Write(configFileLocation);
 
 			// Apply the optional tag to the produced archive
-			foreach (string TagName in FindTagNamesFromList(Parameters.Tag))
+			foreach (string tagName in FindTagNamesFromList(_parameters.Tag))
 			{
-				FindOrAddTagSet(TagNameToFileSet, TagName).Add(ConfigFileLocation);
+				FindOrAddTagSet(tagNameToFileSet, tagName).Add(configFileLocation);
 			}
 
 			// Add the archive to the set of build products
-			BuildProducts.Add(ConfigFileLocation);
+			buildProducts.Add(configFileLocation);
 			return Task.CompletedTask;
 		}
 
 		/// <summary>
 		/// Output this task out to an XML writer.
 		/// </summary>
-		public override void Write(XmlWriter Writer)
+		public override void Write(XmlWriter writer)
 		{
-			Write(Writer, Parameters);
+			Write(writer, _parameters);
 		}
 
 		/// <summary>
@@ -121,7 +114,7 @@ namespace AutomationTool.Tasks
 		/// <returns>The tag names which are read by this task</returns>
 		public override IEnumerable<string> FindConsumedTagNames()
 		{
-			return FindTagNamesFromFilespec(Parameters.File);
+			return FindTagNamesFromFilespec(_parameters.File);
 		}
 
 		/// <summary>
@@ -130,7 +123,7 @@ namespace AutomationTool.Tasks
 		/// <returns>The tag names which are modified by this task</returns>
 		public override IEnumerable<string> FindProducedTagNames()
 		{
-			return FindTagNamesFromList(Parameters.Tag);
+			return FindTagNamesFromList(_parameters.Tag);
 		}
 	}
 }
diff --git a/Engine/Source/Programs/AutomationTool/BuildGraph/Tasks/ModifyJsonValue.cs b/Engine/Source/Programs/AutomationTool/BuildGraph/Tasks/ModifyJsonValue.cs
--- a/Engine/Source/Programs/AutomationTool/BuildGraph/Tasks/ModifyJsonValue.cs
+++ b/Engine/Source/Programs/AutomationTool/BuildGraph/Tasks/ModifyJsonValue.cs
@@ -1,16 +1,14 @@
 // Copyright Epic Games, Inc. All Rights Reserved.
 
-using AutomationTool;
-using EpicGames.BuildGraph;
 using System.Collections.Generic;
-using System.Linq;
-using System.Xml;
 using System.IO;
+using System.Linq;
 using System.Text;
-using EpicGames.Core;
-using UnrealBuildBase;
 using System.Text.Json;
 using System.Threading.Tasks;
+using System.Xml;
+using EpicGames.Core;
+using UnrealBuildBase;
 
 namespace AutomationTool.Tasks
 {
@@ -44,53 +42,53 @@ namespace AutomationTool.Tasks
 	[TaskElement("ModifyJsonValue", typeof(ModifyJsonValueParameters))]
 	public class ModifyJsonValue : BgTaskImpl
 	{
-		ModifyJsonValueParameters Parameters;
+		readonly ModifyJsonValueParameters _parameters;
 
 		/// <summary>
 		/// Create a new ModifyJsonValue.
 		/// </summary>
-		/// <param name="InParameters">Parameters for this task.</param>
-		public ModifyJsonValue(ModifyJsonValueParameters InParameters)
+		/// <param name="parameters">Parameters for this task.</param>
+		public ModifyJsonValue(ModifyJsonValueParameters parameters)
 		{
-			Parameters = InParameters;
+			_parameters = parameters;
 		}
 
 		/// <summary>
 		/// Placeholder comment
 		/// </summary>
-		public override async Task ExecuteAsync(JobContext Job, HashSet<FileReference> BuildProducts, Dictionary<string, HashSet<FileReference>> TagNameToFileSet)
+		public override async Task ExecuteAsync(JobContext job, HashSet<FileReference> buildProducts, Dictionary<string, HashSet<FileReference>> tagNameToFileSet)
 		{
-			string[] Keys = Parameters.KeyPath.Split('.');
-			if (Keys.Length == 0)
-            {
+			string[] keys = _parameters.KeyPath.Split('.');
+			if (keys.Length == 0)
+			{
 				return;
-            }
-			HashSet<FileReference> Files = ResolveFilespec(Unreal.RootDirectory, Parameters.Files, TagNameToFileSet);
-			foreach (var JsonFile in Files.Select(f => f.FullName))
+			}
+			HashSet<FileReference> files = ResolveFilespec(Unreal.RootDirectory, _parameters.Files, tagNameToFileSet);
+			foreach (var jsonFile in files.Select(f => f.FullName))
 			{
-				var OldContents = await File.ReadAllTextAsync(JsonFile);
-				var ParamObj = fastJSON.JSON.Instance.Parse(OldContents) as IDictionary<string, object>;
-				var CurrObj = ParamObj;
-				for (int i = 0; i < Keys.Length - 1; i++)
+				var oldContents = await File.ReadAllTextAsync(jsonFile);
+				var paramObj = fastJSON.JSON.Instance.Parse(oldContents) as IDictionary<string, object>;
+				var currObj = paramObj;
+				for (int i = 0; i < keys.Length - 1; i++)
 				{
-					if (!CurrObj.TryGetValue(Keys[i], out object NextNode))
-						CurrObj[Keys[i]] = NextNode = new Dictionary<string, object>();
-					CurrObj = (IDictionary<string, object>)NextNode;
+					if (!currObj.TryGetValue(keys[i], out object nextNode))
+						currObj[keys[i]] = nextNode = new Dictionary<string, object>();
+					currObj = (IDictionary<string, object>)nextNode;
 				}
 
-				CurrObj[Keys[Keys.Length - 1]] = Parameters.NewValue;
+				currObj[keys[keys.Length - 1]] = _parameters.NewValue;
 
-				var NewContents = JsonSerializer.Serialize(ParamObj, new JsonSerializerOptions { WriteIndented = true });
-				await File.WriteAllTextAsync(JsonFile, NewContents, new UTF8Encoding(encoderShouldEmitUTF8Identifier: false));
+				var newContents = JsonSerializer.Serialize(paramObj, new JsonSerializerOptions { WriteIndented = true });
+				await File.WriteAllTextAsync(jsonFile, newContents, new UTF8Encoding(encoderShouldEmitUTF8Identifier: false));
 			}
 		}
 
 		/// <summary>
 		/// Placeholder comment
 		/// </summary>
-		public override void Write(XmlWriter Writer)
+		public override void Write(XmlWriter writer)
 		{
-			Write(Writer, Parameters);
+			Write(writer, _parameters);
 		}
 
 		/// <summary>
@@ -98,9 +96,9 @@ namespace AutomationTool.Tasks
 		/// </summary>
 		public override IEnumerable<string> FindConsumedTagNames()
 		{
-			foreach (string TagName in FindTagNamesFromFilespec(Parameters.Files))
+			foreach (string tagName in FindTagNamesFromFilespec(_parameters.Files))
 			{
-				yield return TagName;
+				yield return tagName;
 			}
 		}
 
