EpicGames.BuildGraph: Apply standard naming conventions to BuildGraph.Automation.

diff --git a/Engine/Source/Programs/AutomationTool/BuildGraph/Tasks/MoveTask.cs b/Engine/Source/Programs/AutomationTool/BuildGraph/Tasks/MoveTask.cs
--- a/Engine/Source/Programs/AutomationTool/BuildGraph/Tasks/MoveTask.cs
+++ b/Engine/Source/Programs/AutomationTool/BuildGraph/Tasks/MoveTask.cs
@@ -1,19 +1,13 @@
 // Copyright Epic Games, Inc. All Rights Reserved.
 
-using AutomationTool;
-using EpicGames.BuildGraph;
 using System;
 using System.Collections.Generic;
 using System.Linq;
-using System.Text;
 using System.Threading.Tasks;
 using System.Xml;
 using EpicGames.Core;
-using UnrealBuildTool;
-using UnrealBuildBase;
 using Microsoft.Extensions.Logging;
-
-using static AutomationTool.CommandUtils;
+using UnrealBuildBase;
 
 namespace AutomationTool.Tasks
 {
@@ -65,85 +59,81 @@ namespace AutomationTool.Tasks
 	[TaskElement("Move", typeof(MoveTaskParameters))]
 	public class MoveTask : BgTaskImpl
 	{
-		/// <summary>
-		/// Parameters for this task
-		/// </summary>
-		MoveTaskParameters Parameters;
+		readonly MoveTaskParameters _parameters;
 
 		/// <summary>
 		/// Constructor
 		/// </summary>
-		/// <param name="InParameters">Parameters for this task</param>
-		public MoveTask(MoveTaskParameters InParameters)
+		/// <param name="parameters">Parameters for this task</param>
+		public MoveTask(MoveTaskParameters parameters)
 		{
-			Parameters = InParameters;
+			_parameters = parameters;
 		}
 
 		/// <summary>
-		/// Execute the task.
+		/// ExecuteAsync the task.
 		/// </summary>
-		/// <param name="Job">Information about the current job</param>
-		/// <param name="BuildProducts">Set of build products produced by this node.</param>
-		/// <param name="TagNameToFileSet">Mapping from tag names to the set of files they include</param>
-		public override Task ExecuteAsync(JobContext Job, HashSet<FileReference> BuildProducts, Dictionary<string, HashSet<FileReference>> TagNameToFileSet)
+		/// <param name="job">Information about the current job</param>
+		/// <param name="buildProducts">Set of build products produced by this node.</param>
+		/// <param name="tagNameToFileSet">Mapping from tag names to the set of files they include</param>
+		public override Task ExecuteAsync(JobContext job, HashSet<FileReference> buildProducts, Dictionary<string, HashSet<FileReference>> tagNameToFileSet)
 		{
 			// Parse all the source patterns
-			FilePattern SourcePattern = new FilePattern(Unreal.RootDirectory, Parameters.From);
+			FilePattern sourcePattern = new FilePattern(Unreal.RootDirectory, _parameters.From);
 
 			// Parse the target pattern
-			FilePattern TargetPattern = new FilePattern(Unreal.RootDirectory, Parameters.To);
+			FilePattern targetPattern = new FilePattern(Unreal.RootDirectory, _parameters.To);
 
 			// Apply the filter to the source files
-			HashSet<FileReference> Files = null;
-			if(!String.IsNullOrEmpty(Parameters.Files))
+			HashSet<FileReference> files = null;
+			if (!String.IsNullOrEmpty(_parameters.Files))
 			{
-				SourcePattern = SourcePattern.AsDirectoryPattern();
-				Files = ResolveFilespec(SourcePattern.BaseDirectory, Parameters.Files, TagNameToFileSet);
+				sourcePattern = sourcePattern.AsDirectoryPattern();
+				files = ResolveFilespec(sourcePattern.BaseDirectory, _parameters.Files, tagNameToFileSet);
 			}
 
 			// Build the file mapping
-			Dictionary<FileReference, FileReference> TargetFileToSourceFile;
+			Dictionary<FileReference, FileReference> targetFileToSourceFile;
 
 			try
 			{
-				TargetFileToSourceFile = FilePattern.CreateMapping(Files, ref SourcePattern, ref TargetPattern);
+				targetFileToSourceFile = FilePattern.CreateMapping(files, ref sourcePattern, ref targetPattern);
 
 				// Check we got some files
-				if (TargetFileToSourceFile.Count == 0)
+				if (targetFileToSourceFile.Count == 0)
 				{
-					if (Parameters.ErrorIfNotFound)
+					if (_parameters.ErrorIfNotFound)
 					{
-						Logger.LogError("No files found matching '{SourcePattern}'", SourcePattern);
+						Logger.LogError("No files found matching '{SourcePattern}'", sourcePattern);
 					}
 					else
 					{
-						Logger.LogInformation("No files found matching '{SourcePattern}'", SourcePattern);
+						Logger.LogInformation("No files found matching '{SourcePattern}'", sourcePattern);
 					}
 					return Task.CompletedTask;
 				}
 				// Copy them all
-				Logger.LogInformation("Moving {Arg0} file{Arg1} from {Arg2} to {Arg3}...", TargetFileToSourceFile.Count, (TargetFileToSourceFile.Count == 1) ? "" : "s", SourcePattern.BaseDirectory, TargetPattern.BaseDirectory);
-				CommandUtils.ParallelMoveFiles(TargetFileToSourceFile.Select(x => new KeyValuePair<FileReference, FileReference>(x.Value, x.Key)), Parameters.Overwrite);
+				Logger.LogInformation("Moving {Arg0} file{Arg1} from {Arg2} to {Arg3}...", targetFileToSourceFile.Count, (targetFileToSourceFile.Count == 1) ? "" : "s", sourcePattern.BaseDirectory, targetPattern.BaseDirectory);
+				CommandUtils.ParallelMoveFiles(targetFileToSourceFile.Select(x => new KeyValuePair<FileReference, FileReference>(x.Value, x.Key)), _parameters.Overwrite);
 
 				// Update the list of build products
-				BuildProducts.UnionWith(TargetFileToSourceFile.Keys);
+				buildProducts.UnionWith(targetFileToSourceFile.Keys);
 
 				// Apply the optional output tag to them
-				foreach (string TagName in FindTagNamesFromList(Parameters.Tag))
+				foreach (string tagName in FindTagNamesFromList(_parameters.Tag))
 				{
-					FindOrAddTagSet(TagNameToFileSet, TagName).UnionWith(TargetFileToSourceFile.Keys);
+					FindOrAddTagSet(tagNameToFileSet, tagName).UnionWith(targetFileToSourceFile.Keys);
 				}
-
 			}
-			catch (FilePatternSourceFileMissingException Ex)
+			catch (FilePatternSourceFileMissingException ex)
 			{
-				if (Parameters.ErrorIfNotFound)
+				if (_parameters.ErrorIfNotFound)
 				{
-					Logger.LogError("Error while trying to create file pattern match for '{SourcePattern}', error {ExceptionString}", SourcePattern, Ex.ToString());
+					Logger.LogError("Error while trying to create file pattern match for '{SourcePattern}', error {ExceptionString}", sourcePattern, ex.ToString());
 				}
 				else
 				{
-					Logger.LogInformation("Error while trying to create file pattern match for '{SourcePattern}', error {ExceptionString}", SourcePattern, Ex.ToString());
+					Logger.LogInformation("Error while trying to create file pattern match for '{SourcePattern}', error {ExceptionString}", sourcePattern, ex.ToString());
 				}
 			}
 
@@ -153,9 +143,9 @@ namespace AutomationTool.Tasks
 		/// <summary>
 		/// Output this task out to an XML writer.
 		/// </summary>
-		public override void Write(XmlWriter Writer)
+		public override void Write(XmlWriter writer)
 		{
-			Write(Writer, Parameters);
+			Write(writer, _parameters);
 		}
 
 		/// <summary>
@@ -164,9 +154,9 @@ namespace AutomationTool.Tasks
 		/// <returns>The tag names which are read by this task</returns>
 		public override IEnumerable<string> FindConsumedTagNames()
 		{
-			foreach(string TagName in FindTagNamesFromFilespec(Parameters.Files))
+			foreach (string tagName in FindTagNamesFromFilespec(_parameters.Files))
 			{
-				yield return TagName;
+				yield return tagName;
 			}
 		}
 
@@ -176,7 +166,7 @@ namespace AutomationTool.Tasks
 		/// <returns>The tag names which are modified by this task</returns>
 		public override IEnumerable<string> FindProducedTagNames()
 		{
-			return FindTagNamesFromList(Parameters.Tag);
+			return FindTagNamesFromList(_parameters.Tag);
 		}
 	}
 }
diff --git a/Engine/Source/Programs/AutomationTool/BuildGraph/Tasks/MsBuildTask.cs b/Engine/Source/Programs/AutomationTool/BuildGraph/Tasks/MsBuildTask.cs
--- a/Engine/Source/Programs/AutomationTool/BuildGraph/Tasks/MsBuildTask.cs
+++ b/Engine/Source/Programs/AutomationTool/BuildGraph/Tasks/MsBuildTask.cs
@@ -3,7 +3,6 @@
 using System;
 using System.Collections.Generic;
 using System.Linq;
-using System.Text;
 using System.Threading.Tasks;
 using System.Xml;
 using EpicGames.Core;
@@ -16,10 +15,10 @@ namespace AutomationTool.Tasks
 	/// </summary>
 	public class MsBuildTaskParameters
 	{
-        /// <summary>
-        /// The C# project file to compile. Using semicolons, more than one project file can be specified.
-        /// </summary>
-        [TaskParameter]
+		/// <summary>
+		/// The C# project file to compile. Using semicolons, more than one project file can be specified.
+		/// </summary>
+		[TaskParameter]
 		public string Project { get; set; }
 
 		/// <summary>
@@ -56,59 +55,59 @@ namespace AutomationTool.Tasks
 		/// <summary>
 		/// Parameters for the task
 		/// </summary>
-		MsBuildTaskParameters Parameters;
+		readonly MsBuildTaskParameters _parameters;
 
 		/// <summary>
 		/// Constructor.
 		/// </summary>
-		/// <param name="InParameters">Parameters for this task</param>
-		public MsBuildTask(MsBuildTaskParameters InParameters)
+		/// <param name="parameters">Parameters for this task</param>
+		public MsBuildTask(MsBuildTaskParameters parameters)
 		{
-			Parameters = InParameters;
+			_parameters = parameters;
 		}
 
 		/// <summary>
-		/// Execute the task.
+		/// ExecuteAsync the task.
 		/// </summary>
-		/// <param name="Job">Information about the current job</param>
-		/// <param name="BuildProducts">Set of build products produced by this node.</param>
-		/// <param name="TagNameToFileSet">Mapping from tag names to the set of files they include</param>
-		public override Task ExecuteAsync(JobContext Job, HashSet<FileReference> BuildProducts, Dictionary<string, HashSet<FileReference>> TagNameToFileSet)
+		/// <param name="job">Information about the current job</param>
+		/// <param name="buildProducts">Set of build products produced by this node.</param>
+		/// <param name="tagNameToFileSet">Mapping from tag names to the set of files they include</param>
+		public override Task ExecuteAsync(JobContext job, HashSet<FileReference> buildProducts, Dictionary<string, HashSet<FileReference>> tagNameToFileSet)
 		{
 			// Get the project file
-			HashSet<FileReference> ProjectFiles = ResolveFilespec(Unreal.RootDirectory, Parameters.Project, TagNameToFileSet);
-			foreach(FileReference ProjectFile in ProjectFiles)
+			HashSet<FileReference> projectFiles = ResolveFilespec(Unreal.RootDirectory, _parameters.Project, tagNameToFileSet);
+			foreach (FileReference projectFile in projectFiles)
 			{
-				if(!FileReference.Exists(ProjectFile))
+				if (!FileReference.Exists(projectFile))
 				{
-					throw new AutomationException("Couldn't find project file '{0}'", ProjectFile.FullName);
+					throw new AutomationException("Couldn't find project file '{0}'", projectFile.FullName);
 				}
 			}
 
 			// Build the argument list
-			List<string> Arguments = new List<string>();
-			if(!String.IsNullOrEmpty(Parameters.Platform))
+			List<string> arguments = new List<string>();
+			if (!String.IsNullOrEmpty(_parameters.Platform))
 			{
-				Arguments.Add(String.Format("/p:Platform={0}", CommandUtils.MakePathSafeToUseWithCommandLine(Parameters.Platform)));
+				arguments.Add(String.Format("/p:Platform={0}", CommandUtils.MakePathSafeToUseWithCommandLine(_parameters.Platform)));
 			}
-			if(!String.IsNullOrEmpty(Parameters.Configuration))
+			if (!String.IsNullOrEmpty(_parameters.Configuration))
 			{
-				Arguments.Add(String.Format("/p:Configuration={0}", CommandUtils.MakePathSafeToUseWithCommandLine(Parameters.Configuration)));
+				arguments.Add(String.Format("/p:Configuration={0}", CommandUtils.MakePathSafeToUseWithCommandLine(_parameters.Configuration)));
 			}
-			if(!String.IsNullOrEmpty(Parameters.Arguments))
+			if (!String.IsNullOrEmpty(_parameters.Arguments))
 			{
-				Arguments.Add(Parameters.Arguments);
+				arguments.Add(_parameters.Arguments);
 			}
-			if(!String.IsNullOrEmpty(Parameters.Verbosity))
+			if (!String.IsNullOrEmpty(_parameters.Verbosity))
 			{
-				Arguments.Add(String.Format("/verbosity:{0}", Parameters.Verbosity));
+				arguments.Add(String.Format("/verbosity:{0}", _parameters.Verbosity));
 			}
-			Arguments.Add("/nologo");
+			arguments.Add("/nologo");
 
 			// Build all the projects
-			foreach(FileReference ProjectFile in ProjectFiles)
+			foreach (FileReference projectFile in projectFiles)
 			{
-				CommandUtils.MsBuild(CommandUtils.CmdEnv, ProjectFile.FullName, String.Join(" ", Arguments), null);
+				CommandUtils.MsBuild(CommandUtils.CmdEnv, projectFile.FullName, String.Join(" ", arguments), null);
 			}
 
 			return Task.CompletedTask;
@@ -117,9 +116,9 @@ namespace AutomationTool.Tasks
 		/// <summary>
 		/// Output this task out to an XML writer.
 		/// </summary>
-		public override void Write(XmlWriter Writer)
+		public override void Write(XmlWriter writer)
 		{
-			Write(Writer, Parameters);
+			Write(writer, _parameters);
 		}
 
 		/// <summary>
@@ -128,7 +127,7 @@ namespace AutomationTool.Tasks
 		/// <returns>The tag names which are read by this task</returns>
 		public override IEnumerable<string> FindConsumedTagNames()
 		{
-			return FindTagNamesFromFilespec(Parameters.Project);
+			return FindTagNamesFromFilespec(_parameters.Project);
 		}
 
 		/// <summary>
@@ -139,5 +138,5 @@ namespace AutomationTool.Tasks
 		{
 			return Enumerable.Empty<string>();
 		}
-	}	
+	}
 }
diff --git a/Engine/Source/Programs/AutomationTool/BuildGraph/Tasks/NotarizeTask.cs b/Engine/Source/Programs/AutomationTool/BuildGraph/Tasks/NotarizeTask.cs
--- a/Engine/Source/Programs/AutomationTool/BuildGraph/Tasks/NotarizeTask.cs
+++ b/Engine/Source/Programs/AutomationTool/BuildGraph/Tasks/NotarizeTask.cs
@@ -1,21 +1,13 @@
 // Copyright Epic Games, Inc. All Rights Reserved.
 
-using AutomationTool;
 using System;
 using System.Collections.Generic;
-using System.IO;
-using System.Linq;
-using System.Net;
-using System.Text;
 using System.Text.RegularExpressions;
-using System.Threading;
 using System.Threading.Tasks;
 using System.Xml;
 using EpicGames.Core;
-using UnrealBuildTool;
 using Microsoft.Extensions.Logging;
-
-using static AutomationTool.CommandUtils;
+using UnrealBuildTool;
 
 #pragma warning disable SYSLIB0014
 
@@ -63,154 +55,153 @@ namespace AutomationTool.Tasks
 		/// <summary>
 		/// Parameters for the task
 		/// </summary>
-		NotarizeTaskParameters Parameters;
+		readonly NotarizeTaskParameters _parameters;
 
 		/// <summary>
 		/// Constructor.
 		/// </summary>
-		/// <param name="InParameters">Parameters for the task</param>
-		public NotarizeTask(NotarizeTaskParameters InParameters)
+		/// <param name="parameters">Parameters for the task</param>
+		public NotarizeTask(NotarizeTaskParameters parameters)
 		{
-			Parameters = InParameters;
+			_parameters = parameters;
 		}
 
 		/// <summary>
-		/// Execute the task.
+		/// ExecuteAsync the task.
 		/// </summary>
-		/// <param name="Job">Information about the current job</param>
-		/// <param name="BuildProducts">Set of build products produced by this node.</param>
-		/// <param name="TagNameToFileSet">Mapping from tag names to the set of files they include</param>
-		public override async Task ExecuteAsync(JobContext Job, HashSet<FileReference> BuildProducts, Dictionary<string, HashSet<FileReference>> TagNameToFileSet)
+		/// <param name="job">Information about the current job</param>
+		/// <param name="buildProducts">Set of build products produced by this node.</param>
+		/// <param name="tagNameToFileSet">Mapping from tag names to the set of files they include</param>
+		public override async Task ExecuteAsync(JobContext job, HashSet<FileReference> buildProducts, Dictionary<string, HashSet<FileReference>> tagNameToFileSet)
 		{
 			// Ensure running on a mac.
-			if(BuildHostPlatform.Current.Platform != UnrealTargetPlatform.Mac)
+			if (BuildHostPlatform.Current.Platform != UnrealTargetPlatform.Mac)
 			{
 				throw new AutomationException("Notarization can only be run on a Mac!");
 			}
 
 			// Ensure file exists
-			FileReference Dmg = new FileReference(Parameters.DmgPath);
-			if(!FileReference.Exists(Dmg))
+			FileReference dmg = new FileReference(_parameters.DmgPath);
+			if (!FileReference.Exists(dmg))
 			{
-				throw new AutomationException("Couldn't find a file to notarize at {0}", Dmg.FullName);
+				throw new AutomationException("Couldn't find a file to notarize at {0}", dmg.FullName);
 			}
 
-			int ExitCode = 0;
-			Logger.LogInformation("Uploading {Arg0} to the notarization server...", Dmg.FullName);
-			
+			int exitCode;
+			Logger.LogInformation("Uploading {Arg0} to the notarization server...", dmg.FullName);
+
 			// The notarytool will timeout after 5 retries or 1 hour. Whichever comes first.
 			const int MaxNumRetries = 5;
 			const int MaxTimeoutInMilliseconds = 3600000;
-			long TimeoutInMilliseconds = MaxTimeoutInMilliseconds;
-			string Output = "";
+			long timeoutInMilliseconds = MaxTimeoutInMilliseconds;
+			string output = "";
+
+			System.Diagnostics.Stopwatch timeoutStopwatch = System.Diagnostics.Stopwatch.StartNew();
 
-			System.Diagnostics.Stopwatch TimeoutStopwatch = System.Diagnostics.Stopwatch.StartNew();
-			
-			for (int NumRetries = 0; NumRetries < MaxNumRetries; NumRetries++)
+			for (int numRetries = 0; numRetries < MaxNumRetries; numRetries++)
 			{
-				string CommandLine = string.Format("notarytool submit \"{0}\" --keychain-profile \"{1}\" --wait --timeout \"{2}\"", Dmg.FullName, Parameters.KeyChainID, TimeoutInMilliseconds);
-				Output = CommandUtils.RunAndLog("xcrun", CommandLine, out ExitCode);
-				
-				if (ExitCode == 0)
+				string commandLine = string.Format("notarytool submit \"{0}\" --keychain-profile \"{1}\" --wait --timeout \"{2}\"", dmg.FullName, _parameters.KeyChainID, timeoutInMilliseconds);
+				output = CommandUtils.RunAndLog("xcrun", commandLine, out exitCode);
+
+				if (exitCode == 0)
 				{
 					break;
 				}
-				
-				if (TimeoutStopwatch.ElapsedMilliseconds >= TimeoutInMilliseconds)
+
+				if (timeoutStopwatch.ElapsedMilliseconds >= timeoutInMilliseconds)
 				{
-					Logger.LogInformation("notarytool timed out after {TimeoutInMilliseconds}ms.", TimeoutInMilliseconds);
-					TimeoutStopwatch.Stop();
+					Logger.LogInformation("notarytool timed out after {TimeoutInMilliseconds}ms.", timeoutInMilliseconds);
+					timeoutStopwatch.Stop();
 				}
-				else if (NumRetries < MaxNumRetries)
+				else if (numRetries < MaxNumRetries)
 				{
-					Logger.LogInformation("notarytool failed with exit {ExitCode} attempting retry {NumRetries} of {MaxNumRetries}", ExitCode, NumRetries, MaxNumRetries);
+					Logger.LogInformation("notarytool failed with exit {ExitCode} attempting retry {NumRetries} of {MaxNumRetries}", exitCode, numRetries, MaxNumRetries);
 					await Task.Delay(2000);
-					TimeoutInMilliseconds = MaxTimeoutInMilliseconds - TimeoutStopwatch.ElapsedMilliseconds;
+					timeoutInMilliseconds = MaxTimeoutInMilliseconds - timeoutStopwatch.ElapsedMilliseconds;
 					continue;
 				}
 
-
 				Logger.LogInformation("Retries have been exhausted");
-				throw new AutomationException("notarytool failed with exit {0}", ExitCode);
+				throw new AutomationException("notarytool failed with exit {0}", exitCode);
 			}
 
 			// Grab the UUID from the log
-			string RequestUUID = null;
+			string requestUuid;
 			try
 			{
-				RequestUUID = Regex.Match(Output, "id: ([a-zA-Z0-9]{8}-[a-zA-Z0-9]{4}-[a-zA-Z0-9]{4}-[a-zA-Z0-9]{4}-[a-zA-Z0-9]{12})").Groups[1].Value.Trim();
+				requestUuid = Regex.Match(output, "id: ([a-zA-Z0-9]{8}-[a-zA-Z0-9]{4}-[a-zA-Z0-9]{4}-[a-zA-Z0-9]{4}-[a-zA-Z0-9]{12})").Groups[1].Value.Trim();
 			}
-			catch (Exception Ex)
+			catch (Exception ex)
 			{
-				throw new AutomationException(Ex, "Couldn't get UUID from the log output {0}", Output);
+				throw new AutomationException(ex, "Couldn't get UUID from the log output {0}", output);
 			}
 
 			try
 			{
-				MatchCollection StatusMatches = Regex.Matches(Output, "(?<=status: ).+");
+				MatchCollection statusMatches = Regex.Matches(output, "(?<=status: ).+");
 				// The last status update is the right one.
-				string Status = StatusMatches[StatusMatches.Count - 1].Value.ToLower();
+				string status = statusMatches[statusMatches.Count - 1].Value.ToLower();
 
-				if (Status == "accepted")
+				if (status == "accepted")
 				{
-					if(Parameters.RequireStapling)
+					if (_parameters.RequireStapling)
 					{
 						// once we have a log file, print it out, staple, and we're done.
-						Logger.LogInformation("{Text}", GetRequestLogs(RequestUUID));
-						string CommandLine = string.Format("stapler staple {0}", Dmg.FullName);
-						Output = CommandUtils.RunAndLog("xcrun", CommandLine, out ExitCode);
-						if (ExitCode != 0)
+						Logger.LogInformation("{Text}", GetRequestLogs(requestUuid));
+						string commandLine = string.Format("stapler staple {0}", dmg.FullName);
+						output = CommandUtils.RunAndLog("xcrun", commandLine, out exitCode);
+						if (exitCode != 0)
 						{
-							throw new AutomationException("stapler failed with exit {0}", ExitCode);
+							throw new AutomationException("stapler failed with exit {0}", exitCode);
 						}
 					}
 				}
 				else
 				{
-					Logger.LogError("{Text}", GetRequestLogs(RequestUUID));
-					throw new AutomationException($"Could not notarize the app. Request status: {0}. See log output above.", Status);
+					Logger.LogError("{Text}", GetRequestLogs(requestUuid));
+					throw new AutomationException($"Could not notarize the app. Request status: {0}. See log output above.", status);
 				}
 			}
-			catch (Exception Ex)
+			catch (Exception ex)
 			{
-				if (Ex is AutomationException)
+				if (ex is AutomationException)
 				{
 					throw;
 				}
 				else
 				{
-					throw new AutomationException(Ex, "Querying for the notarization result failed, output: {0}", Output);
+					throw new AutomationException(ex, "Querying for the notarization result failed, output: {0}", output);
 				}
 			}
 		}
 
-		private string GetRequestLogs(string RequestUUID)
+		private string GetRequestLogs(string requestUuid)
 		{
 			try
 			{
-				string LogCommand = string.Format("notarytool log {0} --keychain-profile \"{1}\"", RequestUUID, Parameters.KeyChainID);
-				IProcessResult LogResult = CommandUtils.Run("xcrun", LogCommand);
+				string logCommand = string.Format("notarytool log {0} --keychain-profile \"{1}\"", requestUuid, _parameters.KeyChainID);
+				IProcessResult logResult = CommandUtils.Run("xcrun", logCommand);
 
-				string ResponseContent = null;
-				if (LogResult.bExitCodeSuccess)
+				string responseContent = null;
+				if (logResult.bExitCodeSuccess)
 				{
-					ResponseContent = LogResult.Output;
+					responseContent = logResult.Output;
 				}
 
-				return ResponseContent;
+				return responseContent;
 			}
-			catch (Exception Ex)
+			catch (Exception ex)
 			{
-				throw new AutomationException(Ex, string.Format("Couldn't complete the request, error: {0}", Ex.Message));
+				throw new AutomationException(ex, string.Format("Couldn't complete the request, error: {0}", ex.Message));
 			}
 		}
 
 		/// <summary>
 		/// Output this task out to an XML writer.
 		/// </summary>
-		public override void Write(XmlWriter Writer)
+		public override void Write(XmlWriter writer)
 		{
-			Write(Writer, Parameters);
+			Write(writer, _parameters);
 		}
 
 		/// <summary>
@@ -231,5 +222,4 @@ namespace AutomationTool.Tasks
 			yield break;
 		}
 	}
-
 }
diff --git a/Engine/Source/Programs/AutomationTool/BuildGraph/Tasks/NuGetLicenseCheckTask.cs b/Engine/Source/Programs/AutomationTool/BuildGraph/Tasks/NuGetLicenseCheckTask.cs
--- a/Engine/Source/Programs/AutomationTool/BuildGraph/Tasks/NuGetLicenseCheckTask.cs
+++ b/Engine/Source/Programs/AutomationTool/BuildGraph/Tasks/NuGetLicenseCheckTask.cs
@@ -1,22 +1,18 @@
 // Copyright Epic Games, Inc. All Rights Reserved.
 
-using AutomationTool;
-using EpicGames.Core;
-using Microsoft.Extensions.Logging;
 using System;
-using System.CodeDom;
 using System.Collections.Generic;
 using System.Data;
 using System.IO;
 using System.Linq;
 using System.Net.Http;
-using System.Net.Http.Headers;
 using System.Text;
 using System.Text.Json;
 using System.Text.RegularExpressions;
 using System.Threading.Tasks;
-using System.Web;
 using System.Xml;
+using EpicGames.Core;
+using Microsoft.Extensions.Logging;
 using UnrealBuildBase;
 
 namespace AutomationTool.Tasks
@@ -68,18 +64,15 @@ namespace AutomationTool.Tasks
 			public List<string> Urls { get; set; } = new List<string>();
 		}
 
-		/// <summary>
-		/// Parameters for this task
-		/// </summary>
-		NuGetLicenseCheckTaskParameters Parameters;
+		readonly NuGetLicenseCheckTaskParameters _parameters;
 
 		/// <summary>
 		/// Construct a NuGetLicenseCheckTask task
 		/// </summary>
-		/// <param name="InParameters">Parameters for the task</param>
-		public NuGetLicenseCheckTask(NuGetLicenseCheckTaskParameters InParameters)
+		/// <param name="parameters">Parameters for the task</param>
+		public NuGetLicenseCheckTask(NuGetLicenseCheckTaskParameters parameters)
 		{
-			Parameters = InParameters;
+			_parameters = parameters;
 		}
 
 		enum PackageState
@@ -93,224 +86,224 @@ namespace AutomationTool.Tasks
 
 		class PackageInfo
 		{
-			public string Name;
-			public string Version;
-			public string ProjectUrl;
-			public LicenseInfo License;
-			public string LicenseSource;
-			public PackageState State;
-			public FileReference Descriptor;
+			public string _name;
+			public string _version;
+			public string _projectUrl;
+			public LicenseInfo _license;
+			public string _licenseSource;
+			public PackageState _state;
+			public FileReference _descriptor;
 		}
 
 		class LicenseInfo
 		{
-			public IoHash Hash;
-			public string Text;
-			public string NormalizedText;
-			public string Extension;
-			public bool Approved;
-			public FileReference File;
+			public IoHash _hash;
+			public string _text;
+			public string _normalizedText;
+			public string _extension;
+			public bool _approved;
+			public FileReference _file;
 		}
 
-		static LicenseInfo FindOrAddLicense(Dictionary<IoHash, LicenseInfo> Licenses, string Text, string Extension)
+		static LicenseInfo FindOrAddLicense(Dictionary<IoHash, LicenseInfo> licenses, string text, string extension)
 		{
-			string NormalizedText = Text;
-			NormalizedText = Regex.Replace(NormalizedText, @"^\s+", "", RegexOptions.Multiline);
-			NormalizedText = Regex.Replace(NormalizedText, @"\s+$", "", RegexOptions.Multiline);
-			NormalizedText = Regex.Replace(NormalizedText, "^(?:MIT License|The MIT License \\(MIT\\))\n", "", RegexOptions.Multiline);
-			NormalizedText = Regex.Replace(NormalizedText, "^Copyright \\(c\\)[^\n]*\\s*(?:All rights reserved\\.?\\s*)?", "", RegexOptions.Multiline);
-			NormalizedText = Regex.Replace(NormalizedText, @"\s+", " ");
-			NormalizedText = NormalizedText.Trim();
-
-			byte[] Data = Encoding.UTF8.GetBytes(NormalizedText);
-			IoHash Hash = IoHash.Compute(Data);
-
-			LicenseInfo LicenseInfo;
-			if (!Licenses.TryGetValue(Hash, out LicenseInfo))
+			string normalizedText = text;
+			normalizedText = Regex.Replace(normalizedText, @"^\s+", "", RegexOptions.Multiline);
+			normalizedText = Regex.Replace(normalizedText, @"\s+$", "", RegexOptions.Multiline);
+			normalizedText = Regex.Replace(normalizedText, "^(?:MIT License|The MIT License \\(MIT\\))\n", "", RegexOptions.Multiline);
+			normalizedText = Regex.Replace(normalizedText, "^Copyright \\(c\\)[^\n]*\\s*(?:All rights reserved\\.?\\s*)?", "", RegexOptions.Multiline);
+			normalizedText = Regex.Replace(normalizedText, @"\s+", " ");
+			normalizedText = normalizedText.Trim();
+
+			byte[] data = Encoding.UTF8.GetBytes(normalizedText);
+			IoHash hash = IoHash.Compute(data);
+
+			LicenseInfo licenseInfo;
+			if (!licenses.TryGetValue(hash, out licenseInfo))
 			{
-				LicenseInfo = new LicenseInfo();
-				LicenseInfo.Hash = Hash;
-				LicenseInfo.Text = Text;
-				LicenseInfo.NormalizedText = NormalizedText;
-				LicenseInfo.Extension = Extension;
-				Licenses.Add(Hash, LicenseInfo);
+				licenseInfo = new LicenseInfo();
+				licenseInfo._hash = hash;
+				licenseInfo._text = text;
+				licenseInfo._normalizedText = normalizedText;
+				licenseInfo._extension = extension;
+				licenses.Add(hash, licenseInfo);
 			}
-			return LicenseInfo;
+			return licenseInfo;
 		}
 
 		/// <summary>
-		/// Execute the task.
+		/// ExecuteAsync the task.
 		/// </summary>
-		/// <param name="Job">Information about the current job</param>
-		/// <param name="BuildProducts">Set of build products produced by this node.</param>
-		/// <param name="TagNameToFileSet">Mapping from tag names to the set of files they include</param>
-		public override async Task ExecuteAsync(JobContext Job, HashSet<FileReference> BuildProducts, Dictionary<string, HashSet<FileReference>> TagNameToFileSet)
+		/// <param name="job">Information about the current job</param>
+		/// <param name="buildProducts">Set of build products produced by this node.</param>
+		/// <param name="tagNameToFileSet">Mapping from tag names to the set of files they include</param>
+		public override async Task ExecuteAsync(JobContext job, HashSet<FileReference> buildProducts, Dictionary<string, HashSet<FileReference>> tagNameToFileSet)
 		{
-			FileReference DotNetPath = Parameters.DotNetPath ?? Unreal.DotnetPath;
+			FileReference dotNetPath = _parameters.DotNetPath ?? Unreal.DotnetPath;
 
-			IProcessResult NuGetOutput = await ExecuteAsync(DotNetPath.FullName, $"nuget locals global-packages --list", LogOutput: false);
-			if (NuGetOutput.ExitCode != 0)
+			IProcessResult nuGetOutput = await ExecuteAsync(dotNetPath.FullName, $"nuget locals global-packages --list", logOutput: false);
+			if (nuGetOutput.ExitCode != 0)
 			{
-				throw new AutomationException("DotNet terminated with an exit code indicating an error ({0})", NuGetOutput.ExitCode);
+				throw new AutomationException("DotNet terminated with an exit code indicating an error ({0})", nuGetOutput.ExitCode);
 			}
 
-			List<DirectoryReference> NuGetPackageDirs = new List<DirectoryReference>();
-			foreach (string Line in NuGetOutput.Output.Split('\n'))
+			List<DirectoryReference> nuGetPackageDirs = new List<DirectoryReference>();
+			foreach (string line in nuGetOutput.Output.Split('\n'))
 			{
-				int ColonIdx = Line.IndexOf(':', StringComparison.Ordinal);
-				if (ColonIdx != -1)
+				int colonIdx = line.IndexOf(':', StringComparison.Ordinal);
+				if (colonIdx != -1)
 				{
-					DirectoryReference NuGetPackageDir = new DirectoryReference(Line.Substring(ColonIdx + 1).Trim());
-					Logger.LogInformation("Using NuGet package directory: {Path}", NuGetPackageDir);
-					NuGetPackageDirs.Add(NuGetPackageDir);
+					DirectoryReference nuGetPackageDir = new DirectoryReference(line.Substring(colonIdx + 1).Trim());
+					Logger.LogInformation("Using NuGet package directory: {Path}", nuGetPackageDir);
+					nuGetPackageDirs.Add(nuGetPackageDir);
 				}
 			}
 
 			const string UnknownPrefix = "Unknown-";
 
-			IProcessResult PackageListOutput = await ExecuteAsync(DotNetPath.FullName, "list package --include-transitive", WorkingDir: Parameters.BaseDir, LogOutput: false);
-			if (PackageListOutput.ExitCode != 0)
+			IProcessResult packageListOutput = await ExecuteAsync(dotNetPath.FullName, "list package --include-transitive", workingDir: _parameters.BaseDir, logOutput: false);
+			if (packageListOutput.ExitCode != 0)
 			{
-				throw new AutomationException("DotNet terminated with an exit code indicating an error ({0})", PackageListOutput.ExitCode);
+				throw new AutomationException("DotNet terminated with an exit code indicating an error ({0})", packageListOutput.ExitCode);
 			}
 
-			Dictionary<string, PackageInfo> Packages = new Dictionary<string, PackageInfo>();
-			foreach (string Line in PackageListOutput.Output.Split('\n'))
+			Dictionary<string, PackageInfo> packages = new Dictionary<string, PackageInfo>();
+			foreach (string line in packageListOutput.Output.Split('\n'))
 			{
-				Match Match = Regex.Match(Line, @"^\s*>\s*([^\s]+)\s+(?:[^\s]+\s+)?([^\s]+)\s*$");
-				if (Match.Success)
+				Match match = Regex.Match(line, @"^\s*>\s*([^\s]+)\s+(?:[^\s]+\s+)?([^\s]+)\s*$");
+				if (match.Success)
 				{
-					PackageInfo Info = new PackageInfo();
-					Info.Name = Match.Groups[1].Value;
-					Info.Version = Match.Groups[2].Value;
-					Packages.TryAdd($"{Info.Name}@{Info.Version}", Info);
+					PackageInfo info = new PackageInfo();
+					info._name = match.Groups[1].Value;
+					info._version = match.Groups[2].Value;
+					packages.TryAdd($"{info._name}@{info._version}", info);
 				}
 			}
 
-			DirectoryReference PackageRootDir = DirectoryReference.Combine(DirectoryReference.GetSpecialFolder(Environment.SpecialFolder.UserProfile), ".nuget", "packages");
-			if (!DirectoryReference.Exists(PackageRootDir))
+			DirectoryReference packageRootDir = DirectoryReference.Combine(DirectoryReference.GetSpecialFolder(Environment.SpecialFolder.UserProfile), ".nuget", "packages");
+			if (!DirectoryReference.Exists(packageRootDir))
 			{
-				throw new AutomationException("Missing NuGet package cache at {0}", PackageRootDir);
+				throw new AutomationException("Missing NuGet package cache at {0}", packageRootDir);
 			}
 
-			HashSet<string> LicenseUrls = new HashSet<string>(StringComparer.OrdinalIgnoreCase);
+			HashSet<string> licenseUrls = new HashSet<string>(StringComparer.OrdinalIgnoreCase);
 
-			Dictionary<IoHash, LicenseInfo> Licenses = new Dictionary<IoHash, LicenseInfo>();
-			if (Parameters.LicenseDir != null)
+			Dictionary<IoHash, LicenseInfo> licenses = new Dictionary<IoHash, LicenseInfo>();
+			if (_parameters.LicenseDir != null)
 			{
-				Logger.LogInformation("Reading allowed licenses from {LicenseDir}", Parameters.LicenseDir);
-				foreach (FileReference File in DirectoryReference.EnumerateFiles(Parameters.LicenseDir))
+				Logger.LogInformation("Reading allowed licenses from {LicenseDir}", _parameters.LicenseDir);
+				foreach (FileReference file in DirectoryReference.EnumerateFiles(_parameters.LicenseDir))
 				{
-					if (!File.GetFileName().StartsWith(UnknownPrefix, StringComparison.OrdinalIgnoreCase))
+					if (!file.GetFileName().StartsWith(UnknownPrefix, StringComparison.OrdinalIgnoreCase))
 					{
 						try
 						{
-							if (File.HasExtension(".json"))
+							if (file.HasExtension(".json"))
 							{
-								byte[] Data = await FileReference.ReadAllBytesAsync(File);
-								LicenseConfig Config = JsonSerializer.Deserialize<LicenseConfig>(Data, new JsonSerializerOptions { PropertyNameCaseInsensitive = true, AllowTrailingCommas = true, ReadCommentHandling = JsonCommentHandling.Skip });
-								LicenseUrls.UnionWith(Config.Urls);
+								byte[] data = await FileReference.ReadAllBytesAsync(file);
+								LicenseConfig config = JsonSerializer.Deserialize<LicenseConfig>(data, new JsonSerializerOptions { PropertyNameCaseInsensitive = true, AllowTrailingCommas = true, ReadCommentHandling = JsonCommentHandling.Skip });
+								licenseUrls.UnionWith(config.Urls);
 							}
-							else if (File.HasExtension(".txt") || File.HasExtension(".html"))
+							else if (file.HasExtension(".txt") || file.HasExtension(".html"))
 							{
-								string Text = await FileReference.ReadAllTextAsync(File);
-								LicenseInfo License = FindOrAddLicense(Licenses, Text, File.GetFileNameWithoutExtension());
-								License.File = File;
-								License.Approved = true;
+								string text = await FileReference.ReadAllTextAsync(file);
+								LicenseInfo license = FindOrAddLicense(licenses, text, file.GetFileNameWithoutExtension());
+								license._file = file;
+								license._approved = true;
 							}
 						}
 						catch (Exception ex)
 						{
-							throw new AutomationException(ex, $"Error parsing {File}: {ex.Message}");
+							throw new AutomationException(ex, $"Error parsing {file}: {ex.Message}");
 						}
 					}
 				}
 			}
 
-			HashSet<string> IgnorePackages = new HashSet<string>(StringComparer.OrdinalIgnoreCase);
-			if (Parameters.IgnorePackages != null)
+			HashSet<string> ignorePackages = new HashSet<string>(StringComparer.OrdinalIgnoreCase);
+			if (_parameters.IgnorePackages != null)
 			{
-				IgnorePackages.UnionWith(Parameters.IgnorePackages.Split(';'));
+				ignorePackages.UnionWith(_parameters.IgnorePackages.Split(';'));
 			}
 
-			Dictionary<string, LicenseInfo> LicenseUrlToInfo = new Dictionary<string, LicenseInfo>(StringComparer.OrdinalIgnoreCase);
-			foreach (PackageInfo Info in Packages.Values)
+			Dictionary<string, LicenseInfo> licenseUrlToInfo = new Dictionary<string, LicenseInfo>(StringComparer.OrdinalIgnoreCase);
+			foreach (PackageInfo info in packages.Values)
 			{
-				if (IgnorePackages.Contains(Info.Name) || IgnorePackages.Contains($"{Info.Name}@{Info.Version}"))
+				if (ignorePackages.Contains(info._name) || ignorePackages.Contains($"{info._name}@{info._version}"))
 				{
-					Info.State = PackageState.IgnoredViaArgs;
+					info._state = PackageState.IgnoredViaArgs;
 					continue;
 				}
 
-				DirectoryReference PackageDir = NuGetPackageDirs.Select(x => DirectoryReference.Combine(x, Info.Name.ToLowerInvariant(), Info.Version.ToLowerInvariant())).FirstOrDefault(x => DirectoryReference.Exists(x));
-				if (PackageDir == null)
+				DirectoryReference packageDir = nuGetPackageDirs.Select(x => DirectoryReference.Combine(x, info._name.ToLowerInvariant(), info._version.ToLowerInvariant())).FirstOrDefault(x => DirectoryReference.Exists(x));
+				if (packageDir == null)
 				{
-					Info.State = PackageState.MissingPackageFolder;
+					info._state = PackageState.MissingPackageFolder;
 					continue;
 				}
 
-				Info.Descriptor = FileReference.Combine(PackageDir, $"{Info.Name.ToLowerInvariant()}.nuspec");
-				if (!FileReference.Exists(Info.Descriptor))
+				info._descriptor = FileReference.Combine(packageDir, $"{info._name.ToLowerInvariant()}.nuspec");
+				if (!FileReference.Exists(info._descriptor))
 				{
-					Info.State = PackageState.MissingPackageDescriptor;
+					info._state = PackageState.MissingPackageDescriptor;
 					continue;
 				}
 
-				using (Stream Stream = FileReference.Open(Info.Descriptor, FileMode.Open, FileAccess.Read, FileShare.Read))
+				using (Stream stream = FileReference.Open(info._descriptor, FileMode.Open, FileAccess.Read, FileShare.Read))
 				{
-					using XmlTextReader XmlReader = new XmlTextReader(Stream);
-					XmlReader.Namespaces = false;
+					using XmlTextReader xmlReader = new XmlTextReader(stream);
+					xmlReader.Namespaces = false;
 
-					XmlDocument XmlDocument = new XmlDocument();
-					XmlDocument.Load(XmlReader);
+					XmlDocument xmlDocument = new XmlDocument();
+					xmlDocument.Load(xmlReader);
 
-					XmlNode ProjectUrlNode = XmlDocument.SelectSingleNode("/package/metadata/projectUrl");
-					Info.ProjectUrl = ProjectUrlNode?.InnerText;
+					XmlNode projectUrlNode = xmlDocument.SelectSingleNode("/package/metadata/projectUrl");
+					info._projectUrl = projectUrlNode?.InnerText;
 
-					if (Info.License == null)
+					if (info._license == null)
 					{
-						XmlNode LicenseNode = XmlDocument.SelectSingleNode("/package/metadata/license");
-						if (LicenseNode?.Attributes["type"]?.InnerText?.Equals("file", StringComparison.Ordinal) ?? false)
+						XmlNode licenseNode = xmlDocument.SelectSingleNode("/package/metadata/license");
+						if (licenseNode?.Attributes["type"]?.InnerText?.Equals("file", StringComparison.Ordinal) ?? false)
 						{
-							FileReference LicenseFile = FileReference.Combine(PackageDir, LicenseNode.InnerText);
-							if (FileReference.Exists(LicenseFile))
+							FileReference licenseFile = FileReference.Combine(packageDir, licenseNode.InnerText);
+							if (FileReference.Exists(licenseFile))
 							{
-								string Text = await FileReference.ReadAllTextAsync(LicenseFile);
-								Info.License = FindOrAddLicense(Licenses, Text, LicenseFile.GetExtension());
-								Info.LicenseSource = LicenseFile.FullName;
+								string text = await FileReference.ReadAllTextAsync(licenseFile);
+								info._license = FindOrAddLicense(licenses, text, licenseFile.GetExtension());
+								info._licenseSource = licenseFile.FullName;
 							}
 						}
 					}
 
-					if (Info.License == null)
+					if (info._license == null)
 					{
-						XmlNode LicenseUrlNode = XmlDocument.SelectSingleNode("/package/metadata/licenseUrl");
+						XmlNode licenseUrlNode = xmlDocument.SelectSingleNode("/package/metadata/licenseUrl");
 
-						string LicenseUrl = LicenseUrlNode?.InnerText;
-						if (LicenseUrl != null)
+						string licenseUrl = licenseUrlNode?.InnerText;
+						if (licenseUrl != null)
 						{
-							LicenseUrl = Regex.Replace(LicenseUrl, @"^https://github.com/(.*)/blob/(.*)$", @"https://raw.githubusercontent.com/$1/$2");
-							Info.LicenseSource = LicenseUrl;
+							licenseUrl = Regex.Replace(licenseUrl, @"^https://github.com/(.*)/blob/(.*)$", @"https://raw.githubusercontent.com/$1/$2");
+							info._licenseSource = licenseUrl;
 
-							if (!LicenseUrlToInfo.TryGetValue(LicenseUrl, out Info.License))
+							if (!licenseUrlToInfo.TryGetValue(licenseUrl, out info._license))
 							{
-								using (HttpClient Client = new HttpClient())
+								using (HttpClient client = new HttpClient())
 								{
-									using HttpResponseMessage Response = await Client.GetAsync(LicenseUrl);
-									if (!Response.IsSuccessStatusCode)
+									using HttpResponseMessage response = await client.GetAsync(licenseUrl);
+									if (!response.IsSuccessStatusCode)
 									{
-										Logger.LogError("Unable to fetch license from {LicenseUrl}", LicenseUrl);
+										Logger.LogError("Unable to fetch license from {LicenseUrl}", licenseUrl);
 									}
 									else
 									{
-										string Text = await Response.Content.ReadAsStringAsync();
-										string Type = (Response.Content.Headers.ContentType?.MediaType == "text/html") ? ".html" : ".txt";
-										Info.License = FindOrAddLicense(Licenses, Text, Type);
-										if (!Info.License.Approved)
+										string text = await response.Content.ReadAsStringAsync();
+										string type = (response.Content.Headers.ContentType?.MediaType == "text/html") ? ".html" : ".txt";
+										info._license = FindOrAddLicense(licenses, text, type);
+										if (!info._license._approved)
 										{
-											Info.License.Approved = LicenseUrls.Contains(LicenseUrl);
+											info._license._approved = licenseUrls.Contains(licenseUrl);
 										}
-										LicenseUrlToInfo.Add(LicenseUrl, Info.License);
+										licenseUrlToInfo.Add(licenseUrl, info._license);
 									}
 								}
 							}
@@ -318,102 +311,102 @@ namespace AutomationTool.Tasks
 					}
 				}
 
-				Info.State = PackageState.Valid;
+				info._state = PackageState.Valid;
 			}
 
 			Logger.LogInformation("Referenced Packages:");
 			Logger.LogInformation("");
-			foreach (PackageInfo Info in Packages.Values.OrderBy(x => x.Name).ThenBy(x => x.Version))
+			foreach (PackageInfo info in packages.Values.OrderBy(x => x._name).ThenBy(x => x._version))
 			{
-				switch (Info.State)
+				switch (info._state)
 				{
 					case PackageState.IgnoredViaArgs:
-						Logger.LogInformation("  {Name,-60} {Version,-10} Explicitly ignored via task arguments", Info.Name, Info.Version);
+						Logger.LogInformation("  {Name,-60} {Version,-10} Explicitly ignored via task arguments", info._name, info._version);
 						break;
 					case PackageState.MissingPackageFolder:
-						Logger.LogInformation("  {Name,-60} {Version,-10} NuGet package not found", Info.Name, Info.Version);
+						Logger.LogInformation("  {Name,-60} {Version,-10} NuGet package not found", info._name, info._version);
 						break;
 					case PackageState.MissingPackageDescriptor:
-						Logger.LogWarning("  {Name,-60} {Version,-10} Missing package descriptor: {NuSpecFile}", Info.Name, Info.Version, Info.Descriptor);
+						Logger.LogWarning("  {Name,-60} {Version,-10} Missing package descriptor: {NuSpecFile}", info._name, info._version, info._descriptor);
 						break;
 					case PackageState.Valid:
-						if (Info.License == null)
+						if (info._license == null)
 						{
-							Logger.LogError("  {Name,-60} {Version,-10} No license metadata found", Info.Name, Info.Version);
+							Logger.LogError("  {Name,-60} {Version,-10} No license metadata found", info._name, info._version);
 						}
-						else if (!Info.License.Approved)
+						else if (!info._license._approved)
 						{
-							Logger.LogWarning("  {Name,-60} {Version,-10} {Hash}", Info.Name, Info.Version, Info.License.Hash);
+							Logger.LogWarning("  {Name,-60} {Version,-10} {Hash}", info._name, info._version, info._license._hash);
 						}
 						else
 						{
-							Logger.LogInformation("  {Name,-60} {Version,-10} {Hash}", Info.Name, Info.Version, Info.License.Hash);
+							Logger.LogInformation("  {Name,-60} {Version,-10} {Hash}", info._name, info._version, info._license._hash);
 						}
 						break;
 					default:
-						Logger.LogError("  {Name,-60} {Version,-10} Unhandled state: {State}", Info.Name, Info.Version, Info.State);
+						Logger.LogError("  {Name,-60} {Version,-10} Unhandled state: {State}", info._name, info._version, info._state);
 						break;
 				}
 			}
 
-			Dictionary<LicenseInfo, List<PackageInfo>> MissingLicenses = new Dictionary<LicenseInfo, List<PackageInfo>>();
-			foreach (PackageInfo PackageInfo in Packages.Values)
+			Dictionary<LicenseInfo, List<PackageInfo>> missingLicenses = new Dictionary<LicenseInfo, List<PackageInfo>>();
+			foreach (PackageInfo packageInfo in packages.Values)
 			{
-				if (PackageInfo.License != null && !PackageInfo.License.Approved)
+				if (packageInfo._license != null && !packageInfo._license._approved)
 				{
-					List<PackageInfo> LicensePackages;
-					if (!MissingLicenses.TryGetValue(PackageInfo.License, out LicensePackages))
+					List<PackageInfo> licensePackages;
+					if (!missingLicenses.TryGetValue(packageInfo._license, out licensePackages))
 					{
-						LicensePackages = new List<PackageInfo>();
-						MissingLicenses.Add(PackageInfo.License, LicensePackages);
+						licensePackages = new List<PackageInfo>();
+						missingLicenses.Add(packageInfo._license, licensePackages);
 					}
-					LicensePackages.Add(PackageInfo);
+					licensePackages.Add(packageInfo);
 				}
 			}
 
-			if (MissingLicenses.Count > 0)
+			if (missingLicenses.Count > 0)
 			{
-				DirectoryReference LicenseDir = Parameters.LicenseDir ?? DirectoryReference.Combine(Unreal.RootDirectory, "Engine", "Saved", "Licenses");
-				DirectoryReference.CreateDirectory(LicenseDir);
+				DirectoryReference licenseDir = _parameters.LicenseDir ?? DirectoryReference.Combine(Unreal.RootDirectory, "Engine", "Saved", "Licenses");
+				DirectoryReference.CreateDirectory(licenseDir);
 
 				Logger.LogInformation("");
 				Logger.LogInformation("Missing licenses:");
-				foreach ((LicenseInfo MissingLicense, List<PackageInfo> MissingLicensePackages) in MissingLicenses.OrderBy(x => x.Key.Hash))
+				foreach ((LicenseInfo missingLicense, List<PackageInfo> missingLicensePackages) in missingLicenses.OrderBy(x => x.Key._hash))
 				{
-					FileReference OutputFile = FileReference.Combine(LicenseDir, $"{UnknownPrefix}{MissingLicense.Hash}{MissingLicense.Extension}");
-					await FileReference.WriteAllTextAsync(OutputFile, MissingLicense.Text);
+					FileReference outputFile = FileReference.Combine(licenseDir, $"{UnknownPrefix}{missingLicense._hash}{missingLicense._extension}");
+					await FileReference.WriteAllTextAsync(outputFile, missingLicense._text);
 
 					Logger.LogInformation("");
-					Logger.LogInformation("  {LicenseFile}", OutputFile);
-					foreach (PackageInfo LicensePackage in MissingLicensePackages)
+					Logger.LogInformation("  {LicenseFile}", outputFile);
+					foreach (PackageInfo licensePackage in missingLicensePackages)
 					{
-						Logger.LogInformation("  -> {Name} {Version} ({Source})", LicensePackage.Name, LicensePackage.Version, LicensePackage.LicenseSource);
+						Logger.LogInformation("  -> {Name} {Version} ({Source})", licensePackage._name, licensePackage._version, licensePackage._licenseSource);
 					}
 				}
 			}
 
-			if (Parameters.CsvFile != null)
+			if (_parameters.CsvFile != null)
 			{
-				Logger.LogInformation("Writing {File}", Parameters.CsvFile);
-				DirectoryReference.CreateDirectory(Parameters.CsvFile.Directory);
-				using (StreamWriter writer = new StreamWriter(Parameters.CsvFile.FullName))
+				Logger.LogInformation("Writing {File}", _parameters.CsvFile);
+				DirectoryReference.CreateDirectory(_parameters.CsvFile.Directory);
+				using (StreamWriter writer = new StreamWriter(_parameters.CsvFile.FullName))
 				{
 					await writer.WriteLineAsync($"Package,Version,Project Url,License Url,License Hash,License File");
-					foreach (PackageInfo PackageInfo in Packages.Values)
+					foreach (PackageInfo packageInfo in packages.Values)
 					{
-						string RelativeLicensePath = "";
-						if (PackageInfo.License?.File != null)
+						string relativeLicensePath = "";
+						if (packageInfo._license?._file != null)
 						{
-							RelativeLicensePath = PackageInfo.License.File.MakeRelativeTo(Parameters.CsvFile.Directory);
+							relativeLicensePath = packageInfo._license._file.MakeRelativeTo(_parameters.CsvFile.Directory);
 						}
 
-						string LicenseUrl = "";
-						if (PackageInfo.LicenseSource != null && PackageInfo.LicenseSource.StartsWith("http", StringComparison.OrdinalIgnoreCase))
+						string licenseUrl = "";
+						if (packageInfo._licenseSource != null && packageInfo._licenseSource.StartsWith("http", StringComparison.OrdinalIgnoreCase))
 						{
-							LicenseUrl = PackageInfo.LicenseSource;
+							licenseUrl = packageInfo._licenseSource;
 						}
 
-						await writer.WriteLineAsync($"\"{PackageInfo.Name}\",\"{PackageInfo.Version}\",{PackageInfo.ProjectUrl},{LicenseUrl},{PackageInfo.License?.Hash},{RelativeLicensePath}");
+						await writer.WriteLineAsync($"\"{packageInfo._name}\",\"{packageInfo._version}\",{packageInfo._projectUrl},{licenseUrl},{packageInfo._license?._hash},{relativeLicensePath}");
 					}
 				}
 			}
@@ -422,9 +415,9 @@ namespace AutomationTool.Tasks
 		/// <summary>
 		/// Output this task out to an XML writer.
 		/// </summary>
-		public override void Write(XmlWriter Writer)
+		public override void Write(XmlWriter writer)
 		{
-			Write(Writer, Parameters);
+			Write(writer, _parameters);
 		}
 
 		/// <summary>
diff --git a/Engine/Source/Programs/AutomationTool/BuildGraph/Tasks/OnExitTask.cs b/Engine/Source/Programs/AutomationTool/BuildGraph/Tasks/OnExitTask.cs
--- a/Engine/Source/Programs/AutomationTool/BuildGraph/Tasks/OnExitTask.cs
+++ b/Engine/Source/Programs/AutomationTool/BuildGraph/Tasks/OnExitTask.cs
@@ -1,12 +1,11 @@
 // Copyright Epic Games, Inc. All Rights Reserved.
 
-using EpicGames.Core;
 using System;
 using System.Collections.Generic;
 using System.Linq;
-using System.Text;
 using System.Threading.Tasks;
 using System.Xml;
+using EpicGames.Core;
 
 namespace AutomationTool.Tasks
 {
@@ -34,31 +33,28 @@ namespace AutomationTool.Tasks
 	[TaskElement("OnExit", typeof(OnExitTaskParameters))]
 	public class OnExitTask : BgTaskImpl
 	{
-		/// <summary>
-		/// Parameters for this task
-		/// </summary>
-		OnExitTaskParameters Parameters;
+		readonly OnExitTaskParameters _parameters;
 
 		/// <summary>
 		/// Construct a spawn task
 		/// </summary>
-		/// <param name="InParameters">Parameters for the task</param>
-		public OnExitTask(OnExitTaskParameters InParameters)
+		/// <param name="parameters">Parameters for the task</param>
+		public OnExitTask(OnExitTaskParameters parameters)
 		{
-			Parameters = InParameters;
+			_parameters = parameters;
 		}
 
 		/// <inheritdoc/>
-		public override async Task ExecuteAsync(JobContext Job, HashSet<FileReference> BuildProducts, Dictionary<string, HashSet<FileReference>> TagNameToFileSet)
+		public override async Task ExecuteAsync(JobContext job, HashSet<FileReference> buildProducts, Dictionary<string, HashSet<FileReference>> tagNameToFileSet)
 		{
-			string[] commands = Parameters.Command.Split('\n').Select(x => x.Trim()).ToArray();
-			await AddCleanupCommandsAsync(commands, Parameters.Lease);
+			string[] commands = _parameters.Command.Split('\n').Select(x => x.Trim()).ToArray();
+			await AddCleanupCommandsAsync(commands, _parameters.Lease);
 		}
 
 		/// <inheritdoc/>
-		public override void Write(XmlWriter Writer)
+		public override void Write(XmlWriter writer)
 		{
-			Write(Writer, Parameters);
+			Write(writer, _parameters);
 		}
 
 		/// <inheritdoc/>
diff --git a/Engine/Source/Programs/AutomationTool/BuildGraph/Tasks/PakFileTask.cs b/Engine/Source/Programs/AutomationTool/BuildGraph/Tasks/PakFileTask.cs
--- a/Engine/Source/Programs/AutomationTool/BuildGraph/Tasks/PakFileTask.cs
+++ b/Engine/Source/Programs/AutomationTool/BuildGraph/Tasks/PakFileTask.cs
@@ -1,20 +1,15 @@
 // Copyright Epic Games, Inc. All Rights Reserved.
 
-using AutomationTool;
-using EpicGames.BuildGraph;
 using System;
 using System.Collections.Generic;
 using System.IO;
-using System.Linq;
 using System.Text;
 using System.Threading.Tasks;
 using System.Xml;
 using EpicGames.Core;
-using UnrealBuildTool;
-using UnrealBuildBase;
 using Microsoft.Extensions.Logging;
-
-using static AutomationTool.CommandUtils;
+using UnrealBuildBase;
+using UnrealBuildTool;
 
 namespace AutomationTool.Tasks
 {
@@ -84,139 +79,136 @@ namespace AutomationTool.Tasks
 	[TaskElement("PakFile", typeof(PakFileTaskParameters))]
 	public class PakFileTask : BgTaskImpl
 	{
-		/// <summary>
-		/// Parameters for the task
-		/// </summary>
-		PakFileTaskParameters Parameters;
+		readonly PakFileTaskParameters _parameters;
 
 		/// <summary>
 		/// Constructor.
 		/// </summary>
-		/// <param name="InParameters">Parameters for this task</param>
-		public PakFileTask(PakFileTaskParameters InParameters)
+		/// <param name="parameters">Parameters for this task</param>
+		public PakFileTask(PakFileTaskParameters parameters)
 		{
-			Parameters = InParameters;
+			_parameters = parameters;
 		}
 
 		/// <summary>
-		/// Execute the task.
+		/// ExecuteAsync the task.
 		/// </summary>
-		/// <param name="Job">Information about the current job</param>
-		/// <param name="BuildProducts">Set of build products produced by this node.</param>
-		/// <param name="TagNameToFileSet">Mapping from tag names to the set of files they include</param>
-		public override async Task ExecuteAsync(JobContext Job, HashSet<FileReference> BuildProducts, Dictionary<string, HashSet<FileReference>> TagNameToFileSet)
+		/// <param name="job">Information about the current job</param>
+		/// <param name="buildProducts">Set of build products produced by this node.</param>
+		/// <param name="tagNameToFileSet">Mapping from tag names to the set of files they include</param>
+		public override async Task ExecuteAsync(JobContext job, HashSet<FileReference> buildProducts, Dictionary<string, HashSet<FileReference>> tagNameToFileSet)
 		{
 			// Find the directories we're going to rebase relative to
-			HashSet<DirectoryReference> RebaseDirs = new HashSet<DirectoryReference>{ Unreal.RootDirectory };
-			if(Parameters.RebaseDir != null)
+			HashSet<DirectoryReference> rebaseDirs = new HashSet<DirectoryReference> { Unreal.RootDirectory };
+			if (_parameters.RebaseDir != null)
 			{
-				RebaseDirs.UnionWith(Parameters.RebaseDir);
+				rebaseDirs.UnionWith(_parameters.RebaseDir);
 			}
 
 			// Get the output parameter
-			FileReference OutputFile = Parameters.Output;
+			FileReference outputFile = _parameters.Output;
 
 			// Check for a ResponseFile parameter
-			FileReference ResponseFile = Parameters.ResponseFile;
-			if (ResponseFile == null)
+			FileReference responseFile = _parameters.ResponseFile;
+			if (responseFile == null)
 			{
 				// Get a unique filename for the response file
-				ResponseFile = FileReference.Combine(new DirectoryReference(CommandUtils.CmdEnv.LogFolder), String.Format("PakList_{0}.txt", OutputFile.GetFileNameWithoutExtension()));
-				for (int Idx = 2; FileReference.Exists(ResponseFile); Idx++)
+				responseFile = FileReference.Combine(new DirectoryReference(CommandUtils.CmdEnv.LogFolder), String.Format("PakList_{0}.txt", outputFile.GetFileNameWithoutExtension()));
+				for (int idx = 2; FileReference.Exists(responseFile); idx++)
 				{
-					ResponseFile = FileReference.Combine(ResponseFile.Directory, String.Format("PakList_{0}_{1}.txt", OutputFile.GetFileNameWithoutExtension(), Idx));
+					responseFile = FileReference.Combine(responseFile.Directory, String.Format("PakList_{0}_{1}.txt", outputFile.GetFileNameWithoutExtension(), idx));
 				}
 
 				// Write out the response file
-				HashSet<FileReference> Files = ResolveFilespec(Unreal.RootDirectory, Parameters.Files, TagNameToFileSet);
-				using (StreamWriter Writer = new StreamWriter(ResponseFile.FullName, false, new System.Text.UTF8Encoding(true)))
+				HashSet<FileReference> files = ResolveFilespec(Unreal.RootDirectory, _parameters.Files, tagNameToFileSet);
+				using (StreamWriter writer = new StreamWriter(responseFile.FullName, false, new System.Text.UTF8Encoding(true)))
 				{
-					foreach (FileReference File in Files)
+					foreach (FileReference file in files)
 					{
-						string RelativePath = FindShortestRelativePath(File, RebaseDirs);
-						if (RelativePath == null)
+						string relativePath = FindShortestRelativePath(file, rebaseDirs);
+						if (relativePath == null)
 						{
-							throw new AutomationException("Couldn't find relative path for '{0}' - not under any rebase directories", File.FullName);
+							throw new AutomationException("Couldn't find relative path for '{0}' - not under any rebase directories", file.FullName);
 						}
 
-						string compressArg = Parameters.Compress ? " -compress" : "";
-						await Writer.WriteLineAsync($"\"{File.FullName}\" \"{RelativePath}\"{compressArg}");
+						string compressArg = _parameters.Compress ? " -compress" : "";
+						await writer.WriteLineAsync($"\"{file.FullName}\" \"{relativePath}\"{compressArg}");
 					}
 				}
 			}
 
 			// Format the command line
-			StringBuilder CommandLine = new StringBuilder();
-			CommandLine.AppendFormat("{0} -create={1}", CommandUtils.MakePathSafeToUseWithCommandLine(OutputFile.FullName), CommandUtils.MakePathSafeToUseWithCommandLine(ResponseFile.FullName));
-			if(Parameters.Sign != null)
+			StringBuilder commandLine = new StringBuilder();
+			commandLine.AppendFormat("{0} -create={1}", CommandUtils.MakePathSafeToUseWithCommandLine(outputFile.FullName), CommandUtils.MakePathSafeToUseWithCommandLine(responseFile.FullName));
+			if (_parameters.Sign != null)
 			{
-				CommandLine.AppendFormat(" -sign={0}", CommandUtils.MakePathSafeToUseWithCommandLine(Parameters.Sign.FullName));
+				commandLine.AppendFormat(" -sign={0}", CommandUtils.MakePathSafeToUseWithCommandLine(_parameters.Sign.FullName));
 			}
-			if(Parameters.Order != null)
+			if (_parameters.Order != null)
 			{
-				CommandLine.AppendFormat(" -order={0}", CommandUtils.MakePathSafeToUseWithCommandLine(Parameters.Order.FullName));
+				commandLine.AppendFormat(" -order={0}", CommandUtils.MakePathSafeToUseWithCommandLine(_parameters.Order.FullName));
 			}
 			if (Unreal.IsEngineInstalled())
 			{
-				CommandLine.Append(" -installed");
+				commandLine.Append(" -installed");
 			}
 			if (GlobalCommandLine.UTF8Output)
 			{
-				CommandLine.AppendFormat(" -UTF8Output");
+				commandLine.AppendFormat(" -UTF8Output");
 			}
 
 			// Get the executable path
-			FileReference UnrealPakExe;
-			if(HostPlatform.Current.HostEditorPlatform == UnrealTargetPlatform.Win64)
+			FileReference unrealPakExe;
+			if (HostPlatform.Current.HostEditorPlatform == UnrealTargetPlatform.Win64)
 			{
-				UnrealPakExe = ResolveFile("Engine/Binaries/Win64/UnrealPak.exe");
+				unrealPakExe = ResolveFile("Engine/Binaries/Win64/UnrealPak.exe");
 			}
 			else
 			{
-				UnrealPakExe = ResolveFile(String.Format("Engine/Binaries/{0}/UnrealPak", HostPlatform.Current.HostEditorPlatform.ToString()));
+				unrealPakExe = ResolveFile(String.Format("Engine/Binaries/{0}/UnrealPak", HostPlatform.Current.HostEditorPlatform.ToString()));
 			}
 
 			// Run it
-			Logger.LogInformation("Running '{Arg0} {Arg1}'", CommandUtils.MakePathSafeToUseWithCommandLine(UnrealPakExe.FullName), CommandLine.ToString());
-			CommandUtils.RunAndLog(CommandUtils.CmdEnv, UnrealPakExe.FullName, CommandLine.ToString(), Options: CommandUtils.ERunOptions.Default | CommandUtils.ERunOptions.UTF8Output);
-			BuildProducts.Add(OutputFile);
+			Logger.LogInformation("Running '{Arg0} {Arg1}'", CommandUtils.MakePathSafeToUseWithCommandLine(unrealPakExe.FullName), commandLine.ToString());
+			CommandUtils.RunAndLog(CommandUtils.CmdEnv, unrealPakExe.FullName, commandLine.ToString(), Options: CommandUtils.ERunOptions.Default | CommandUtils.ERunOptions.UTF8Output);
+			buildProducts.Add(outputFile);
 
 			// Apply the optional tag to the output file
-			foreach(string TagName in FindTagNamesFromList(Parameters.Tag))
+			foreach (string tagName in FindTagNamesFromList(_parameters.Tag))
 			{
-				FindOrAddTagSet(TagNameToFileSet, TagName).Add(OutputFile);
+				FindOrAddTagSet(tagNameToFileSet, tagName).Add(outputFile);
 			}
 		}
 
 		/// <summary>
 		/// Find the shortest relative path of the given file from a set of base directories.
 		/// </summary>
-		/// <param name="File">Full path to a file</param>
-		/// <param name="RebaseDirs">Possible base directories</param>
+		/// <param name="file">Full path to a file</param>
+		/// <param name="rebaseDirs">Possible base directories</param>
 		/// <returns>The shortest relative path, or null if the file is not under any of them</returns>
-		public static string FindShortestRelativePath(FileReference File, IEnumerable<DirectoryReference> RebaseDirs)
+		public static string FindShortestRelativePath(FileReference file, IEnumerable<DirectoryReference> rebaseDirs)
 		{
-			string RelativePath = null;
-			foreach(DirectoryReference RebaseDir in RebaseDirs)
+			string relativePath = null;
+			foreach (DirectoryReference rebaseDir in rebaseDirs)
 			{
-				if(File.IsUnderDirectory(RebaseDir))
+				if (file.IsUnderDirectory(rebaseDir))
 				{
-					string NewRelativePath = File.MakeRelativeTo(RebaseDir);
-					if(RelativePath == null || NewRelativePath.Length < RelativePath.Length)
+					string newRelativePath = file.MakeRelativeTo(rebaseDir);
+					if (relativePath == null || newRelativePath.Length < relativePath.Length)
 					{
-						RelativePath = NewRelativePath;
+						relativePath = newRelativePath;
 					}
 				}
 			}
-			return RelativePath;
+			return relativePath;
 		}
 
 		/// <summary>
 		/// Output this task out to an XML writer.
 		/// </summary>
-		public override void Write(XmlWriter Writer)
+		public override void Write(XmlWriter writer)
 		{
-			Write(Writer, Parameters);
+			Write(writer, _parameters);
 		}
 
 		/// <summary>
@@ -225,7 +217,7 @@ namespace AutomationTool.Tasks
 		/// <returns>The tag names which are read by this task</returns>
 		public override IEnumerable<string> FindConsumedTagNames()
 		{
-			return FindTagNamesFromFilespec(Parameters.Files);
+			return FindTagNamesFromFilespec(_parameters.Files);
 		}
 
 		/// <summary>
@@ -234,7 +226,7 @@ namespace AutomationTool.Tasks
 		/// <returns>The tag names which are modified by this task</returns>
 		public override IEnumerable<string> FindProducedTagNames()
 		{
-			return FindTagNamesFromList(Parameters.Tag);
+			return FindTagNamesFromList(_parameters.Tag);
 		}
 	}
 }
diff --git a/Engine/Source/Programs/AutomationTool/BuildGraph/Tasks/QueryTargetsTask.cs b/Engine/Source/Programs/AutomationTool/BuildGraph/Tasks/QueryTargetsTask.cs
--- a/Engine/Source/Programs/AutomationTool/BuildGraph/Tasks/QueryTargetsTask.cs
+++ b/Engine/Source/Programs/AutomationTool/BuildGraph/Tasks/QueryTargetsTask.cs
@@ -1,17 +1,12 @@
 // Copyright Epic Games, Inc. All Rights Reserved.
 
-using AutomationTool;
-using AutomationTool.Tasks;
-using EpicGames.BuildGraph;
-using System;
 using System.Collections.Generic;
 using System.Linq;
-using System.Text;
 using System.Threading.Tasks;
 using System.Xml;
 using EpicGames.Core;
-using UnrealBuildTool;
 using UnrealBuildBase;
+using UnrealBuildTool;
 
 #nullable enable
 
@@ -56,75 +51,75 @@ namespace AutomationTool.Tasks
 		/// <summary>
 		/// Parameters for this task
 		/// </summary>
-		QueryTargetsTaskParameters Parameters;
+		readonly QueryTargetsTaskParameters _parameters;
 
 		/// <summary>
 		/// Construct a spawn task
 		/// </summary>
-		/// <param name="InParameters">Parameters for the task</param>
-		public QueryTargetsTask(QueryTargetsTaskParameters InParameters)
+		/// <param name="parameters">Parameters for the task</param>
+		public QueryTargetsTask(QueryTargetsTaskParameters parameters)
 		{
-			Parameters = InParameters;
+			_parameters = parameters;
 		}
 
 		/// <summary>
-		/// Execute the task.
+		/// ExecuteAsync the task.
 		/// </summary>
-		/// <param name="Job">Information about the current job</param>
-		/// <param name="BuildProducts">Set of build products produced by this node.</param>
-		/// <param name="TagNameToFileSet">Mapping from tag names to the set of files they include</param>
-		public override Task ExecuteAsync(JobContext Job, HashSet<FileReference> BuildProducts, Dictionary<string, HashSet<FileReference>> TagNameToFileSet)
+		/// <param name="job">Information about the current job</param>
+		/// <param name="buildProducts">Set of build products produced by this node.</param>
+		/// <param name="tagNameToFileSet">Mapping from tag names to the set of files they include</param>
+		public override Task ExecuteAsync(JobContext job, HashSet<FileReference> buildProducts, Dictionary<string, HashSet<FileReference>> tagNameToFileSet)
 		{
 			// Get the output file
-			FileReference? OutputFile = Parameters.OutputFile;
-			if (OutputFile == null)
+			FileReference? outputFile = _parameters.OutputFile;
+			if (outputFile == null)
 			{
-				if (Parameters.ProjectFile == null)
+				if (_parameters.ProjectFile == null)
 				{
-					OutputFile = FileReference.Combine(Unreal.EngineDirectory, "Intermediate", "TargetInfo.json");
+					outputFile = FileReference.Combine(Unreal.EngineDirectory, "Intermediate", "TargetInfo.json");
 				}
 				else
 				{
-					OutputFile = FileReference.Combine(Parameters.ProjectFile.Directory, "Intermediate", "TargetInfo.json");
+					outputFile = FileReference.Combine(_parameters.ProjectFile.Directory, "Intermediate", "TargetInfo.json");
 				}
 			}
-			FileUtils.ForceDeleteFile(OutputFile);
+			FileUtils.ForceDeleteFile(outputFile);
 
 			// Run UBT to generate the target info
-			List<string> Arguments = new List<string> { "-Mode=QueryTargets" };
-			if (Parameters.ProjectFile != null)
+			List<string> arguments = new List<string> { "-Mode=QueryTargets" };
+			if (_parameters.ProjectFile != null)
 			{
-				Arguments.Add($"-Project={Parameters.ProjectFile}");
+				arguments.Add($"-Project={_parameters.ProjectFile}");
 			}
-			if (Parameters.IncludeAllTargets)
+			if (_parameters.IncludeAllTargets)
 			{
-				Arguments.Add("-IncludeAllTargets");
+				arguments.Add("-IncludeAllTargets");
 			}
-			CommandUtils.RunUBT(CommandUtils.CmdEnv, Unreal.UnrealBuildToolDllPath, CommandLineArguments.Join(Arguments));
+			CommandUtils.RunUBT(CommandUtils.CmdEnv, Unreal.UnrealBuildToolDllPath, CommandLineArguments.Join(arguments));
 
 			// Check the output file exists
-			if (!FileReference.Exists(OutputFile))
+			if (!FileReference.Exists(outputFile))
 			{
-				throw new BuildException($"Missing {OutputFile}");
+				throw new BuildException($"Missing {outputFile}");
 			}
 
 			// Apply the optional tag to the build products
-			foreach(string TagName in FindTagNamesFromList(Parameters.Tag))
+			foreach (string tagName in FindTagNamesFromList(_parameters.Tag))
 			{
-				FindOrAddTagSet(TagNameToFileSet, TagName).Add(OutputFile);
+				FindOrAddTagSet(tagNameToFileSet, tagName).Add(outputFile);
 			}
 
 			// Add the target files to the set of build products
-			BuildProducts.Add(OutputFile);
+			buildProducts.Add(outputFile);
 			return Task.CompletedTask;
 		}
 
 		/// <summary>
 		/// Output this task out to an XML writer.
 		/// </summary>
-		public override void Write(XmlWriter Writer)
+		public override void Write(XmlWriter writer)
 		{
-			Write(Writer, Parameters);
+			Write(writer, _parameters);
 		}
 
 		/// <summary>
@@ -142,7 +137,7 @@ namespace AutomationTool.Tasks
 		/// <returns>The tag names which are modified by this task</returns>
 		public override IEnumerable<string> FindProducedTagNames()
 		{
-			return FindTagNamesFromList(Parameters.Tag);
+			return FindTagNamesFromList(_parameters.Tag);
 		}
 	}
 }
diff --git a/Engine/Source/Programs/AutomationTool/BuildGraph/Tasks/RandomDataTask.cs b/Engine/Source/Programs/AutomationTool/BuildGraph/Tasks/RandomDataTask.cs
--- a/Engine/Source/Programs/AutomationTool/BuildGraph/Tasks/RandomDataTask.cs
+++ b/Engine/Source/Programs/AutomationTool/BuildGraph/Tasks/RandomDataTask.cs
@@ -1,19 +1,12 @@
 // Copyright Epic Games, Inc. All Rights Reserved.
 
-using AutomationTool;
-using EpicGames.BuildGraph;
 using System;
 using System.Collections.Generic;
-using System.Linq;
-using System.Text;
-using System.Threading;
+using System.Security.Cryptography;
 using System.Threading.Tasks;
 using System.Xml;
 using EpicGames.Core;
-using UnrealBuildTool;
-using UnrealBuildBase;
 using Microsoft.Extensions.Logging;
-using System.Security.Cryptography;
 
 namespace AutomationTool.Tasks
 {
@@ -59,58 +52,55 @@ namespace AutomationTool.Tasks
 	[TaskElement("RandomData", typeof(RandomDataTaskParameters))]
 	public class RandomDataTask : BgTaskImpl
 	{
-		/// <summary>
-		/// Parameters for this task
-		/// </summary>
-		RandomDataTaskParameters Parameters;
+		readonly RandomDataTaskParameters _parameters;
 
 		/// <summary>
 		/// Constructor
 		/// </summary>
-		/// <param name="InParameters">Parameters for this task</param>
-		public RandomDataTask(RandomDataTaskParameters InParameters)
+		/// <param name="parameters">Parameters for this task</param>
+		public RandomDataTask(RandomDataTaskParameters parameters)
 		{
-			Parameters = InParameters;
+			_parameters = parameters;
 		}
 
 		/// <summary>
-		/// Execute the task.
+		/// ExecuteAsync the task.
 		/// </summary>
-		/// <param name="Job">Information about the current job</param>
-		/// <param name="BuildProducts">Set of build products produced by this node.</param>
-		/// <param name="TagNameToFileSet">Mapping from tag names to the set of files they include</param>
-		public override async Task ExecuteAsync(JobContext Job, HashSet<FileReference> BuildProducts, Dictionary<string, HashSet<FileReference>> TagNameToFileSet)
+		/// <param name="job">Information about the current job</param>
+		/// <param name="buildProducts">Set of build products produced by this node.</param>
+		/// <param name="tagNameToFileSet">Mapping from tag names to the set of files they include</param>
+		public override async Task ExecuteAsync(JobContext job, HashSet<FileReference> buildProducts, Dictionary<string, HashSet<FileReference>> tagNameToFileSet)
 		{
-			DirectoryReference OutputDir = ResolveDirectory(Parameters.OutputDir);
+			DirectoryReference outputDir = ResolveDirectory(_parameters.OutputDir);
 
 			byte[] buffer = Array.Empty<byte>();
-			for (int idx = 0; idx < Parameters.Count; idx++)
+			for (int idx = 0; idx < _parameters.Count; idx++)
 			{
-				if (idx == 0 || Parameters.Different)
+				if (idx == 0 || _parameters.Different)
 				{
-					buffer = RandomNumberGenerator.GetBytes(Parameters.Size);
+					buffer = RandomNumberGenerator.GetBytes(_parameters.Size);
 				}
 
-				FileReference file = FileReference.Combine(OutputDir, $"test-{Parameters.Size}-{idx}.dat");
+				FileReference file = FileReference.Combine(outputDir, $"test-{_parameters.Size}-{idx}.dat");
 				await FileReference.WriteAllBytesAsync(file, buffer);
-				BuildProducts.Add(file);
+				buildProducts.Add(file);
 			}
 
-			Logger.LogInformation("Created {NumFiles:n0} files of {Size:n0} bytes in {OutputDir} (Different={Different})", Parameters.Count, Parameters.Size, OutputDir, Parameters.Different);
+			Logger.LogInformation("Created {NumFiles:n0} files of {Size:n0} bytes in {OutputDir} (Different={Different})", _parameters.Count, _parameters.Size, outputDir, _parameters.Different);
 
 			// Apply the optional output tag to them
-			foreach (string TagName in FindTagNamesFromList(Parameters.Tag))
+			foreach (string tagName in FindTagNamesFromList(_parameters.Tag))
 			{
-				FindOrAddTagSet(TagNameToFileSet, TagName).UnionWith(BuildProducts);
+				FindOrAddTagSet(tagNameToFileSet, tagName).UnionWith(buildProducts);
 			}
 		}
 
 		/// <summary>
 		/// Output this task out to an XML writer.
 		/// </summary>
-		public override void Write(XmlWriter Writer)
+		public override void Write(XmlWriter writer)
 		{
-			Write(Writer, Parameters);
+			Write(writer, _parameters);
 		}
 
 		/// <summary>
@@ -128,7 +118,7 @@ namespace AutomationTool.Tasks
 		/// <returns>The tag names which are modified by this task</returns>
 		public override IEnumerable<string> FindProducedTagNames()
 		{
-			return FindTagNamesFromList(Parameters.Tag);
+			return FindTagNamesFromList(_parameters.Tag);
 		}
 	}
 }
diff --git a/Engine/Source/Programs/AutomationTool/BuildGraph/Tasks/RenameTask.cs b/Engine/Source/Programs/AutomationTool/BuildGraph/Tasks/RenameTask.cs
--- a/Engine/Source/Programs/AutomationTool/BuildGraph/Tasks/RenameTask.cs
+++ b/Engine/Source/Programs/AutomationTool/BuildGraph/Tasks/RenameTask.cs
@@ -1,17 +1,13 @@
 // Copyright Epic Games, Inc. All Rights Reserved.
 
-using AutomationTool;
-using EpicGames.BuildGraph;
 using System;
 using System.Collections.Generic;
 using System.IO;
-using System.Linq;
 using System.Text;
 using System.Text.RegularExpressions;
 using System.Threading.Tasks;
 using System.Xml;
 using EpicGames.Core;
-using UnrealBuildTool;
 using UnrealBuildBase;
 
 namespace AutomationTool.Tasks
@@ -52,98 +48,95 @@ namespace AutomationTool.Tasks
 	[TaskElement("Rename", typeof(RenameTaskParameters))]
 	public class RenameTask : BgTaskImpl
 	{
-		/// <summary>
-		/// Parameters for this task
-		/// </summary>
-		RenameTaskParameters Parameters;
+		readonly RenameTaskParameters _parameters;
 
 		/// <summary>
 		/// Constructor
 		/// </summary>
-		/// <param name="InParameters">Parameters for this task</param>
-		public RenameTask(RenameTaskParameters InParameters)
+		/// <param name="parameters">Parameters for this task</param>
+		public RenameTask(RenameTaskParameters parameters)
 		{
-			Parameters = InParameters;
+			_parameters = parameters;
 		}
 
 		/// <summary>
-		/// Execute the task.
+		/// ExecuteAsync the task.
 		/// </summary>
-		/// <param name="Job">Information about the current job</param>
-		/// <param name="BuildProducts">Set of build products produced by this node.</param>
-		/// <param name="TagNameToFileSet">Mapping from tag names to the set of files they include</param>
-		public override Task ExecuteAsync(JobContext Job, HashSet<FileReference> BuildProducts, Dictionary<string, HashSet<FileReference>> TagNameToFileSet)
+		/// <param name="job">Information about the current job</param>
+		/// <param name="buildProducts">Set of build products produced by this node.</param>
+		/// <param name="tagNameToFileSet">Mapping from tag names to the set of files they include</param>
+		public override Task ExecuteAsync(JobContext job, HashSet<FileReference> buildProducts, Dictionary<string, HashSet<FileReference>> tagNameToFileSet)
 		{
 			// Get the pattern to match against. If it's a simple pattern (eg. *.cpp, Engine/Build/...), automatically infer the source wildcard
-			string FromPattern = Parameters.From;
-			if (FromPattern == null)
+			string fromPattern = _parameters.From;
+			if (fromPattern == null)
 			{
-				List<string> Patterns = SplitDelimitedList(Parameters.Files);
-				if (Patterns.Count != 1 || Patterns[0].StartsWith("#", StringComparison.Ordinal))
+				List<string> patterns = SplitDelimitedList(_parameters.Files);
+				if (patterns.Count != 1 || patterns[0].StartsWith("#", StringComparison.Ordinal))
 				{
 					throw new AutomationException("Missing 'From' attribute specifying pattern to match source files against");
 				}
 
-				FromPattern = Patterns[0];
+				fromPattern = patterns[0];
 
-				int SlashIdx = FromPattern.LastIndexOfAny(new char[] { Path.DirectorySeparatorChar, Path.AltDirectorySeparatorChar });
-				if (SlashIdx != -1)
+				int slashIdx = fromPattern.LastIndexOfAny(new char[] { Path.DirectorySeparatorChar, Path.AltDirectorySeparatorChar });
+				if (slashIdx != -1)
 				{
-					FromPattern = FromPattern.Substring(SlashIdx + 1);
+					fromPattern = fromPattern.Substring(slashIdx + 1);
 				}
-				if (FromPattern.StartsWith("...", StringComparison.Ordinal))
+				if (fromPattern.StartsWith("...", StringComparison.Ordinal))
 				{
-					FromPattern = "*" + FromPattern.Substring(3);
+					fromPattern = "*" + fromPattern.Substring(3);
 				}
 			}
 
 			// Convert the source pattern into a regex
-			string EscapedFromPattern = "^" + Regex.Escape(FromPattern) + "$";
-			EscapedFromPattern = EscapedFromPattern.Replace("\\*", "(.*)", StringComparison.Ordinal);
-			EscapedFromPattern = EscapedFromPattern.Replace("\\?", "(.)", StringComparison.Ordinal);
-			Regex FromRegex = new Regex(EscapedFromPattern, RegexOptions.IgnoreCase | RegexOptions.CultureInvariant);
+			string escapedFromPattern = "^" + Regex.Escape(fromPattern) + "$";
+			escapedFromPattern = escapedFromPattern.Replace("\\*", "(.*)", StringComparison.Ordinal);
+			escapedFromPattern = escapedFromPattern.Replace("\\?", "(.)", StringComparison.Ordinal);
+			Regex fromRegex = new Regex(escapedFromPattern, RegexOptions.IgnoreCase | RegexOptions.CultureInvariant);
 
 			// Split the output pattern into fragments that we can insert captures between
-			string[] FromFragments = FromPattern.Split('*', '?');
-			string[] ToFragments = Parameters.To.Split('*', '?');
-			if(FromFragments.Length < ToFragments.Length)
+			string[] fromFragments = fromPattern.Split('*', '?');
+			string[] toFragments = _parameters.To.Split('*', '?');
+			if (fromFragments.Length < toFragments.Length)
 			{
-				throw new AutomationException("Too few capture groups in source pattern '{0}' to rename to '{1}'", FromPattern, Parameters.To);
+				throw new AutomationException("Too few capture groups in source pattern '{0}' to rename to '{1}'", fromPattern, _parameters.To);
 			}
 
 			// Find the input files
-			HashSet<FileReference> InputFiles = ResolveFilespec(Unreal.RootDirectory, Parameters.Files, TagNameToFileSet);
+			HashSet<FileReference> inputFiles = ResolveFilespec(Unreal.RootDirectory, _parameters.Files, tagNameToFileSet);
 
 			// Find all the corresponding output files
-			Dictionary<FileReference, FileReference> RenameFiles = new Dictionary<FileReference, FileReference>();
-			foreach (FileReference InputFile in InputFiles)
+			Dictionary<FileReference, FileReference> renameFiles = new Dictionary<FileReference, FileReference>();
+			foreach (FileReference inputFile in inputFiles)
 			{
-				Match Match = FromRegex.Match(InputFile.GetFileName());
-				if (Match.Success)
+				Match match = fromRegex.Match(inputFile.GetFileName());
+				if (match.Success)
 				{
-					StringBuilder OutputName = new StringBuilder(ToFragments[0]);
-					for (int Idx = 1; Idx < ToFragments.Length; Idx++)
+					StringBuilder outputName = new StringBuilder(toFragments[0]);
+					for (int idx = 1; idx < toFragments.Length; idx++)
 					{
-						OutputName.Append(Match.Groups[Idx].Value);
-						OutputName.Append(ToFragments[Idx]);
+						outputName.Append(match.Groups[idx].Value);
+						outputName.Append(toFragments[idx]);
 					}
-					RenameFiles[InputFile] = FileReference.Combine(InputFile.Directory, OutputName.ToString());
+					renameFiles[inputFile] = FileReference.Combine(inputFile.Directory, outputName.ToString());
 				}
 			}
 
 			// Print out everything we're going to do
-			foreach(KeyValuePair<FileReference, FileReference> Pair in RenameFiles)
+			foreach (KeyValuePair<FileReference, FileReference> pair in renameFiles)
 			{
-				CommandUtils.RenameFile(Pair.Key.FullName, Pair.Value.FullName, true);
+				CommandUtils.RenameFile(pair.Key.FullName, pair.Value.FullName, true);
 			}
 
 			// Add the build product
-			BuildProducts.UnionWith(RenameFiles.Values);
+			buildProducts.UnionWith(renameFiles.Values);
 
 			// Apply the optional output tag to them
-			foreach(string TagName in FindTagNamesFromList(Parameters.Tag))
+			foreach (string tagName in FindTagNamesFromList(_parameters.Tag))
 			{
-				FindOrAddTagSet(TagNameToFileSet, TagName).UnionWith(RenameFiles.Values);
+				FindOrAddTagSet(tagNameToFileSet, tagName).UnionWith(renameFiles.Values);
 			}
 
 			return Task.CompletedTask;
@@ -152,9 +145,9 @@ namespace AutomationTool.Tasks
 		/// <summary>
 		/// Output this task out to an XML writer.
 		/// </summary>
-		public override void Write(XmlWriter Writer)
+		public override void Write(XmlWriter writer)
 		{
-			Write(Writer, Parameters);
+			Write(writer, _parameters);
 		}
 
 		/// <summary>
@@ -163,7 +156,7 @@ namespace AutomationTool.Tasks
 		/// <returns>The tag names which are read by this task</returns>
 		public override IEnumerable<string> FindConsumedTagNames()
 		{
-			return FindTagNamesFromFilespec(Parameters.Files);
+			return FindTagNamesFromFilespec(_parameters.Files);
 		}
 
 		/// <summary>
@@ -172,7 +165,7 @@ namespace AutomationTool.Tasks
 		/// <returns>The tag names which are modified by this task</returns>
 		public override IEnumerable<string> FindProducedTagNames()
 		{
-			return FindTagNamesFromList(Parameters.Tag);
+			return FindTagNamesFromList(_parameters.Tag);
 		}
 	}
 }
diff --git a/Engine/Source/Programs/AutomationTool/BuildGraph/Tasks/SanitizeReceiptTask.cs b/Engine/Source/Programs/AutomationTool/BuildGraph/Tasks/SanitizeReceiptTask.cs
--- a/Engine/Source/Programs/AutomationTool/BuildGraph/Tasks/SanitizeReceiptTask.cs
+++ b/Engine/Source/Programs/AutomationTool/BuildGraph/Tasks/SanitizeReceiptTask.cs
@@ -1,18 +1,13 @@
 // Copyright Epic Games, Inc. All Rights Reserved.
 
-using EpicGames.BuildGraph;
-using System;
 using System.Collections.Generic;
 using System.Linq;
-using System.Text;
 using System.Threading.Tasks;
 using System.Xml;
 using EpicGames.Core;
+using Microsoft.Extensions.Logging;
 using UnrealBuildBase;
 using UnrealBuildTool;
-using Microsoft.Extensions.Logging;
-
-using static AutomationTool.CommandUtils;
 
 namespace AutomationTool.Tasks
 {
@@ -40,92 +35,89 @@ namespace AutomationTool.Tasks
 	[TaskElement("SanitizeReceipt", typeof(SanitizeReceiptTaskParameters))]
 	class SanitizeReceiptTask : BgTaskImpl
 	{
-		/// <summary>
-		/// Parameters to this task
-		/// </summary>
-		SanitizeReceiptTaskParameters Parameters;
+		readonly SanitizeReceiptTaskParameters _parameters;
 
 		/// <summary>
 		/// Constructor
 		/// </summary>
-		/// <param name="InParameters">Parameters to select which files to search</param>
-		public SanitizeReceiptTask(SanitizeReceiptTaskParameters InParameters)
+		/// <param name="parameters">Parameters to select which files to search</param>
+		public SanitizeReceiptTask(SanitizeReceiptTaskParameters parameters)
 		{
-			Parameters = InParameters;
+			_parameters = parameters;
 		}
 
 		/// <summary>
-		/// Execute the task.
+		/// ExecuteAsync the task.
 		/// </summary>
-		/// <param name="Job">Information about the current job</param>
-		/// <param name="BuildProducts">Set of build products produced by this node.</param>
-		/// <param name="TagNameToFileSet">Mapping from tag names to the set of files they include</param>
-		public override async Task ExecuteAsync(JobContext Job, HashSet<FileReference> BuildProducts, Dictionary<string, HashSet<FileReference>> TagNameToFileSet)
+		/// <param name="job">Information about the current job</param>
+		/// <param name="buildProducts">Set of build products produced by this node.</param>
+		/// <param name="tagNameToFileSet">Mapping from tag names to the set of files they include</param>
+		public override async Task ExecuteAsync(JobContext job, HashSet<FileReference> buildProducts, Dictionary<string, HashSet<FileReference>> tagNameToFileSet)
 		{
 			// Set the Engine directory
-			DirectoryReference EngineDir = Parameters.EngineDir ?? Unreal.EngineDirectory;
+			DirectoryReference engineDir = _parameters.EngineDir ?? Unreal.EngineDirectory;
 
 			// Resolve the input list
-			IEnumerable<FileReference> TargetFiles = ResolveFilespec(Unreal.RootDirectory, Parameters.Files, TagNameToFileSet);
-			await ExecuteAsync(TargetFiles, EngineDir);
+			IEnumerable<FileReference> targetFiles = ResolveFilespec(Unreal.RootDirectory, _parameters.Files, tagNameToFileSet);
+			await Execute(targetFiles, engineDir);
 		}
 
-		public static Task ExecuteAsync(IEnumerable<FileReference> TargetFiles, DirectoryReference EngineDir)
+		public static Task Execute(IEnumerable<FileReference> targetFiles, DirectoryReference engineDir)
 		{
-			EngineDir ??= Unreal.EngineDirectory;
+			engineDir ??= Unreal.EngineDirectory;
 
-			foreach(FileReference TargetFile in TargetFiles)
+			foreach (FileReference targetFile in targetFiles)
 			{
 				// check all files are .target files
-				if (TargetFile.GetExtension() != ".target")
+				if (targetFile.GetExtension() != ".target")
 				{
-					throw new AutomationException("Invalid file passed to TagReceipt task ({0})", TargetFile.FullName);
+					throw new AutomationException("Invalid file passed to TagReceipt task ({0})", targetFile.FullName);
 				}
 
 				// Print the name of the file being scanned
-				Logger.LogInformation("Sanitizing {TargetFile}", TargetFile);
-				using(new LogIndentScope("  "))
+				Logger.LogInformation("Sanitizing {TargetFile}", targetFile);
+				using (new LogIndentScope("  "))
 				{
 					// Read the receipt
-					TargetReceipt Receipt;
-					if (!TargetReceipt.TryRead(TargetFile, EngineDir, out Receipt))
+					TargetReceipt receipt;
+					if (!TargetReceipt.TryRead(targetFile, engineDir, out receipt))
 					{
-						Logger.LogWarning("Unable to load file using TagReceipt task ({Arg0})", TargetFile.FullName);
+						Logger.LogWarning("Unable to load file using TagReceipt task ({Arg0})", targetFile.FullName);
 						continue;
 					}
 
 					// Remove any build products that don't exist
-					List<BuildProduct> NewBuildProducts = new List<BuildProduct>(Receipt.BuildProducts.Count);
-					foreach(BuildProduct BuildProduct in Receipt.BuildProducts)
+					List<BuildProduct> newBuildProducts = new List<BuildProduct>(receipt.BuildProducts.Count);
+					foreach (BuildProduct buildProduct in receipt.BuildProducts)
 					{
-						if(FileReference.Exists(BuildProduct.Path))
+						if (FileReference.Exists(buildProduct.Path))
 						{
-							NewBuildProducts.Add(BuildProduct);
+							newBuildProducts.Add(buildProduct);
 						}
 						else
 						{
-							Logger.LogInformation("Removing build product: {File}", BuildProduct.Path);
+							Logger.LogInformation("Removing build product: {File}", buildProduct.Path);
 						}
 					}
-					Receipt.BuildProducts = NewBuildProducts;
+					receipt.BuildProducts = newBuildProducts;
 
 					// Remove any runtime dependencies that don't exist
-					RuntimeDependencyList NewRuntimeDependencies = new RuntimeDependencyList();
-					foreach(RuntimeDependency RuntimeDependency in Receipt.RuntimeDependencies)
+					RuntimeDependencyList newRuntimeDependencies = new RuntimeDependencyList();
+					foreach (RuntimeDependency runtimeDependency in receipt.RuntimeDependencies)
 					{
-						if(FileReference.Exists(RuntimeDependency.Path))
+						if (FileReference.Exists(runtimeDependency.Path))
 						{
-							NewRuntimeDependencies.Add(RuntimeDependency);
+							newRuntimeDependencies.Add(runtimeDependency);
 						}
 						else
 						{
-							Logger.LogInformation("Removing runtime dependency: {File}", RuntimeDependency.Path);
+							Logger.LogInformation("Removing runtime dependency: {File}", runtimeDependency.Path);
 						}
 					}
-					Receipt.RuntimeDependencies = NewRuntimeDependencies;
-				
+					receipt.RuntimeDependencies = newRuntimeDependencies;
+
 					// Save the new receipt
-					Receipt.Write(TargetFile, EngineDir);
+					receipt.Write(targetFile, engineDir);
 				}
 			}
 			return Task.CompletedTask;
@@ -134,9 +126,9 @@ namespace AutomationTool.Tasks
 		/// <summary>
 		/// Output this task out to an XML writer.
 		/// </summary>
-		public override void Write(XmlWriter Writer)
+		public override void Write(XmlWriter writer)
 		{
-			Write(Writer, Parameters);
+			Write(writer, _parameters);
 		}
 
 		/// <summary>
@@ -145,7 +137,7 @@ namespace AutomationTool.Tasks
 		/// <returns>The tag names which are required by this task</returns>
 		public override IEnumerable<string> FindConsumedTagNames()
 		{
-			return FindTagNamesFromFilespec(Parameters.Files);
+			return FindTagNamesFromFilespec(_parameters.Files);
 		}
 
 		/// <summary>
@@ -166,9 +158,9 @@ namespace AutomationTool.Tasks
 		/// <summary>
 		/// Sanitize the given receipt files, removing any files that don't exist in the current workspace
 		/// </summary>
-		public static async Task SanitizeReceiptsAsync(this FileSet TargetFiles, DirectoryReference EngineDir = null)
+		public static async Task SanitizeReceiptsAsync(this FileSet targetFiles, DirectoryReference engineDir = null)
 		{
-			await SanitizeReceiptTask.ExecuteAsync(TargetFiles, EngineDir);
+			await SanitizeReceiptTask.Execute(targetFiles, engineDir);
 		}
 	}
 }
diff --git a/Engine/Source/Programs/AutomationTool/BuildGraph/Tasks/SetVersionTask.cs b/Engine/Source/Programs/AutomationTool/BuildGraph/Tasks/SetVersionTask.cs
--- a/Engine/Source/Programs/AutomationTool/BuildGraph/Tasks/SetVersionTask.cs
+++ b/Engine/Source/Programs/AutomationTool/BuildGraph/Tasks/SetVersionTask.cs
@@ -1,14 +1,9 @@
 // Copyright Epic Games, Inc. All Rights Reserved.
 
-using EpicGames.BuildGraph;
-using System;
 using System.Collections.Generic;
-using System.Linq;
-using System.Text;
 using System.Threading.Tasks;
 using System.Xml;
 using EpicGames.Core;
-using UnrealBuildTool;
 using UnrealBuildBase;
 
 namespace AutomationTool.Tasks
@@ -79,48 +74,45 @@ namespace AutomationTool.Tasks
 	[TaskElement("SetVersion", typeof(SetVersionTaskParameters))]
 	public class SetVersionTask : BgTaskImpl
 	{
-		/// <summary>
-		/// Parameters for the task
-		/// </summary>
-		SetVersionTaskParameters Parameters;
+		readonly SetVersionTaskParameters _parameters;
 
 		/// <summary>
 		/// Construct a version task
 		/// </summary>
-		/// <param name="InParameters">Parameters for this task</param>
-		public SetVersionTask(SetVersionTaskParameters InParameters)
+		/// <param name="parameters">Parameters for this task</param>
+		public SetVersionTask(SetVersionTaskParameters parameters)
 		{
-			Parameters = InParameters;
+			_parameters = parameters;
 		}
 
 		/// <summary>
-		/// Execute the task.
+		/// ExecuteAsync the task.
 		/// </summary>
-		/// <param name="Job">Information about the current job</param>
-		/// <param name="BuildProducts">Set of build products produced by this node.</param>
-		/// <param name="TagNameToFileSet">Mapping from tag names to the set of files they include</param>
-		public override Task ExecuteAsync(JobContext Job, HashSet<FileReference> BuildProducts, Dictionary<string, HashSet<FileReference>> TagNameToFileSet)
+		/// <param name="job">Information about the current job</param>
+		/// <param name="buildProducts">Set of build products produced by this node.</param>
+		/// <param name="tagNameToFileSet">Mapping from tag names to the set of files they include</param>
+		public override Task ExecuteAsync(JobContext job, HashSet<FileReference> buildProducts, Dictionary<string, HashSet<FileReference>> tagNameToFileSet)
 		{
 			// Update the version files
-			List<FileReference> VersionFiles = UnrealBuild.StaticUpdateVersionFiles(Parameters.Change, Parameters.CompatibleChange, Parameters.Branch, Parameters.Build, Parameters.BuildURL, Parameters.Licensee, Parameters.Promoted, !Parameters.SkipWrite);
+			List<FileReference> versionFiles = UnrealBuild.StaticUpdateVersionFiles(_parameters.Change, _parameters.CompatibleChange, _parameters.Branch, _parameters.Build, _parameters.BuildURL, _parameters.Licensee, _parameters.Promoted, !_parameters.SkipWrite);
 
 			// Apply the optional tag to them
-			foreach(string TagName in FindTagNamesFromList(Parameters.Tag))
+			foreach (string tagName in FindTagNamesFromList(_parameters.Tag))
 			{
-				FindOrAddTagSet(TagNameToFileSet, TagName).UnionWith(VersionFiles);
+				FindOrAddTagSet(tagNameToFileSet, tagName).UnionWith(versionFiles);
 			}
 
 			// Add them to the list of build products
-			BuildProducts.UnionWith(VersionFiles);
+			buildProducts.UnionWith(versionFiles);
 			return Task.CompletedTask;
 		}
 
 		/// <summary>
 		/// Output this task out to an XML writer.
 		/// </summary>
-		public override void Write(XmlWriter Writer)
+		public override void Write(XmlWriter writer)
 		{
-			Write(Writer, Parameters);
+			Write(writer, _parameters);
 		}
 
 		/// <summary>
@@ -138,7 +130,7 @@ namespace AutomationTool.Tasks
 		/// <returns>The tag names which are modified by this task</returns>
 		public override IEnumerable<string> FindProducedTagNames()
 		{
-			return FindTagNamesFromList(Parameters.Tag);
+			return FindTagNamesFromList(_parameters.Tag);
 		}
 	}
 
@@ -148,32 +140,32 @@ namespace AutomationTool.Tasks
 	public static partial class StandardTasks
 	{
 		/// <summary>
-		/// Execute a task instance
+		/// ExecuteAsync a task instance
 		/// </summary>
-		/// <param name="Task"></param>
+		/// <param name="task"></param>
 		/// <returns></returns>
-		public static async Task<FileSet> ExecuteAsync(BgTaskImpl Task)
+		public static async Task<FileSet> ExecuteAsync(BgTaskImpl task)
 		{
-			HashSet<FileReference> BuildProducts = new HashSet<FileReference>();
-			await Task.ExecuteAsync(new JobContext(null!, null!), BuildProducts, new Dictionary<string, HashSet<FileReference>>());
-			return FileSet.FromFiles(Unreal.RootDirectory, BuildProducts);
+			HashSet<FileReference> buildProducts = new HashSet<FileReference>();
+			await task.ExecuteAsync(new JobContext(null!, null!), buildProducts, new Dictionary<string, HashSet<FileReference>>());
+			return FileSet.FromFiles(Unreal.RootDirectory, buildProducts);
 		}
 
 		/// <summary>
 		/// Updates the current engine version
 		/// </summary>
-		public static async Task<FileSet> SetVersionAsync(int Change, string Branch, int? CompatibleChange = null, string Build = null, string BuildURL = null, bool? Licensee = null, bool? Promoted = null, bool? SkipWrite = null)
+		public static async Task<FileSet> SetVersionAsync(int change, string branch, int? compatibleChange = null, string build = null, string buildUrl = null, bool? licensee = null, bool? promoted = null, bool? skipWrite = null)
 		{
-			SetVersionTaskParameters Parameters = new SetVersionTaskParameters();
-			Parameters.Change = Change;
-			Parameters.CompatibleChange = CompatibleChange ?? Parameters.CompatibleChange;
-			Parameters.Branch = Branch ?? Parameters.Branch;
-			Parameters.Build = Build ?? Parameters.Build;
-			Parameters.BuildURL = BuildURL ?? Parameters.BuildURL;
-			Parameters.Licensee = Licensee ?? Parameters.Licensee;
-			Parameters.Promoted = Promoted ?? Parameters.Promoted;
-			Parameters.SkipWrite = SkipWrite ?? Parameters.SkipWrite;
-			return await ExecuteAsync(new SetVersionTask(Parameters));
+			SetVersionTaskParameters parameters = new SetVersionTaskParameters();
+			parameters.Change = change;
+			parameters.CompatibleChange = compatibleChange ?? parameters.CompatibleChange;
+			parameters.Branch = branch ?? parameters.Branch;
+			parameters.Build = build ?? parameters.Build;
+			parameters.BuildURL = buildUrl ?? parameters.BuildURL;
+			parameters.Licensee = licensee ?? parameters.Licensee;
+			parameters.Promoted = promoted ?? parameters.Promoted;
+			parameters.SkipWrite = skipWrite ?? parameters.SkipWrite;
+			return await ExecuteAsync(new SetVersionTask(parameters));
 		}
 	}
 }
