EpicGames.BuildGraph: Apply standard naming conventions to BuildGraph.Automation.

diff --git a/Engine/Source/Programs/AutomationTool/BuildGraph/Tasks/SignTask.cs b/Engine/Source/Programs/AutomationTool/BuildGraph/Tasks/SignTask.cs
--- a/Engine/Source/Programs/AutomationTool/BuildGraph/Tasks/SignTask.cs
+++ b/Engine/Source/Programs/AutomationTool/BuildGraph/Tasks/SignTask.cs
@@ -1,10 +1,10 @@
 // Copyright Epic Games, Inc. All Rights Reserved.
 
-using EpicGames.Core;
 using System.Collections.Generic;
 using System.Linq;
 using System.Threading.Tasks;
 using System.Xml;
+using EpicGames.Core;
 using UnrealBuildBase;
 
 namespace AutomationTool.Tasks
@@ -48,48 +48,48 @@ namespace AutomationTool.Tasks
 		/// <summary>
 		/// Parameters for this task
 		/// </summary>
-		SignTaskParameters Parameters;
+		readonly SignTaskParameters _parameters;
 
 		/// <summary>
 		/// Construct a spawn task
 		/// </summary>
-		/// <param name="InParameters">Parameters for the task</param>
-		public SignTask(SignTaskParameters InParameters)
+		/// <param name="parameters">Parameters for the task</param>
+		public SignTask(SignTaskParameters parameters)
 		{
-			Parameters = InParameters;
+			_parameters = parameters;
 		}
 
 		/// <summary>
-		/// Execute the task.
+		/// ExecuteAsync the task.
 		/// </summary>
-		/// <param name="Job">Information about the current job</param>
-		/// <param name="BuildProducts">Set of build products produced by this node.</param>
-		/// <param name="TagNameToFileSet">Mapping from tag names to the set of files they include</param>
-		public override Task ExecuteAsync(JobContext Job, HashSet<FileReference> BuildProducts, Dictionary<string, HashSet<FileReference>> TagNameToFileSet)
+		/// <param name="job">Information about the current job</param>
+		/// <param name="buildProducts">Set of build products produced by this node.</param>
+		/// <param name="tagNameToFileSet">Mapping from tag names to the set of files they include</param>
+		public override Task ExecuteAsync(JobContext job, HashSet<FileReference> buildProducts, Dictionary<string, HashSet<FileReference>> tagNameToFileSet)
 		{
 			// Find the matching files
-			FileReference[] Files = ResolveFilespec(Unreal.RootDirectory, Parameters.Files, TagNameToFileSet).OrderBy(x => x.FullName).ToArray();
+			FileReference[] files = ResolveFilespec(Unreal.RootDirectory, _parameters.Files, tagNameToFileSet).OrderBy(x => x.FullName).ToArray();
 
 			// Sign all the files
-			CodeSign.SignMultipleIfEXEOrDLL(Job.OwnerCommand, (Files.Select(x => x.FullName).ToList()), Description: Parameters.Description, Parameters.Parallel);
+			CodeSign.SignMultipleIfEXEOrDLL(job.OwnerCommand, (files.Select(x => x.FullName).ToList()), Description: _parameters.Description, _parameters.Parallel);
 
 			// Apply the optional tag to the build products
-			foreach(string TagName in FindTagNamesFromList(Parameters.Tag))
+			foreach (string tagName in FindTagNamesFromList(_parameters.Tag))
 			{
-				FindOrAddTagSet(TagNameToFileSet, TagName).UnionWith(Files);
+				FindOrAddTagSet(tagNameToFileSet, tagName).UnionWith(files);
 			}
 
 			// Add them to the list of build products
-			BuildProducts.UnionWith(Files);
+			buildProducts.UnionWith(files);
 			return Task.CompletedTask;
 		}
 
 		/// <summary>
 		/// Output this task out to an XML writer.
 		/// </summary>
-		public override void Write(XmlWriter Writer)
+		public override void Write(XmlWriter writer)
 		{
-			Write(Writer, Parameters);
+			Write(writer, _parameters);
 		}
 
 		/// <summary>
@@ -98,7 +98,7 @@ namespace AutomationTool.Tasks
 		/// <returns>The tag names which are read by this task</returns>
 		public override IEnumerable<string> FindConsumedTagNames()
 		{
-			return FindTagNamesFromFilespec(Parameters.Files);
+			return FindTagNamesFromFilespec(_parameters.Files);
 		}
 
 		/// <summary>
@@ -107,7 +107,7 @@ namespace AutomationTool.Tasks
 		/// <returns>The tag names which are modified by this task</returns>
 		public override IEnumerable<string> FindProducedTagNames()
 		{
-			return FindTagNamesFromList(Parameters.Tag);
+			return FindTagNamesFromList(_parameters.Tag);
 		}
 	}
 }
diff --git a/Engine/Source/Programs/AutomationTool/BuildGraph/Tasks/SpawnTask.cs b/Engine/Source/Programs/AutomationTool/BuildGraph/Tasks/SpawnTask.cs
--- a/Engine/Source/Programs/AutomationTool/BuildGraph/Tasks/SpawnTask.cs
+++ b/Engine/Source/Programs/AutomationTool/BuildGraph/Tasks/SpawnTask.cs
@@ -1,17 +1,11 @@
 // Copyright Epic Games, Inc. All Rights Reserved.
 
-using AutomationTool;
 using System;
 using System.Collections.Generic;
 using System.IO;
-using System.Linq;
-using System.Text;
 using System.Threading.Tasks;
 using System.Xml;
 using EpicGames.Core;
-using UnrealBuildTool;
-using EpicGames.BuildGraph;
-using AutomationTool.Tasks;
 
 namespace AutomationTool.Tasks
 {
@@ -69,77 +63,77 @@ namespace AutomationTool.Tasks
 	public abstract class SpawnTaskBase : BgTaskImpl
 	{
 		/// <summary>
-		/// Execute a command
+		/// ExecuteAsync a command
 		/// </summary>
-		protected static Task<IProcessResult> ExecuteAsync(string Exe, string Arguments, string WorkingDir = null, Dictionary<string, string> EnvVars = null, bool LogOutput = true, int ErrorLevel = 1, string Input = null, ProcessResult.SpewFilterCallbackType SpewFilterCallback = null)
+		protected static Task<IProcessResult> ExecuteAsync(string exe, string arguments, string workingDir = null, Dictionary<string, string> envVars = null, bool logOutput = true, int errorLevel = 1, string input = null, ProcessResult.SpewFilterCallbackType spewFilterCallback = null)
 		{
-			if (WorkingDir != null)
+			if (workingDir != null)
 			{
-				WorkingDir = ResolveDirectory(WorkingDir).FullName;
+				workingDir = ResolveDirectory(workingDir).FullName;
 			}
 
-			CommandUtils.ERunOptions Options = CommandUtils.ERunOptions.Default;
-			if (!LogOutput)
+			CommandUtils.ERunOptions options = CommandUtils.ERunOptions.Default;
+			if (!logOutput)
 			{
-				Options |= CommandUtils.ERunOptions.SpewIsVerbose;
+				options |= CommandUtils.ERunOptions.SpewIsVerbose;
 			}
 
-			IProcessResult Result = CommandUtils.Run(Exe, Arguments, Env: EnvVars, WorkingDir: WorkingDir, Options: Options, Input: Input, SpewFilterCallback: SpewFilterCallback);
-			if (Result.ExitCode < 0 || Result.ExitCode >= ErrorLevel)
+			IProcessResult result = CommandUtils.Run(exe, arguments, Env: envVars, WorkingDir: workingDir, Options: options, Input: input, SpewFilterCallback: spewFilterCallback);
+			if (result.ExitCode < 0 || result.ExitCode >= errorLevel)
 			{
-				throw new AutomationException("{0} terminated with an exit code indicating an error ({1})", Path.GetFileName(Exe), Result.ExitCode);
+				throw new AutomationException("{0} terminated with an exit code indicating an error ({1})", Path.GetFileName(exe), result.ExitCode);
 			}
 
-			return Task.FromResult(Result);
+			return Task.FromResult(result);
 		}
 
 		/// <summary>
 		/// Parses environment from a property and file
 		/// </summary>
-		/// <param name="Environment"></param>
-		/// <param name="EnvironmentFile"></param>
+		/// <param name="environment"></param>
+		/// <param name="environmentFile"></param>
 		/// <returns></returns>
-		protected static Dictionary<string, string> ParseEnvVars(string Environment, string EnvironmentFile)
+		protected static Dictionary<string, string> ParseEnvVars(string environment, string environmentFile)
 		{
-			Dictionary<string, string> EnvVars = new Dictionary<string, string>();
-			if (Environment != null)
+			Dictionary<string, string> envVars = new Dictionary<string, string>();
+			if (environment != null)
 			{
-				ParseEnvironment(Environment, ';', EnvVars);
+				ParseEnvironment(environment, ';', envVars);
 			}
-			if (!String.IsNullOrEmpty(EnvironmentFile))
+			if (!String.IsNullOrEmpty(environmentFile))
 			{
-				ParseEnvironment(FileUtils.ReadAllText(ResolveFile(EnvironmentFile)), '\n', EnvVars);
+				ParseEnvironment(FileUtils.ReadAllText(ResolveFile(environmentFile)), '\n', envVars);
 			}
-			return EnvVars;
+			return envVars;
 		}
 
 		/// <summary>
 		/// Parse environment from a string
 		/// </summary>
-		/// <param name="Environment"></param>
-		/// <param name="Separator"></param>
-		/// <param name="EnvVars"></param>
-		static void ParseEnvironment(string Environment, char Separator, Dictionary<string, string> EnvVars)
+		/// <param name="environment"></param>
+		/// <param name="separator"></param>
+		/// <param name="envVars"></param>
+		static void ParseEnvironment(string environment, char separator, Dictionary<string, string> envVars)
 		{
-			for (int BaseIdx = 0; BaseIdx < Environment.Length;)
+			for (int baseIdx = 0; baseIdx < environment.Length;)
 			{
-				int EqualsIdx = Environment.IndexOf('=', BaseIdx);
-				if (EqualsIdx == -1)
+				int equalsIdx = environment.IndexOf('=', baseIdx);
+				if (equalsIdx == -1)
 				{
-					throw new AutomationException("Missing value in environment variable string '{0}'", Environment);
+					throw new AutomationException("Missing value in environment variable string '{0}'", environment);
 				}
 
-				int EndIdx = Environment.IndexOf(Separator, EqualsIdx + 1);
-				if (EndIdx == -1)
+				int endIdx = environment.IndexOf(separator, equalsIdx + 1);
+				if (endIdx == -1)
 				{
-					EndIdx = Environment.Length;
+					endIdx = environment.Length;
 				}
 
-				string Name = Environment.Substring(BaseIdx, EqualsIdx - BaseIdx).Trim();
-				string Value = Environment.Substring(EqualsIdx + 1, EndIdx - (EqualsIdx + 1)).Trim();
-				EnvVars[Name] = Value;
+				string name = environment.Substring(baseIdx, equalsIdx - baseIdx).Trim();
+				string value = environment.Substring(equalsIdx + 1, endIdx - (equalsIdx + 1)).Trim();
+				envVars[name] = value;
 
-				BaseIdx = EndIdx + 1;
+				baseIdx = endIdx + 1;
 			}
 		}
 	}
@@ -153,34 +147,34 @@ namespace AutomationTool.Tasks
 		/// <summary>
 		/// Parameters for this task
 		/// </summary>
-		SpawnTaskParameters Parameters;
+		readonly SpawnTaskParameters _parameters;
 
 		/// <summary>
 		/// Construct a spawn task
 		/// </summary>
-		/// <param name="InParameters">Parameters for the task</param>
-		public SpawnTask(SpawnTaskParameters InParameters)
+		/// <param name="parameters">Parameters for the task</param>
+		public SpawnTask(SpawnTaskParameters parameters)
 		{
-			Parameters = InParameters;
+			_parameters = parameters;
 		}
 
 		/// <summary>
-		/// Execute the task.
+		/// ExecuteAsync the task.
 		/// </summary>
-		/// <param name="Job">Information about the current job</param>
-		/// <param name="BuildProducts">Set of build products produced by this node.</param>
-		/// <param name="TagNameToFileSet">Mapping from tag names to the set of files they include</param>
-		public override async Task ExecuteAsync(JobContext Job, HashSet<FileReference> BuildProducts, Dictionary<string, HashSet<FileReference>> TagNameToFileSet)
+		/// <param name="job">Information about the current job</param>
+		/// <param name="buildProducts">Set of build products produced by this node.</param>
+		/// <param name="tagNameToFileSet">Mapping from tag names to the set of files they include</param>
+		public override async Task ExecuteAsync(JobContext job, HashSet<FileReference> buildProducts, Dictionary<string, HashSet<FileReference>> tagNameToFileSet)
 		{
-			await ExecuteAsync(Parameters.Exe, Parameters.Arguments, Parameters.WorkingDir, EnvVars: ParseEnvVars(Parameters.Environment, Parameters.EnvironmentFile), LogOutput: Parameters.LogOutput, ErrorLevel: Parameters.ErrorLevel);
+			await ExecuteAsync(_parameters.Exe, _parameters.Arguments, _parameters.WorkingDir, envVars: ParseEnvVars(_parameters.Environment, _parameters.EnvironmentFile), logOutput: _parameters.LogOutput, errorLevel: _parameters.ErrorLevel);
 		}
 
 		/// <summary>
 		/// Output this task out to an XML writer.
 		/// </summary>
-		public override void Write(XmlWriter Writer)
+		public override void Write(XmlWriter writer)
 		{
-			Write(Writer, Parameters);
+			Write(writer, _parameters);
 		}
 
 		/// <summary>
@@ -205,27 +199,27 @@ namespace AutomationTool.Tasks
 	public static partial class StandardTasks
 	{
 		/// <summary>
-		/// Execute an external program
+		/// ExecuteAsync an external program
 		/// </summary>
-		/// <param name="Exe">Executable to spawn.</param>
-		/// <param name="Arguments">Arguments for the newly created process.</param>
-		/// <param name="WorkingDir">Working directory for spawning the new task.</param>
-		/// <param name="Environment">Environment variables to set.</param>
-		/// <param name="EnvironmentFile">File to read environment from.</param>
-		/// <param name="LogOutput">Write output to the log.</param>
-		/// <param name="ErrorLevel">The minimum exit code which is treated as an error.</param>
-		public static async Task SpawnAsync(string Exe, string Arguments = null, string WorkingDir = null, string Environment = null, string EnvironmentFile = null, bool? LogOutput = null, int? ErrorLevel = null)
+		/// <param name="exe">Executable to spawn.</param>
+		/// <param name="arguments">Arguments for the newly created process.</param>
+		/// <param name="workingDir">Working directory for spawning the new task.</param>
+		/// <param name="environment">Environment variables to set.</param>
+		/// <param name="environmentFile">File to read environment from.</param>
+		/// <param name="logOutput">Write output to the log.</param>
+		/// <param name="errorLevel">The minimum exit code which is treated as an error.</param>
+		public static async Task SpawnAsync(string exe, string arguments = null, string workingDir = null, string environment = null, string environmentFile = null, bool? logOutput = null, int? errorLevel = null)
 		{
-			SpawnTaskParameters Parameters = new SpawnTaskParameters();
-			Parameters.Exe = Exe;
-			Parameters.Arguments = Arguments;
-			Parameters.WorkingDir = WorkingDir;
-			Parameters.Environment = Environment;
-			Parameters.EnvironmentFile = EnvironmentFile;
-			Parameters.LogOutput = LogOutput ?? Parameters.LogOutput;
-			Parameters.ErrorLevel = ErrorLevel ?? Parameters.ErrorLevel;
-
-			await ExecuteAsync(new SpawnTask(Parameters));
+			SpawnTaskParameters parameters = new SpawnTaskParameters();
+			parameters.Exe = exe;
+			parameters.Arguments = arguments;
+			parameters.WorkingDir = workingDir;
+			parameters.Environment = environment;
+			parameters.EnvironmentFile = environmentFile;
+			parameters.LogOutput = logOutput ?? parameters.LogOutput;
+			parameters.ErrorLevel = errorLevel ?? parameters.ErrorLevel;
+
+			await ExecuteAsync(new SpawnTask(parameters));
 		}
 	}
 }
diff --git a/Engine/Source/Programs/AutomationTool/BuildGraph/Tasks/StageTask.cs b/Engine/Source/Programs/AutomationTool/BuildGraph/Tasks/StageTask.cs
--- a/Engine/Source/Programs/AutomationTool/BuildGraph/Tasks/StageTask.cs
+++ b/Engine/Source/Programs/AutomationTool/BuildGraph/Tasks/StageTask.cs
@@ -1,16 +1,12 @@
 // Copyright Epic Games, Inc. All Rights Reserved.
 
-using AutomationTool;
-using EpicGames.BuildGraph;
-using System;
 using System.Collections.Generic;
 using System.Linq;
-using System.Text;
 using System.Threading.Tasks;
 using System.Xml;
 using EpicGames.Core;
-using UnrealBuildTool;
 using UnrealBuildBase;
+using UnrealBuildTool;
 
 namespace AutomationTool.Tasks
 {
@@ -74,110 +70,107 @@ namespace AutomationTool.Tasks
 	[TaskElement("Stage", typeof(StageTaskParameters))]
 	public class StageTask : BgTaskImpl
 	{
-		/// <summary>
-		/// Parameters for the task
-		/// </summary>
-		StageTaskParameters Parameters;
+		readonly StageTaskParameters _parameters;
 
 		/// <summary>
 		/// Constructor.
 		/// </summary>
-		/// <param name="InParameters">Parameters for this task</param>
-		public StageTask(StageTaskParameters InParameters)
+		/// <param name="parameters">Parameters for this task</param>
+		public StageTask(StageTaskParameters parameters)
 		{
-			Parameters = InParameters;
+			_parameters = parameters;
 		}
 
 		/// <summary>
-		/// Execute the task.
+		/// ExecuteAsync the task.
 		/// </summary>
-		/// <param name="Job">Information about the current job</param>
-		/// <param name="BuildProducts">Set of build products produced by this node.</param>
-		/// <param name="TagNameToFileSet">Mapping from tag names to the set of files they include</param>
-		public override Task ExecuteAsync(JobContext Job, HashSet<FileReference> BuildProducts, Dictionary<string, HashSet<FileReference>> TagNameToFileSet)
+		/// <param name="job">Information about the current job</param>
+		/// <param name="buildProducts">Set of build products produced by this node.</param>
+		/// <param name="tagNameToFileSet">Mapping from tag names to the set of files they include</param>
+		public override Task ExecuteAsync(JobContext job, HashSet<FileReference> buildProducts, Dictionary<string, HashSet<FileReference>> tagNameToFileSet)
 		{
 			// Get the project path, and check it exists
-			FileReference ProjectFile = Parameters.Project;
-			if(Parameters.Project != null && !FileReference.Exists(ProjectFile))
+			FileReference projectFile = _parameters.Project;
+			if (_parameters.Project != null && !FileReference.Exists(projectFile))
 			{
-				throw new AutomationException("Couldn't find project '{0}'", ProjectFile.FullName);
+				throw new AutomationException("Couldn't find project '{0}'", projectFile.FullName);
 			}
 
 			// Get the directories used for staging this project
-			DirectoryReference SourceEngineDir = Unreal.EngineDirectory;
-			DirectoryReference SourceProjectDir = (ProjectFile == null)? SourceEngineDir : ProjectFile.Directory;
+			DirectoryReference sourceEngineDir = Unreal.EngineDirectory;
+			DirectoryReference sourceProjectDir = (projectFile == null) ? sourceEngineDir : projectFile.Directory;
 
 			// Get the output directories. We flatten the directory structure on output.
-			DirectoryReference TargetDir = Parameters.ToDir;
-			DirectoryReference TargetEngineDir = DirectoryReference.Combine(TargetDir, "Engine");
-			DirectoryReference TargetProjectDir = (ProjectFile == null) ? TargetEngineDir : DirectoryReference.Combine(TargetDir, ProjectFile.GetFileNameWithoutExtension());
+			DirectoryReference targetDir = _parameters.ToDir;
+			DirectoryReference targetEngineDir = DirectoryReference.Combine(targetDir, "Engine");
+			DirectoryReference targetProjectDir = (projectFile == null) ? targetEngineDir : DirectoryReference.Combine(targetDir, projectFile.GetFileNameWithoutExtension());
 
 			// Get the path to the receipt
-			FileReference ReceiptFileName = TargetReceipt.GetDefaultPath(SourceProjectDir, Parameters.Target, Parameters.Platform, Parameters.Configuration, UnrealArchitectures.FromString(Parameters.Architecture, Parameters.Platform));
+			FileReference receiptFileName = TargetReceipt.GetDefaultPath(sourceProjectDir, _parameters.Target, _parameters.Platform, _parameters.Configuration, UnrealArchitectures.FromString(_parameters.Architecture, _parameters.Platform));
 
 			// Try to load it
-			TargetReceipt Receipt;
-			if(!TargetReceipt.TryRead(ReceiptFileName, out Receipt))
+			TargetReceipt receipt;
+			if (!TargetReceipt.TryRead(receiptFileName, out receipt))
 			{
-				throw new AutomationException("Couldn't read receipt '{0}'", ReceiptFileName);
+				throw new AutomationException("Couldn't read receipt '{0}'", receiptFileName);
 			}
 
 			// Stage all the build products needed at runtime
-			HashSet<FileReference> SourceFiles = new HashSet<FileReference>();
-			foreach(BuildProduct BuildProduct in Receipt.BuildProducts)
+			HashSet<FileReference> sourceFiles = new HashSet<FileReference>();
+			foreach (BuildProduct buildProduct in receipt.BuildProducts)
 			{
-				SourceFiles.Add(BuildProduct.Path);
+				sourceFiles.Add(buildProduct.Path);
 			}
-			foreach(RuntimeDependency RuntimeDependency in Receipt.RuntimeDependencies.Where(x => x.Type != StagedFileType.UFS))
+			foreach (RuntimeDependency runtimeDependency in receipt.RuntimeDependencies.Where(x => x.Type != StagedFileType.UFS))
 			{
-				SourceFiles.Add(RuntimeDependency.Path);
+				sourceFiles.Add(runtimeDependency.Path);
 			}
 
 			// Get all the target files
-			List<FileReference> TargetFiles = new List<FileReference>();
-			foreach(FileReference SourceFile in SourceFiles)
+			List<FileReference> targetFiles = new List<FileReference>();
+			foreach (FileReference sourceFile in sourceFiles)
 			{
 				// Get the destination file to copy to, mapping to the new engine and project directories as appropriate
-				FileReference TargetFile;
-				if(SourceFile.IsUnderDirectory(SourceEngineDir))
+				FileReference targetFile;
+				if (sourceFile.IsUnderDirectory(sourceEngineDir))
 				{
-					TargetFile = FileReference.Combine(TargetEngineDir, SourceFile.MakeRelativeTo(SourceEngineDir));
+					targetFile = FileReference.Combine(targetEngineDir, sourceFile.MakeRelativeTo(sourceEngineDir));
 				}
 				else
 				{
-					TargetFile = FileReference.Combine(TargetProjectDir, SourceFile.MakeRelativeTo(SourceProjectDir));
+					targetFile = FileReference.Combine(targetProjectDir, sourceFile.MakeRelativeTo(sourceProjectDir));
 				}
 
 				// Only copy the output file if it doesn't already exist. We can stage multiple targets to the same output directory.
-				if(Parameters.Overwrite || !FileReference.Exists(TargetFile))
+				if (_parameters.Overwrite || !FileReference.Exists(targetFile))
 				{
-					DirectoryReference.CreateDirectory(TargetFile.Directory);
-					CommandUtils.CopyFile(SourceFile.FullName, TargetFile.FullName);
+					DirectoryReference.CreateDirectory(targetFile.Directory);
+					CommandUtils.CopyFile(sourceFile.FullName, targetFile.FullName);
 					// Force all destination files to not readonly.
-					CommandUtils.SetFileAttributes(TargetFile.FullName, ReadOnly: false);
+					CommandUtils.SetFileAttributes(targetFile.FullName, ReadOnly: false);
 				}
 
 				// Add it to the list of target files
-				TargetFiles.Add(TargetFile);
+				targetFiles.Add(targetFile);
 			}
 
 			// Apply the optional tag to the build products
-			foreach(string TagName in FindTagNamesFromList(Parameters.Tag))
+			foreach (string tagName in FindTagNamesFromList(_parameters.Tag))
 			{
-				FindOrAddTagSet(TagNameToFileSet, TagName).UnionWith(TargetFiles);
+				FindOrAddTagSet(tagNameToFileSet, tagName).UnionWith(targetFiles);
 			}
 
 			// Add the target file to the list of build products
-			BuildProducts.UnionWith(TargetFiles);
+			buildProducts.UnionWith(targetFiles);
 			return Task.CompletedTask;
 		}
 
 		/// <summary>
 		/// Output this task out to an XML writer.
 		/// </summary>
-		public override void Write(XmlWriter Writer)
+		public override void Write(XmlWriter writer)
 		{
-			Write(Writer, Parameters);
+			Write(writer, _parameters);
 		}
 
 		/// <summary>
@@ -195,7 +188,7 @@ namespace AutomationTool.Tasks
 		/// <returns>The tag names which are modified by this task</returns>
 		public override IEnumerable<string> FindProducedTagNames()
 		{
-			return FindTagNamesFromList(Parameters.Tag);
+			return FindTagNamesFromList(_parameters.Tag);
 		}
 	}
 }
diff --git a/Engine/Source/Programs/AutomationTool/BuildGraph/Tasks/StripTask.cs b/Engine/Source/Programs/AutomationTool/BuildGraph/Tasks/StripTask.cs
--- a/Engine/Source/Programs/AutomationTool/BuildGraph/Tasks/StripTask.cs
+++ b/Engine/Source/Programs/AutomationTool/BuildGraph/Tasks/StripTask.cs
@@ -1,19 +1,13 @@
 // Copyright Epic Games, Inc. All Rights Reserved.
 
-using AutomationTool;
-using AutomationTool.Tasks;
-using EpicGames.BuildGraph;
 using System;
 using System.Collections.Generic;
 using System.Linq;
-using System.Text;
 using System.Threading.Tasks;
 using System.Xml;
 using EpicGames.Core;
-using UnrealBuildTool;
 using Microsoft.Extensions.Logging;
-
-using static AutomationTool.CommandUtils;
+using UnrealBuildTool;
 
 namespace AutomationTool.Tasks
 {
@@ -59,73 +53,70 @@ namespace AutomationTool.Tasks
 	[TaskElement("Strip", typeof(StripTaskParameters))]
 	public class StripTask : BgTaskImpl
 	{
-		/// <summary>
-		/// Parameters for this task
-		/// </summary>
-		StripTaskParameters Parameters;
+		readonly StripTaskParameters _parameters;
 
 		/// <summary>
 		/// Construct a spawn task
 		/// </summary>
-		/// <param name="InParameters">Parameters for the task</param>
-		public StripTask(StripTaskParameters InParameters)
+		/// <param name="parameters">Parameters for the task</param>
+		public StripTask(StripTaskParameters parameters)
 		{
-			Parameters = InParameters;
+			_parameters = parameters;
 		}
 
 		/// <summary>
-		/// Execute the task.
+		/// ExecuteAsync the task.
 		/// </summary>
-		/// <param name="Job">Information about the current job</param>
-		/// <param name="BuildProducts">Set of build products produced by this node.</param>
-		/// <param name="TagNameToFileSet">Mapping from tag names to the set of files they include</param>
-		public override Task ExecuteAsync(JobContext Job, HashSet<FileReference> BuildProducts, Dictionary<string, HashSet<FileReference>> TagNameToFileSet)
+		/// <param name="job">Information about the current job</param>
+		/// <param name="buildProducts">Set of build products produced by this node.</param>
+		/// <param name="tagNameToFileSet">Mapping from tag names to the set of files they include</param>
+		public override Task ExecuteAsync(JobContext job, HashSet<FileReference> buildProducts, Dictionary<string, HashSet<FileReference>> tagNameToFileSet)
 		{
 			// Get the base directory
-			DirectoryReference BaseDir = Parameters.BaseDir;
+			DirectoryReference baseDir = _parameters.BaseDir;
 
 			// Get the output directory
-			DirectoryReference OutputDir = Parameters.OutputDir;
+			DirectoryReference outputDir = _parameters.OutputDir;
 
 			// Find the matching files
-			FileReference[] SourceFiles = ResolveFilespec(BaseDir, Parameters.Files, TagNameToFileSet).OrderBy(x => x.FullName).ToArray();
+			FileReference[] sourceFiles = ResolveFilespec(baseDir, _parameters.Files, tagNameToFileSet).OrderBy(x => x.FullName).ToArray();
 
 			// Create the matching target files
-			FileReference[] TargetFiles = SourceFiles.Select(x => FileReference.Combine(OutputDir, x.MakeRelativeTo(BaseDir))).ToArray();
+			FileReference[] targetFiles = sourceFiles.Select(x => FileReference.Combine(outputDir, x.MakeRelativeTo(baseDir))).ToArray();
 
 			// Run the stripping command
-			Platform TargetPlatform = Platform.GetPlatform(Parameters.Platform);
-			for (int Idx = 0; Idx < SourceFiles.Length; Idx++)
+			Platform targetPlatform = Platform.GetPlatform(_parameters.Platform);
+			for (int idx = 0; idx < sourceFiles.Length; idx++)
 			{
-				DirectoryReference.CreateDirectory(TargetFiles[Idx].Directory);
-				if (SourceFiles[Idx] == TargetFiles[Idx])
+				DirectoryReference.CreateDirectory(targetFiles[idx].Directory);
+				if (sourceFiles[idx] == targetFiles[idx])
 				{
-					Logger.LogInformation("Stripping symbols: {Arg0}", SourceFiles[Idx].FullName);
+					Logger.LogInformation("Stripping symbols: {Arg0}", sourceFiles[idx].FullName);
 				}
 				else
 				{
-					Logger.LogInformation("Stripping symbols: {Arg0} -> {Arg1}", SourceFiles[Idx].FullName, TargetFiles[Idx].FullName);
+					Logger.LogInformation("Stripping symbols: {Arg0} -> {Arg1}", sourceFiles[idx].FullName, targetFiles[idx].FullName);
 				}
-				TargetPlatform.StripSymbols(SourceFiles[Idx], TargetFiles[Idx]);
+				targetPlatform.StripSymbols(sourceFiles[idx], targetFiles[idx]);
 			}
 
 			// Apply the optional tag to the build products
-			foreach(string TagName in FindTagNamesFromList(Parameters.Tag))
+			foreach (string tagName in FindTagNamesFromList(_parameters.Tag))
 			{
-				FindOrAddTagSet(TagNameToFileSet, TagName).UnionWith(TargetFiles);
+				FindOrAddTagSet(tagNameToFileSet, tagName).UnionWith(targetFiles);
 			}
 
 			// Add the target files to the set of build products
-			BuildProducts.UnionWith(TargetFiles);
+			buildProducts.UnionWith(targetFiles);
 			return Task.CompletedTask;
 		}
 
 		/// <summary>
 		/// Output this task out to an XML writer.
 		/// </summary>
-		public override void Write(XmlWriter Writer)
+		public override void Write(XmlWriter writer)
 		{
-			Write(Writer, Parameters);
+			Write(writer, _parameters);
 		}
 
 		/// <summary>
@@ -134,7 +125,7 @@ namespace AutomationTool.Tasks
 		/// <returns>The tag names which are read by this task</returns>
 		public override IEnumerable<string> FindConsumedTagNames()
 		{
-			return FindTagNamesFromFilespec(Parameters.Files);
+			return FindTagNamesFromFilespec(_parameters.Files);
 		}
 
 		/// <summary>
@@ -143,7 +134,7 @@ namespace AutomationTool.Tasks
 		/// <returns>The tag names which are modified by this task</returns>
 		public override IEnumerable<string> FindProducedTagNames()
 		{
-			return FindTagNamesFromList(Parameters.Tag);
+			return FindTagNamesFromList(_parameters.Tag);
 		}
 	}
 
@@ -152,19 +143,19 @@ namespace AutomationTool.Tasks
 		/// <summary>
 		/// Strips symbols from a set of files.
 		/// </summary>
-		/// <param name="Files"></param>
-		/// <param name="Platform">The platform toolchain to strip binaries.</param>
-		/// <param name="BaseDir">The directory to find files in.</param>
-		/// <param name="OutputDir">Output directory for the stripped files. Defaults to the input path, overwriting the input files.</param>
+		/// <param name="files"></param>
+		/// <param name="platform">The platform toolchain to strip binaries.</param>
+		/// <param name="baseDir">The directory to find files in.</param>
+		/// <param name="outputDir">Output directory for the stripped files. Defaults to the input path, overwriting the input files.</param>
 		/// <returns></returns>
-		public static async Task<FileSet> StripAsync(FileSet Files, UnrealTargetPlatform Platform, DirectoryReference BaseDir = null, DirectoryReference OutputDir = null)
+		public static async Task<FileSet> StripAsync(FileSet files, UnrealTargetPlatform platform, DirectoryReference baseDir = null, DirectoryReference outputDir = null)
 		{
-			StripTaskParameters Parameters = new StripTaskParameters();
-			Parameters.Platform = Platform;
-			Parameters.BaseDir = BaseDir;
-			Parameters.Files = String.Join(";", Files.Flatten().Values.Select(x => x.FullName));
-			Parameters.OutputDir = OutputDir;
-			return await ExecuteAsync(new StripTask(Parameters));
+			StripTaskParameters parameters = new StripTaskParameters();
+			parameters.Platform = platform;
+			parameters.BaseDir = baseDir;
+			parameters.Files = String.Join(";", files.Flatten().Values.Select(x => x.FullName));
+			parameters.OutputDir = outputDir;
+			return await ExecuteAsync(new StripTask(parameters));
 		}
 	}
 }
diff --git a/Engine/Source/Programs/AutomationTool/BuildGraph/Tasks/SubmitTask.cs b/Engine/Source/Programs/AutomationTool/BuildGraph/Tasks/SubmitTask.cs
--- a/Engine/Source/Programs/AutomationTool/BuildGraph/Tasks/SubmitTask.cs
+++ b/Engine/Source/Programs/AutomationTool/BuildGraph/Tasks/SubmitTask.cs
@@ -1,19 +1,13 @@
 // Copyright Epic Games, Inc. All Rights Reserved.
 
-using EpicGames.BuildGraph;
 using System;
 using System.Collections.Generic;
-using System.Linq;
-using System.Text;
 using System.Text.RegularExpressions;
 using System.Threading.Tasks;
 using System.Xml;
 using EpicGames.Core;
-using UnrealBuildBase;
-using UnrealBuildTool;
 using Microsoft.Extensions.Logging;
-
-using static AutomationTool.CommandUtils;
+using UnrealBuildBase;
 
 namespace AutomationTool.Tasks
 {
@@ -43,7 +37,7 @@ namespace AutomationTool.Tasks
 		/// <summary>
 		/// The workspace name. If specified, a new workspace will be created using the given stream and root directory to submit the files. If not, the current workspace will be used.
 		/// </summary>
-		[TaskParameter(Optional=true)]
+		[TaskParameter(Optional = true)]
 		public string Workspace { get; set; }
 
 		/// <summary>
@@ -89,30 +83,27 @@ namespace AutomationTool.Tasks
 	[TaskElement("Submit", typeof(SubmitTaskParameters))]
 	public class SubmitTask : BgTaskImpl
 	{
-		/// <summary>
-		/// Parameters for the task
-		/// </summary>
-		SubmitTaskParameters Parameters;
+		readonly SubmitTaskParameters _parameters;
 
 		/// <summary>
 		/// Construct a version task
 		/// </summary>
-		/// <param name="InParameters">Parameters for this task</param>
-		public SubmitTask(SubmitTaskParameters InParameters)
+		/// <param name="parameters">Parameters for this task</param>
+		public SubmitTask(SubmitTaskParameters parameters)
 		{
-			Parameters = InParameters;
+			_parameters = parameters;
 		}
 
 		/// <summary>
-		/// Execute the task.
+		/// ExecuteAsync the task.
 		/// </summary>
-		/// <param name="Job">Information about the current job</param>
-		/// <param name="BuildProducts">Set of build products produced by this node.</param>
-		/// <param name="TagNameToFileSet">Mapping from tag names to the set of files they include</param>
-		public override Task ExecuteAsync(JobContext Job, HashSet<FileReference> BuildProducts, Dictionary<string, HashSet<FileReference>> TagNameToFileSet)
+		/// <param name="job">Information about the current job</param>
+		/// <param name="buildProducts">Set of build products produced by this node.</param>
+		/// <param name="tagNameToFileSet">Mapping from tag names to the set of files they include</param>
+		public override Task ExecuteAsync(JobContext job, HashSet<FileReference> buildProducts, Dictionary<string, HashSet<FileReference>> tagNameToFileSet)
 		{
-			HashSet<FileReference> Files = ResolveFilespec(Unreal.RootDirectory, Parameters.Files, TagNameToFileSet);
-			if (Files.Count == 0)
+			HashSet<FileReference> files = ResolveFilespec(Unreal.RootDirectory, _parameters.Files, tagNameToFileSet);
+			if (files.Count == 0)
 			{
 				Logger.LogInformation("No files to submit.");
 			}
@@ -125,50 +116,50 @@ namespace AutomationTool.Tasks
 				try
 				{
 					// Get the connection that we're going to submit with
-					P4Connection SubmitP4 = CommandUtils.P4;
-					if (Parameters.Workspace != null)
+					P4Connection submitP4 = CommandUtils.P4;
+					if (_parameters.Workspace != null)
 					{
 						// Create a brand new workspace
-						P4ClientInfo Client = new P4ClientInfo();
-						Client.Owner = CommandUtils.P4Env.User;
-						Client.Host = Unreal.MachineName;
-						Client.RootPath = Parameters.RootDir.FullName ?? Unreal.RootDirectory.FullName;
-						Client.Name = $"{Parameters.Workspace}_{Regex.Replace(Client.Host, "[^a-zA-Z0-9]", "-")}_{ContentHash.MD5((CommandUtils.P4Env.ServerAndPort ?? "").ToUpperInvariant())}";
-						Client.Options = P4ClientOption.NoAllWrite | P4ClientOption.Clobber | P4ClientOption.NoCompress | P4ClientOption.Unlocked | P4ClientOption.NoModTime | P4ClientOption.RmDir;
-						Client.LineEnd = P4LineEnd.Local;
-						if (!String.IsNullOrEmpty(Parameters.Branch))
+						P4ClientInfo client = new P4ClientInfo();
+						client.Owner = CommandUtils.P4Env.User;
+						client.Host = Unreal.MachineName;
+						client.RootPath = _parameters.RootDir.FullName ?? Unreal.RootDirectory.FullName;
+						client.Name = $"{_parameters.Workspace}_{Regex.Replace(client.Host, "[^a-zA-Z0-9]", "-")}_{ContentHash.MD5((CommandUtils.P4Env.ServerAndPort ?? "").ToUpperInvariant())}";
+						client.Options = P4ClientOption.NoAllWrite | P4ClientOption.Clobber | P4ClientOption.NoCompress | P4ClientOption.Unlocked | P4ClientOption.NoModTime | P4ClientOption.RmDir;
+						client.LineEnd = P4LineEnd.Local;
+						if (!String.IsNullOrEmpty(_parameters.Branch))
 						{
-							Client.View.Add(new KeyValuePair<string, string>($"{Parameters.Branch}/...", $"/..."));
+							client.View.Add(new KeyValuePair<string, string>($"{_parameters.Branch}/...", $"/..."));
 						}
 						else
 						{
-							Client.Stream = Parameters.Stream ?? CommandUtils.P4Env.Branch;
+							client.Stream = _parameters.Stream ?? CommandUtils.P4Env.Branch;
 						}
-						CommandUtils.P4.CreateClient(Client, AllowSpew: Parameters.P4Verbose);
+						CommandUtils.P4.CreateClient(client, AllowSpew: _parameters.P4Verbose);
 
 						// Create a new connection for it
-						SubmitP4 = new P4Connection(Client.Owner, Client.Name);
+						submitP4 = new P4Connection(client.Owner, client.Name);
 					}
 
 					// Get the latest version of it
-					int NewCL = SubmitP4.CreateChange(Description: Parameters.Description.Replace("\\n", "\n", StringComparison.Ordinal));
-					foreach(FileReference File in Files)
+					int newCl = submitP4.CreateChange(Description: _parameters.Description.Replace("\\n", "\n", StringComparison.Ordinal));
+					foreach (FileReference file in files)
 					{
-						SubmitP4.Revert(String.Format("-k \"{0}\"", File.FullName), AllowSpew: Parameters.P4Verbose);
-						SubmitP4.Sync(String.Format("-k \"{0}\"", File.FullName), AllowSpew: Parameters.P4Verbose);
-						SubmitP4.Add(NewCL, String.Format("\"{0}\"", File.FullName));
-						SubmitP4.Edit(NewCL, String.Format("\"{0}\"", File.FullName), AllowSpew: Parameters.P4Verbose);
-						if (Parameters.FileType != null)
+						submitP4.Revert(String.Format("-k \"{0}\"", file.FullName), AllowSpew: _parameters.P4Verbose);
+						submitP4.Sync(String.Format("-k \"{0}\"", file.FullName), AllowSpew: _parameters.P4Verbose);
+						submitP4.Add(newCl, String.Format("\"{0}\"", file.FullName));
+						submitP4.Edit(newCl, String.Format("\"{0}\"", file.FullName), AllowSpew: _parameters.P4Verbose);
+						if (_parameters.FileType != null)
 						{
-							SubmitP4.P4(String.Format("reopen -t \"{0}\" \"{1}\"", Parameters.FileType, File.FullName), AllowSpew: Parameters.P4Verbose);
+							submitP4.P4(String.Format("reopen -t \"{0}\" \"{1}\"", _parameters.FileType, file.FullName), AllowSpew: _parameters.P4Verbose);
 						}
 					}
 
 					// Revert any unchanged files
-					if(Parameters.RevertUnchanged)
+					if (_parameters.RevertUnchanged)
 					{
-						SubmitP4.RevertUnchanged(NewCL);
-						if(SubmitP4.TryDeleteEmptyChange(NewCL))
+						submitP4.RevertUnchanged(newCl);
+						if (submitP4.TryDeleteEmptyChange(newCl))
 						{
 							Logger.LogInformation("No files to submit; ignored.");
 							return Task.CompletedTask;
@@ -176,19 +167,19 @@ namespace AutomationTool.Tasks
 					}
 
 					// Submit it
-					int SubmittedCL;
-					SubmitP4.Submit(NewCL, out SubmittedCL, Force: Parameters.Force);
-					if (SubmittedCL <= 0)
+					int submittedCl;
+					submitP4.Submit(newCl, out submittedCl, Force: _parameters.Force);
+					if (submittedCl <= 0)
 					{
 						throw new AutomationException("Submit failed.");
 					}
 
-					Logger.LogInformation("Submitted in changelist {SubmittedCL}", SubmittedCL);
+					Logger.LogInformation("Submitted in changelist {SubmittedCL}", submittedCl);
 				}
-				catch (P4Exception Ex)
+				catch (P4Exception ex)
 				{
-					Logger.LogError(KnownLogEvents.Systemic_Perforce, "{Message}", Ex.Message);
-					throw new AutomationException(Ex.ErrorCode, Ex, "{0}", Ex.Message) { OutputFormat = AutomationExceptionOutputFormat.Silent };
+					Logger.LogError(KnownLogEvents.Systemic_Perforce, "{Message}", ex.Message);
+					throw new AutomationException(ex.ErrorCode, ex, "{0}", ex.Message) { OutputFormat = AutomationExceptionOutputFormat.Silent };
 				}
 			}
 			return Task.CompletedTask;
@@ -197,9 +188,9 @@ namespace AutomationTool.Tasks
 		/// <summary>
 		/// Output this task out to an XML writer.
 		/// </summary>
-		public override void Write(XmlWriter Writer)
+		public override void Write(XmlWriter writer)
 		{
-			Write(Writer, Parameters);
+			Write(writer, _parameters);
 		}
 
 		/// <summary>
@@ -208,7 +199,7 @@ namespace AutomationTool.Tasks
 		/// <returns>The tag names which are read by this task</returns>
 		public override IEnumerable<string> FindConsumedTagNames()
 		{
-			return FindTagNamesFromFilespec(Parameters.Files);
+			return FindTagNamesFromFilespec(_parameters.Files);
 		}
 
 		/// <summary>
diff --git a/Engine/Source/Programs/AutomationTool/BuildGraph/Tasks/SymStoreTask.cs b/Engine/Source/Programs/AutomationTool/BuildGraph/Tasks/SymStoreTask.cs
--- a/Engine/Source/Programs/AutomationTool/BuildGraph/Tasks/SymStoreTask.cs
+++ b/Engine/Source/Programs/AutomationTool/BuildGraph/Tasks/SymStoreTask.cs
@@ -2,49 +2,44 @@
 
 using System;
 using System.Collections.Generic;
-using System.IO;
 using System.Linq;
-using System.Text;
-using System.Threading;
 using System.Threading.Tasks;
 using System.Xml;
 using EpicGames.Core;
+using Microsoft.Extensions.Logging;
 using UnrealBuildBase;
 using UnrealBuildTool;
-using Microsoft.Extensions.Logging;
-
-using static AutomationTool.CommandUtils;
 
 namespace AutomationTool.Tasks
 {
-    /// <summary>
-    /// Parameters for a task that uploads symbols to a symbol server
-    /// </summary>
-    public class SymStoreTaskParameters
-    {
-        /// <summary>
-        /// The platform toolchain required to handle symbol files.
-        /// </summary>
-        [TaskParameter]
-        public UnrealTargetPlatform Platform { get; set; }
+	/// <summary>
+	/// Parameters for a task that uploads symbols to a symbol server
+	/// </summary>
+	public class SymStoreTaskParameters
+	{
+		/// <summary>
+		/// The platform toolchain required to handle symbol files.
+		/// </summary>
+		[TaskParameter]
+		public UnrealTargetPlatform Platform { get; set; }
 
 		/// <summary>
 		/// List of output files. PDBs will be extracted from this list.
 		/// </summary>
 		[TaskParameter]
-        public string Files { get; set; }
+		public string Files { get; set; }
 
 		/// <summary>
 		/// Output directory for the compressed symbols.
 		/// </summary>
 		[TaskParameter]
-        public string StoreDir { get; set; }
+		public string StoreDir { get; set; }
 
 		/// <summary>
 		/// Name of the product for the symbol store records.
 		/// </summary>
 		[TaskParameter]
-        public string Product { get; set; }
+		public string Product { get; set; }
 
 		/// <summary>
 		/// Name of the Branch to base all the depot source files from.
@@ -84,56 +79,56 @@ namespace AutomationTool.Tasks
 		public string SourceFiles { get; set; }
 	}
 
-    /// <summary>
-    /// Task that strips symbols from a set of files.
-    /// </summary>
-    [TaskElement("SymStore", typeof(SymStoreTaskParameters))]
-    public class SymStoreTask : BgTaskImpl
-    {
-        /// <summary>
-        /// Parameters for this task
-        /// </summary>
-        SymStoreTaskParameters Parameters;
-
-        /// <summary>
-        /// Construct a spawn task
-        /// </summary>
-        /// <param name="InParameters">Parameters for the task</param>
-        public SymStoreTask(SymStoreTaskParameters InParameters)
-        {
-            Parameters = InParameters;
-        }
-
-        /// <summary>
-        /// Execute the task.
-        /// </summary>
-        /// <param name="Job">Information about the current job</param>
-        /// <param name="BuildProducts">Set of build products produced by this node.</param>
-        /// <param name="TagNameToFileSet">Mapping from tag names to the set of files they include</param>
-        public override Task ExecuteAsync(JobContext Job, HashSet<FileReference> BuildProducts, Dictionary<string, HashSet<FileReference>> TagNameToFileSet)
-        {
+	/// <summary>
+	/// Task that strips symbols from a set of files.
+	/// </summary>
+	[TaskElement("SymStore", typeof(SymStoreTaskParameters))]
+	public class SymStoreTask : BgTaskImpl
+	{
+		/// <summary>
+		/// Parameters for this task
+		/// </summary>
+		readonly SymStoreTaskParameters _parameters;
+
+		/// <summary>
+		/// Construct a spawn task
+		/// </summary>
+		/// <param name="parameters">Parameters for the task</param>
+		public SymStoreTask(SymStoreTaskParameters parameters)
+		{
+			_parameters = parameters;
+		}
+
+		/// <summary>
+		/// ExecuteAsync the task.
+		/// </summary>
+		/// <param name="job">Information about the current job</param>
+		/// <param name="buildProducts">Set of build products produced by this node.</param>
+		/// <param name="tagNameToFileSet">Mapping from tag names to the set of files they include</param>
+		public override Task ExecuteAsync(JobContext job, HashSet<FileReference> buildProducts, Dictionary<string, HashSet<FileReference>> tagNameToFileSet)
+		{
 			// Find the matching files
-			List<FileReference> Files = ResolveFilespec(Unreal.RootDirectory, Parameters.Files, TagNameToFileSet).ToList();
-            
-            // Get the symbol store directory
-            DirectoryReference StoreDir = ResolveDirectory(Parameters.StoreDir);
+			List<FileReference> files = ResolveFilespec(Unreal.RootDirectory, _parameters.Files, tagNameToFileSet).ToList();
+
+			// Get the symbol store directory
+			DirectoryReference storeDir = ResolveDirectory(_parameters.StoreDir);
 
 			// Take the lock before accessing the symbol server, if required by the platform
-			Platform TargetPlatform = Platform.GetPlatform(Parameters.Platform);
+			Platform targetPlatform = Platform.GetPlatform(_parameters.Platform);
 
-			List<FileReference> SourceFiles = new List<FileReference>();
+			List<FileReference> sourceFiles = new List<FileReference>();
 
-			if (Parameters.IndexSources && TargetPlatform.SymbolServerSourceIndexingRequiresListOfSourceFiles)
+			if (_parameters.IndexSources && targetPlatform.SymbolServerSourceIndexingRequiresListOfSourceFiles)
 			{
 				Logger.LogInformation("Discovering source code files...");
 
-				SourceFiles = ResolveFilespec(Unreal.RootDirectory, Parameters.SourceFiles, TagNameToFileSet).ToList();
+				sourceFiles = ResolveFilespec(Unreal.RootDirectory, _parameters.SourceFiles, tagNameToFileSet).ToList();
 			}
 
-			CommandUtils.OptionallyTakeLock(TargetPlatform.SymbolServerRequiresLock, StoreDir, TimeSpan.FromMinutes(60), () =>
+			CommandUtils.OptionallyTakeLock(targetPlatform.SymbolServerRequiresLock, storeDir, TimeSpan.FromMinutes(60), () =>
 			{
-				if (!TargetPlatform.PublishSymbols(StoreDir, Files, Parameters.IndexSources, SourceFiles,
-					Parameters.Product, Parameters.Branch, Parameters.Change, Parameters.BuildVersion))
+				if (!targetPlatform.PublishSymbols(storeDir, files, _parameters.IndexSources, sourceFiles,
+					_parameters.Product, _parameters.Branch, _parameters.Change, _parameters.BuildVersion))
 				{
 					throw new AutomationException("Failure publishing symbol files.");
 				}
@@ -142,30 +137,30 @@ namespace AutomationTool.Tasks
 			return Task.CompletedTask;
 		}
 
-        /// <summary>
-        /// Output this task out to an XML writer.
-        /// </summary>
-        public override void Write(XmlWriter Writer)
-        {
-            Write(Writer, Parameters);
-        }
-
-        /// <summary>
-        /// Find all the tags which are used as inputs to this task
-        /// </summary>
-        /// <returns>The tag names which are read by this task</returns>
-        public override IEnumerable<string> FindConsumedTagNames()
-        {
-            return FindTagNamesFromFilespec(Parameters.Files);
-        }
-
-        /// <summary>
-        /// Find all the tags which are modified by this task
-        /// </summary>
-        /// <returns>The tag names which are modified by this task</returns>
-        public override IEnumerable<string> FindProducedTagNames()
-        {
-            yield break;
-        }
-    }
+		/// <summary>
+		/// Output this task out to an XML writer.
+		/// </summary>
+		public override void Write(XmlWriter writer)
+		{
+			Write(writer, _parameters);
+		}
+
+		/// <summary>
+		/// Find all the tags which are used as inputs to this task
+		/// </summary>
+		/// <returns>The tag names which are read by this task</returns>
+		public override IEnumerable<string> FindConsumedTagNames()
+		{
+			return FindTagNamesFromFilespec(_parameters.Files);
+		}
+
+		/// <summary>
+		/// Find all the tags which are modified by this task
+		/// </summary>
+		/// <returns>The tag names which are modified by this task</returns>
+		public override IEnumerable<string> FindProducedTagNames()
+		{
+			yield break;
+		}
+	}
 }
diff --git a/Engine/Source/Programs/AutomationTool/BuildGraph/Tasks/TagReceiptTask.cs b/Engine/Source/Programs/AutomationTool/BuildGraph/Tasks/TagReceiptTask.cs
--- a/Engine/Source/Programs/AutomationTool/BuildGraph/Tasks/TagReceiptTask.cs
+++ b/Engine/Source/Programs/AutomationTool/BuildGraph/Tasks/TagReceiptTask.cs
@@ -1,19 +1,13 @@
 // Copyright Epic Games, Inc. All Rights Reserved.
 
-using EpicGames.BuildGraph;
 using System;
 using System.Collections.Generic;
-using System.IO;
-using System.Linq;
-using System.Text;
 using System.Threading.Tasks;
 using System.Xml;
 using EpicGames.Core;
+using Microsoft.Extensions.Logging;
 using UnrealBuildBase;
 using UnrealBuildTool;
-using Microsoft.Extensions.Logging;
-
-using static AutomationTool.CommandUtils;
 
 namespace AutomationTool.Tasks
 {
@@ -80,136 +74,136 @@ namespace AutomationTool.Tasks
 		/// <summary>
 		/// Parameters to this task
 		/// </summary>
-		TagReceiptTaskParameters Parameters;
+		readonly TagReceiptTaskParameters _parameters;
 
 		/// <summary>
 		/// The type of build products to enumerate. May be null.
 		/// </summary>
-		Nullable<BuildProductType> BuildProductType;
+		readonly BuildProductType? _buildProductType;
 
 		/// <summary>
 		/// The type of staged files to enumerate. May be null,
 		/// </summary>
-		Nullable<StagedFileType> StagedFileType;
+		readonly StagedFileType? _stagedFileType;
 
 		/// <summary>
 		/// Constructor
 		/// </summary>
-		/// <param name="InParameters">Parameters to select which files to search</param>
-		public TagReceiptTask(TagReceiptTaskParameters InParameters)
+		/// <param name="parameters">Parameters to select which files to search</param>
+		public TagReceiptTask(TagReceiptTaskParameters parameters)
 		{
-			Parameters = InParameters;
+			_parameters = parameters;
 
-			if (!String.IsNullOrEmpty(Parameters.BuildProductType))
+			if (!String.IsNullOrEmpty(_parameters.BuildProductType))
 			{
-				BuildProductType = (BuildProductType)Enum.Parse(typeof(BuildProductType), Parameters.BuildProductType);
+				_buildProductType = (BuildProductType)Enum.Parse(typeof(BuildProductType), _parameters.BuildProductType);
 			}
-			if (!String.IsNullOrEmpty(Parameters.StagedFileType))
+			if (!String.IsNullOrEmpty(_parameters.StagedFileType))
 			{
-				StagedFileType = (StagedFileType)Enum.Parse(typeof(StagedFileType), Parameters.StagedFileType);
+				_stagedFileType = (StagedFileType)Enum.Parse(typeof(StagedFileType), _parameters.StagedFileType);
 			}
 		}
 
 		/// <summary>
-		/// Execute the task.
+		/// ExecuteAsync the task.
 		/// </summary>
-		/// <param name="Job">Information about the current job</param>
-		/// <param name="BuildProducts">Set of build products produced by this node.</param>
-		/// <param name="TagNameToFileSet">Mapping from tag names to the set of files they include</param>
-		public override async Task ExecuteAsync(JobContext Job, HashSet<FileReference> BuildProducts, Dictionary<string, HashSet<FileReference>> TagNameToFileSet)
+		/// <param name="job">Information about the current job</param>
+		/// <param name="buildProducts">Set of build products produced by this node.</param>
+		/// <param name="tagNameToFileSet">Mapping from tag names to the set of files they include</param>
+		public override async Task ExecuteAsync(JobContext job, HashSet<FileReference> buildProducts, Dictionary<string, HashSet<FileReference>> tagNameToFileSet)
 		{
 			// Output a warning if the project directory is specified
-			if (Parameters.ProjectDir != null)
+			if (_parameters.ProjectDir != null)
 			{
 				Logger.LogWarning("The ProjectDir argument to the TagReceipt parameter is deprecated. This path is now determined automatically from the receipt.");
 			}
 
 			// Set the Engine directory
-			DirectoryReference EngineDir = Parameters.EngineDir ?? Unreal.EngineDirectory;
+			DirectoryReference engineDir = _parameters.EngineDir ?? Unreal.EngineDirectory;
 
 			// Resolve the input list
-			IEnumerable<FileReference> TargetFiles = ResolveFilespec(Unreal.RootDirectory, Parameters.Files, TagNameToFileSet);
+			IEnumerable<FileReference> targetFiles = ResolveFilespec(Unreal.RootDirectory, _parameters.Files, tagNameToFileSet);
 
 			// Filter the files
-			HashSet<FileReference> Files = await ExecuteAsync(EngineDir, TargetFiles, Parameters.BuildProducts, BuildProductType, Parameters.RuntimeDependencies, StagedFileType);
+			HashSet<FileReference> files = await Execute(engineDir, targetFiles, _parameters.BuildProducts, _buildProductType, _parameters.RuntimeDependencies, _stagedFileType);
 
 			// Apply the tag to all the matching files
-			FindOrAddTagSet(TagNameToFileSet, Parameters.With).UnionWith(Files);
+			FindOrAddTagSet(tagNameToFileSet, _parameters.With).UnionWith(files);
 		}
 
-		public static Task<HashSet<FileReference>> ExecuteAsync(DirectoryReference EngineDir, IEnumerable<FileReference> TargetFiles, bool BuildProducts, BuildProductType? BuildProductType, bool RuntimeDependencies, StagedFileType? StagedFileType = null)
+		public static Task<HashSet<FileReference>> Execute(DirectoryReference engineDir, IEnumerable<FileReference> targetFiles, bool buildProducts, BuildProductType? buildProductType, bool runtimeDependencies, StagedFileType? stagedFileType = null)
 		{
-			HashSet<FileReference> Files = new HashSet<FileReference>();
+			HashSet<FileReference> files = new HashSet<FileReference>();
 
-			foreach (FileReference TargetFile in TargetFiles)
+			foreach (FileReference targetFile in targetFiles)
 			{
 				// check all files are .target files
-				if (TargetFile.GetExtension() != ".target")
+				if (targetFile.GetExtension() != ".target")
 				{
-					throw new AutomationException("Invalid file passed to TagReceipt task ({0})", TargetFile.FullName);
+					throw new AutomationException("Invalid file passed to TagReceipt task ({0})", targetFile.FullName);
 				}
 
 				// Read the receipt
-				TargetReceipt Receipt;
-				if (!TargetReceipt.TryRead(TargetFile, EngineDir, out Receipt))
+				TargetReceipt receipt;
+				if (!TargetReceipt.TryRead(targetFile, engineDir, out receipt))
 				{
-					Logger.LogWarning("Unable to load file using TagReceipt task ({Arg0})", TargetFile.FullName);
+					Logger.LogWarning("Unable to load file using TagReceipt task ({Arg0})", targetFile.FullName);
 					continue;
 				}
 
-				if (BuildProducts)
+				if (buildProducts)
 				{
-					foreach (BuildProduct BuildProduct in Receipt.BuildProducts)
+					foreach (BuildProduct buildProduct in receipt.BuildProducts)
 					{
-						if(BuildProductType.HasValue && BuildProduct.Type != BuildProductType.Value)
+						if (buildProductType.HasValue && buildProduct.Type != buildProductType.Value)
 						{
 							continue;
 						}
-						if(StagedFileType.HasValue && TargetReceipt.GetStageTypeFromBuildProductType(BuildProduct) != StagedFileType.Value)
+						if (stagedFileType.HasValue && TargetReceipt.GetStageTypeFromBuildProductType(buildProduct) != stagedFileType.Value)
 						{
 							continue;
 						}
-						Files.Add(BuildProduct.Path);
+						files.Add(buildProduct.Path);
 					}
 				}
 
-				if (RuntimeDependencies)
+				if (runtimeDependencies)
 				{
-					foreach (RuntimeDependency RuntimeDependency in Receipt.RuntimeDependencies)
+					foreach (RuntimeDependency runtimeDependency in receipt.RuntimeDependencies)
 					{
 						// Skip anything that doesn't match the files we want
-						if(BuildProductType.HasValue)
+						if (buildProductType.HasValue)
 						{
 							continue;
 						}
-						if(StagedFileType.HasValue && RuntimeDependency.Type != StagedFileType.Value)
+						if (stagedFileType.HasValue && runtimeDependency.Type != stagedFileType.Value)
 						{
 							continue;
 						}
 
 						// Check which files exist, and warn about any that don't. Ignore debug files, as they are frequently excluded for size (eg. UE on GitHub). This matches logic during staging.
-						FileReference DependencyPath = RuntimeDependency.Path;
-						if (FileReference.Exists(DependencyPath))
+						FileReference dependencyPath = runtimeDependency.Path;
+						if (FileReference.Exists(dependencyPath))
 						{
-							Files.Add(DependencyPath);
+							files.Add(dependencyPath);
 						}
-						else if(RuntimeDependency.Type != UnrealBuildTool.StagedFileType.DebugNonUFS)
+						else if (runtimeDependency.Type != UnrealBuildTool.StagedFileType.DebugNonUFS)
 						{
-							Logger.LogWarning("File listed as RuntimeDependency in {Arg0} does not exist ({Arg1})", TargetFile.FullName, DependencyPath.FullName);
+							Logger.LogWarning("File listed as RuntimeDependency in {Arg0} does not exist ({Arg1})", targetFile.FullName, dependencyPath.FullName);
 						}
 					}
 				}
 			}
 
-			return Task.FromResult(Files);
+			return Task.FromResult(files);
 		}
 
 		/// <summary>
 		/// Output this task out to an XML writer.
 		/// </summary>
-		public override void Write(XmlWriter Writer)
+		public override void Write(XmlWriter writer)
 		{
-			Write(Writer, Parameters);
+			Write(writer, _parameters);
 		}
 
 		/// <summary>
@@ -218,7 +212,7 @@ namespace AutomationTool.Tasks
 		/// <returns>The tag names which are required by this task</returns>
 		public override IEnumerable<string> FindConsumedTagNames()
 		{
-			return FindTagNamesFromFilespec(Parameters.Files);
+			return FindTagNamesFromFilespec(_parameters.Files);
 		}
 
 		/// <summary>
@@ -227,7 +221,7 @@ namespace AutomationTool.Tasks
 		/// <returns>The tag names which are produced/modified by this task</returns>
 		public override IEnumerable<string> FindProducedTagNames()
 		{
-			return FindTagNamesFromList(Parameters.With);
+			return FindTagNamesFromList(_parameters.With);
 		}
 	}
 
@@ -239,10 +233,10 @@ namespace AutomationTool.Tasks
 		/// <summary>
 		/// Task that tags build products and/or runtime dependencies by reading from *.target files.
 		/// </summary>
-		public static async Task<FileSet> TagReceiptsAsync(this FileSet Files, DirectoryReference EngineDir = null, bool BuildProducts = false, BuildProductType? BuildProductType = null, bool RuntimeDependencies = false, StagedFileType? StagedFileType = null)
+		public static async Task<FileSet> TagReceiptsAsync(this FileSet files, DirectoryReference engineDir = null, bool buildProducts = false, BuildProductType? buildProductType = null, bool runtimeDependencies = false, StagedFileType? stagedFileType = null)
 		{
-			HashSet<FileReference> Result = await TagReceiptTask.ExecuteAsync(EngineDir ?? Unreal.EngineDirectory, Files, BuildProducts, BuildProductType, RuntimeDependencies, StagedFileType);
-			return FileSet.FromFiles(Unreal.RootDirectory, Result);
+			HashSet<FileReference> result = await TagReceiptTask.Execute(engineDir ?? Unreal.EngineDirectory, files, buildProducts, buildProductType, runtimeDependencies, stagedFileType);
+			return FileSet.FromFiles(Unreal.RootDirectory, result);
 		}
 	}
 }
diff --git a/Engine/Source/Programs/AutomationTool/BuildGraph/Tasks/TagTask.cs b/Engine/Source/Programs/AutomationTool/BuildGraph/Tasks/TagTask.cs
--- a/Engine/Source/Programs/AutomationTool/BuildGraph/Tasks/TagTask.cs
+++ b/Engine/Source/Programs/AutomationTool/BuildGraph/Tasks/TagTask.cs
@@ -1,16 +1,11 @@
 // Copyright Epic Games, Inc. All Rights Reserved.
 
-using EpicGames.BuildGraph;
 using System;
 using System.Collections.Generic;
-using System.IO;
-using System.Linq;
-using System.Text;
 using System.Threading.Tasks;
 using System.Xml;
 using EpicGames.Core;
 using UnrealBuildBase;
-using UnrealBuildTool;
 
 namespace AutomationTool.Tasks
 {
@@ -63,134 +58,131 @@ namespace AutomationTool.Tasks
 	[TaskElement("Tag", typeof(TagTaskParameters))]
 	class TagTask : BgTaskImpl
 	{
-		/// <summary>
-		/// Parameters to this task
-		/// </summary>
-		TagTaskParameters Parameters;
+		readonly TagTaskParameters _parameters;
 
 		/// <summary>
 		/// Constructor
 		/// </summary>
-		/// <param name="InParameters">Parameters to select which files to match</param>
-		public TagTask(TagTaskParameters InParameters)
+		/// <param name="parameters">Parameters to select which files to match</param>
+		public TagTask(TagTaskParameters parameters)
 		{
-			Parameters = InParameters;
+			_parameters = parameters;
 		}
 
 		/// <summary>
-		/// Execute the task.
+		/// ExecuteAsync the task.
 		/// </summary>
-		/// <param name="Job">Information about the current job</param>
-		/// <param name="BuildProducts">Set of build products produced by this node.</param>
-		/// <param name="TagNameToFileSet">Mapping from tag names to the set of files they include</param>
-		public override async Task ExecuteAsync(JobContext Job, HashSet<FileReference> BuildProducts, Dictionary<string, HashSet<FileReference>> TagNameToFileSet)
+		/// <param name="job">Information about the current job</param>
+		/// <param name="buildProducts">Set of build products produced by this node.</param>
+		/// <param name="tagNameToFileSet">Mapping from tag names to the set of files they include</param>
+		public override async Task ExecuteAsync(JobContext job, HashSet<FileReference> buildProducts, Dictionary<string, HashSet<FileReference>> tagNameToFileSet)
 		{
 			// Get the base directory
-			DirectoryReference BaseDir = Parameters.BaseDir ?? Unreal.RootDirectory;
+			DirectoryReference baseDir = _parameters.BaseDir ?? Unreal.RootDirectory;
 
 			// Parse all the exclude rules
-			List<string> ExcludeRules = ParseRules(BaseDir, Parameters.Except ?? "", TagNameToFileSet);
+			List<string> excludeRules = ParseRules(baseDir, _parameters.Except ?? "", tagNameToFileSet);
 
 			// Resolve the input list
-			HashSet<FileReference> Files = new HashSet<FileReference>();
-			if(!String.IsNullOrEmpty(Parameters.Files))
+			HashSet<FileReference> files = new HashSet<FileReference>();
+			if (!String.IsNullOrEmpty(_parameters.Files))
 			{
-				Files.UnionWith(ResolveFilespecWithExcludePatterns(BaseDir, Parameters.Files, ExcludeRules, TagNameToFileSet));
+				files.UnionWith(ResolveFilespecWithExcludePatterns(baseDir, _parameters.Files, excludeRules, tagNameToFileSet));
 			}
 
 			// Resolve the input file lists
-			if(!String.IsNullOrEmpty(Parameters.FileLists))
+			if (!String.IsNullOrEmpty(_parameters.FileLists))
 			{
-				HashSet<FileReference> FileLists = ResolveFilespec(BaseDir, Parameters.FileLists, TagNameToFileSet);
-				foreach(FileReference FileList in FileLists)
+				HashSet<FileReference> fileLists = ResolveFilespec(baseDir, _parameters.FileLists, tagNameToFileSet);
+				foreach (FileReference fileList in fileLists)
 				{
-					if(!FileReference.Exists(FileList))
+					if (!FileReference.Exists(fileList))
 					{
-						throw new AutomationException("Specified file list '{0}' does not exist", FileList);
+						throw new AutomationException("Specified file list '{0}' does not exist", fileList);
 					}
 
-					string[] Lines = await FileReference.ReadAllLinesAsync(FileList);
-					foreach(string Line in Lines)
+					string[] lines = await FileReference.ReadAllLinesAsync(fileList);
+					foreach (string line in lines)
 					{
-						string TrimLine = Line.Trim();
-						if(TrimLine.Length > 0)
+						string trimLine = line.Trim();
+						if (trimLine.Length > 0)
 						{
-							Files.Add(FileReference.Combine(BaseDir, TrimLine));
+							files.Add(FileReference.Combine(baseDir, trimLine));
 						}
 					}
 				}
 			}
 
 			// Limit to matches against the 'Filter' parameter, if set
-			if(Parameters.Filter != null)
+			if (_parameters.Filter != null)
 			{
-				FileFilter Filter = new FileFilter();
-				Filter.AddRules(ParseRules(BaseDir, Parameters.Filter, TagNameToFileSet));
-				Files.RemoveWhere(x => !Filter.Matches(x.FullName));
+				FileFilter filter = new FileFilter();
+				filter.AddRules(ParseRules(baseDir, _parameters.Filter, tagNameToFileSet));
+				files.RemoveWhere(x => !filter.Matches(x.FullName));
 			}
 
 			// Apply the tag to all the matching files
-			foreach(string TagName in FindTagNamesFromList(Parameters.With))
+			foreach (string tagName in FindTagNamesFromList(_parameters.With))
 			{
-				FindOrAddTagSet(TagNameToFileSet, TagName).UnionWith(Files);
+				FindOrAddTagSet(tagNameToFileSet, tagName).UnionWith(files);
 			}
 		}
 
 		/// <summary>
 		/// Add rules matching a given set of patterns to a file filter. Patterns are added as absolute paths from the root.
 		/// </summary>
-		/// <param name="BaseDir">The base directory for relative paths.</param>
-		/// <param name="DelimitedPatterns">List of patterns to add, separated by semicolons.</param>
-		/// <param name="TagNameToFileSet">Mapping of tag name to a set of files.</param>
+		/// <param name="baseDir">The base directory for relative paths.</param>
+		/// <param name="delimitedPatterns">List of patterns to add, separated by semicolons.</param>
+		/// <param name="tagNameToFileSet">Mapping of tag name to a set of files.</param>
 		/// <returns>List of rules, suitable for adding to a FileFilter object</returns>
-		static List<string> ParseRules(DirectoryReference BaseDir, string DelimitedPatterns, Dictionary<string, HashSet<FileReference>> TagNameToFileSet)
+		static List<string> ParseRules(DirectoryReference baseDir, string delimitedPatterns, Dictionary<string, HashSet<FileReference>> tagNameToFileSet)
 		{
 			// Split up the list of patterns
-			List<string> Patterns = SplitDelimitedList(DelimitedPatterns);
+			List<string> patterns = SplitDelimitedList(delimitedPatterns);
 
 			// Parse them into a list of rules
-			List<string> Rules = new List<string>();
-			foreach(string Pattern in Patterns)
+			List<string> rules = new List<string>();
+			foreach (string pattern in patterns)
 			{
-				if(Pattern.StartsWith("#", StringComparison.Ordinal))
+				if (pattern.StartsWith("#", StringComparison.Ordinal))
 				{
 					// Add the files in a specific set to the filter
-					HashSet<FileReference> Files = FindOrAddTagSet(TagNameToFileSet, Pattern);
-					foreach(FileReference File in Files)
+					HashSet<FileReference> files = FindOrAddTagSet(tagNameToFileSet, pattern);
+					foreach (FileReference file in files)
 					{
-						Rules.Add(File.FullName);
+						rules.Add(file.FullName);
 					}
 				}
 				else
 				{
 					// Parse a wildcard filter
-					if(Pattern.StartsWith("...", StringComparison.Ordinal))
+					if (pattern.StartsWith("...", StringComparison.Ordinal))
 					{
-						Rules.Add(Pattern);
+						rules.Add(pattern);
 					}
-					else if(!Pattern.Contains('/', StringComparison.Ordinal))
+					else if (!pattern.Contains('/', StringComparison.Ordinal))
 					{
-						Rules.Add(".../" + Pattern);
+						rules.Add(".../" + pattern);
 					}
-					else if(!Pattern.StartsWith("/", StringComparison.Ordinal))
+					else if (!pattern.StartsWith("/", StringComparison.Ordinal))
 					{
-						Rules.Add(BaseDir.FullName + "/" + Pattern);
+						rules.Add(baseDir.FullName + "/" + pattern);
 					}
 					else
 					{
-						Rules.Add(BaseDir.FullName + Pattern);
+						rules.Add(baseDir.FullName + pattern);
 					}
 				}
 			}
-			return Rules;
+			return rules;
 		}
 
 		/// <summary>
 		/// Output this task out to an XML writer.
 		/// </summary>
-		public override void Write(XmlWriter Writer)
+		public override void Write(XmlWriter writer)
 		{
-			Write(Writer, Parameters);
+			Write(writer, _parameters);
 		}
 
 		/// <summary>
@@ -199,24 +191,24 @@ namespace AutomationTool.Tasks
 		/// <returns>The tag names which are read by this task</returns>
 		public override IEnumerable<string> FindConsumedTagNames()
 		{
-			foreach(string TagName in FindTagNamesFromFilespec(Parameters.Files))
+			foreach (string tagName in FindTagNamesFromFilespec(_parameters.Files))
 			{
-				yield return TagName;
+				yield return tagName;
 			}
 
-			if(!String.IsNullOrEmpty(Parameters.Filter))
+			if (!String.IsNullOrEmpty(_parameters.Filter))
 			{
-				foreach(string TagName in FindTagNamesFromFilespec(Parameters.Filter))
+				foreach (string tagName in FindTagNamesFromFilespec(_parameters.Filter))
 				{
-					yield return TagName;
+					yield return tagName;
 				}
 			}
 
-			if(!String.IsNullOrEmpty(Parameters.Except))
+			if (!String.IsNullOrEmpty(_parameters.Except))
 			{
-				foreach(string TagName in FindTagNamesFromFilespec(Parameters.Except))
+				foreach (string tagName in FindTagNamesFromFilespec(_parameters.Except))
 				{
-					yield return TagName;
+					yield return tagName;
 				}
 			}
 		}
@@ -227,7 +219,7 @@ namespace AutomationTool.Tasks
 		/// <returns>The tag names which are modified by this task</returns>
 		public override IEnumerable<string> FindProducedTagNames()
 		{
-			return FindTagNamesFromList(Parameters.With);
+			return FindTagNamesFromList(_parameters.With);
 		}
 	}
 }
diff --git a/Engine/Source/Programs/AutomationTool/BuildGraph/Tasks/UnzipTask.cs b/Engine/Source/Programs/AutomationTool/BuildGraph/Tasks/UnzipTask.cs
--- a/Engine/Source/Programs/AutomationTool/BuildGraph/Tasks/UnzipTask.cs
+++ b/Engine/Source/Programs/AutomationTool/BuildGraph/Tasks/UnzipTask.cs
@@ -1,15 +1,10 @@
 ﻿// Copyright Epic Games, Inc. All Rights Reserved.
 
-using AutomationTool;
-using EpicGames.BuildGraph;
-using System;
 using System.Collections.Generic;
 using System.Linq;
-using System.Text;
 using System.Threading.Tasks;
 using System.Xml;
 using EpicGames.Core;
-using UnrealBuildTool;
 using UnrealBuildBase;
 
 namespace AutomationTool.Tasks
@@ -56,64 +51,61 @@ namespace AutomationTool.Tasks
 	[TaskElement("Unzip", typeof(UnzipTaskParameters))]
 	public class UnzipTask : BgTaskImpl
 	{
-		/// <summary>
-		/// Parameters for this task
-		/// </summary>
-		UnzipTaskParameters Parameters;
+		readonly UnzipTaskParameters _parameters;
 
 		/// <summary>
 		/// Constructor
 		/// </summary>
-		/// <param name="InParameters">Parameters for this task</param>
-		public UnzipTask(UnzipTaskParameters InParameters)
+		/// <param name="parameters">Parameters for this task</param>
+		public UnzipTask(UnzipTaskParameters parameters)
 		{
-			Parameters = InParameters;
+			_parameters = parameters;
 		}
 
 		/// <summary>
-		/// Execute the task.
+		/// ExecuteAsync the task.
 		/// </summary>
-		/// <param name="Job">Information about the current job</param>
-		/// <param name="BuildProducts">Set of build products produced by this node.</param>
-		/// <param name="TagNameToFileSet">Mapping from tag names \to the set of files they include</param>
-		public override Task ExecuteAsync(JobContext Job, HashSet<FileReference> BuildProducts, Dictionary<string, HashSet<FileReference>> TagNameToFileSet)
+		/// <param name="job">Information about the current job</param>
+		/// <param name="buildProducts">Set of build products produced by this node.</param>
+		/// <param name="tagNameToFileSet">Mapping from tag names \to the set of files they include</param>
+		public override Task ExecuteAsync(JobContext job, HashSet<FileReference> buildProducts, Dictionary<string, HashSet<FileReference>> tagNameToFileSet)
 		{
-			DirectoryReference ToDir = Parameters.ToDir;
+			DirectoryReference toDir = _parameters.ToDir;
 
 			// Find all the zip files
-			IEnumerable<FileReference> ZipFiles = ResolveFilespec(Unreal.RootDirectory, Parameters.ZipFile, TagNameToFileSet);
+			IEnumerable<FileReference> zipFiles = ResolveFilespec(Unreal.RootDirectory, _parameters.ZipFile, tagNameToFileSet);
 
 			// Extract the files
-			HashSet<FileReference> OutputFiles = new HashSet<FileReference>();
-			foreach(FileReference ZipFile in ZipFiles)
+			HashSet<FileReference> outputFiles = new HashSet<FileReference>();
+			foreach (FileReference zipFile in zipFiles)
 			{
-				if (Parameters.UseLegacyUnzip)
+				if (_parameters.UseLegacyUnzip)
 				{
-					OutputFiles.UnionWith(CommandUtils.LegacyUnzipFiles(ZipFile.FullName, ToDir.FullName, Parameters.OverwriteFiles).Select(x => new FileReference(x)));
+					outputFiles.UnionWith(CommandUtils.LegacyUnzipFiles(zipFile.FullName, toDir.FullName, _parameters.OverwriteFiles).Select(x => new FileReference(x)));
 				}
 				else
 				{
-					OutputFiles.UnionWith(CommandUtils.UnzipFiles(ZipFile, ToDir, Parameters.OverwriteFiles));
+					outputFiles.UnionWith(CommandUtils.UnzipFiles(zipFile, toDir, _parameters.OverwriteFiles));
 				}
 			}
 
 			// Apply the optional tag to the produced archive
-			foreach(string TagName in FindTagNamesFromList(Parameters.Tag))
+			foreach (string tagName in FindTagNamesFromList(_parameters.Tag))
 			{
-				FindOrAddTagSet(TagNameToFileSet, TagName).UnionWith(OutputFiles);
+				FindOrAddTagSet(tagNameToFileSet, tagName).UnionWith(outputFiles);
 			}
 
 			// Add the archive to the set of build products
-			BuildProducts.UnionWith(OutputFiles);
+			buildProducts.UnionWith(outputFiles);
 			return Task.CompletedTask;
 		}
 
 		/// <summary>
 		/// Output this task out to an XML writer.
 		/// </summary>
-		public override void Write(XmlWriter Writer)
+		public override void Write(XmlWriter writer)
 		{
-			Write(Writer, Parameters);
+			Write(writer, _parameters);
 		}
 
 		/// <summary>
@@ -122,7 +114,7 @@ namespace AutomationTool.Tasks
 		/// <returns>The tag names which are read by this task</returns>
 		public override IEnumerable<string> FindConsumedTagNames()
 		{
-			return FindTagNamesFromFilespec(Parameters.ZipFile);
+			return FindTagNamesFromFilespec(_parameters.ZipFile);
 		}
 
 		/// <summary>
@@ -131,7 +123,7 @@ namespace AutomationTool.Tasks
 		/// <returns>The tag names which are modified by this task</returns>
 		public override IEnumerable<string> FindProducedTagNames()
 		{
-			return FindTagNamesFromList(Parameters.Tag);
+			return FindTagNamesFromList(_parameters.Tag);
 		}
 	}
 }
diff --git a/Engine/Source/Programs/AutomationTool/BuildGraph/Tasks/UploadTask.cs b/Engine/Source/Programs/AutomationTool/BuildGraph/Tasks/UploadTask.cs
--- a/Engine/Source/Programs/AutomationTool/BuildGraph/Tasks/UploadTask.cs
+++ b/Engine/Source/Programs/AutomationTool/BuildGraph/Tasks/UploadTask.cs
@@ -1,14 +1,8 @@
 // Copyright Epic Games, Inc. All Rights Reserved.
 
-using AutomationTool;
-using EpicGames.BuildGraph;
 using System;
 using System.Collections.Generic;
-using System.IO;
 using System.Linq;
-using System.Net.Http;
-using System.Security.Cryptography;
-using System.Text;
 using System.Threading.Tasks;
 using System.Xml;
 using EpicGames.Core;
@@ -33,7 +27,7 @@ namespace AutomationTool.Tasks
 		/// </summary>
 		[TaskParameter(Optional = true, ValidationType = TaskParameterValidationType.FileSpec)]
 		public string Files { get; set; }
-		
+
 		/// <summary>
 		/// The jupiter namespace used to upload the build. Used to control who has access to the build.
 		/// </summary>
@@ -95,73 +89,70 @@ namespace AutomationTool.Tasks
 	[TaskElement("Upload", typeof(UploadTaskParameters))]
 	public class UploadTask : BgTaskImpl
 	{
-		/// <summary>
-		/// Parameters for this task
-		/// </summary>
-		UploadTaskParameters Parameters;
+		readonly UploadTaskParameters _parameters;
 
 		/// <summary>
 		/// Constructor
 		/// </summary>
-		/// <param name="InParameters">Parameters for this task</param>
-		public UploadTask(UploadTaskParameters InParameters)
+		/// <param name="parameters">Parameters for this task</param>
+		public UploadTask(UploadTaskParameters parameters)
 		{
-			Parameters = InParameters;
+			_parameters = parameters;
 		}
 
 		/// <summary>
-		/// Execute the task.
+		/// ExecuteAsync the task.
 		/// </summary>
-		/// <param name="Job">Information about the current job</param>
-		/// <param name="BuildProducts">Set of build products produced by this node.</param>
-		/// <param name="TagNameToFileSet">Mapping from tag names to the set of files they include</param>
-		public override async Task ExecuteAsync(JobContext Job, HashSet<FileReference> BuildProducts, Dictionary<string, HashSet<FileReference>> TagNameToFileSet)
+		/// <param name="job">Information about the current job</param>
+		/// <param name="buildProducts">Set of build products produced by this node.</param>
+		/// <param name="tagNameToFileSet">Mapping from tag names to the set of files they include</param>
+		public override async Task ExecuteAsync(JobContext job, HashSet<FileReference> buildProducts, Dictionary<string, HashSet<FileReference>> tagNameToFileSet)
 		{
 			// Find all the input files
-			List<FileReference> Files;
-			if(Parameters.Files == null)
+			List<FileReference> files;
+			if (_parameters.Files == null)
 			{
-				Files = DirectoryReference.EnumerateFiles(Parameters.FromDir, "*", System.IO.SearchOption.AllDirectories).ToList();
+				files = DirectoryReference.EnumerateFiles(_parameters.FromDir, "*", System.IO.SearchOption.AllDirectories).ToList();
 			}
 			else
 			{
-				Files = ResolveFilespec(Parameters.FromDir, Parameters.Files, TagNameToFileSet).ToList();
+				files = ResolveFilespec(_parameters.FromDir, _parameters.Files, tagNameToFileSet).ToList();
 			}
 
 			// Create the jupiter tree
-			Logger.LogInformation("Uploading {NumFiles} files to {Url}...", Files.Count, Parameters.JupiterUrl);
+			Logger.LogInformation("Uploading {NumFiles} files to {Url}...", files.Count, _parameters.JupiterUrl);
 
-			JupiterFileTree FileTree = new JupiterFileTree(Parameters.FromDir, Parameters.LimitMemoryUsage);
-			foreach (FileReference File in Files)
+			JupiterFileTree fileTree = new JupiterFileTree(_parameters.FromDir, _parameters.LimitMemoryUsage);
+			foreach (FileReference file in files)
 			{
-				FileTree.AddFile(File);
+				fileTree.AddFile(file);
 			}
 
-			Dictionary<string, object> Metadata = new Dictionary<string, object>
+			Dictionary<string, object> metadata = new Dictionary<string, object>
 			{
-				{"ArchiveType", Parameters.ArchiveType},
-				{"Project", Parameters.ProjectName},
-				{"Branch", Parameters.Branch},
-				{"Changelist", Parameters.Changelist},
+				{"ArchiveType", _parameters.ArchiveType},
+				{"Project", _parameters.ProjectName},
+				{"Branch", _parameters.Branch},
+				{"Changelist", _parameters.Changelist},
 			};
 
-			if (Parameters.AdditionalMetadata != null)
+			if (_parameters.AdditionalMetadata != null)
 			{
-				string[] KV = Parameters.AdditionalMetadata.Split(';');
-				foreach (string Option in KV)
+				string[] kv = _parameters.AdditionalMetadata.Split(';');
+				foreach (string option in kv)
 				{
-					int SeparatorIndex = Option.IndexOf('=', StringComparison.Ordinal);
-					if (SeparatorIndex == -1)
+					int separatorIndex = option.IndexOf('=', StringComparison.Ordinal);
+					if (separatorIndex == -1)
 						continue;
 
-					string Key = Option.Substring(0, SeparatorIndex);
-					string Value = Option.Substring(SeparatorIndex + 1);
+					string key = option.Substring(0, separatorIndex);
+					string value = option.Substring(separatorIndex + 1);
 
-					Metadata[Key] = Value;
+					metadata[key] = value;
 				}
 			}
 			// Upload the tree to Jupiter
-			Dictionary<FileReference, List<string>> Mapping = await FileTree.UploadToJupiter(Parameters.JupiterUrl, Parameters.JupiterNamespace, Parameters.JupiterKey, Metadata);
+			_ = await fileTree.UploadToJupiter(_parameters.JupiterUrl, _parameters.JupiterNamespace, _parameters.JupiterKey, metadata);
 
 			// Debug output of which files mapped to which blobs, can be useful to determine which files are constantly being uploaded
 			// Json.Save(FileReference.Combine(AutomationTool.Unreal.RootDirectory, "JupiterUpload.json"), Mapping);
@@ -170,9 +161,9 @@ namespace AutomationTool.Tasks
 		/// <summary>
 		/// Output this task out to an XML writer.
 		/// </summary>
-		public override void Write(XmlWriter Writer)
+		public override void Write(XmlWriter writer)
 		{
-			Write(Writer, Parameters);
+			Write(writer, _parameters);
 		}
 
 		/// <summary>
@@ -181,7 +172,7 @@ namespace AutomationTool.Tasks
 		/// <returns>The tag names which are read by this task</returns>
 		public override IEnumerable<string> FindConsumedTagNames()
 		{
-			return FindTagNamesFromFilespec(Parameters.Files);
+			return FindTagNamesFromFilespec(_parameters.Files);
 		}
 
 		/// <summary>
diff --git a/Engine/Source/Programs/AutomationTool/BuildGraph/Tasks/WaitTask.cs b/Engine/Source/Programs/AutomationTool/BuildGraph/Tasks/WaitTask.cs
--- a/Engine/Source/Programs/AutomationTool/BuildGraph/Tasks/WaitTask.cs
+++ b/Engine/Source/Programs/AutomationTool/BuildGraph/Tasks/WaitTask.cs
@@ -2,11 +2,8 @@
 
 using System;
 using System.Collections.Generic;
-using System.IO;
-using System.Threading;
 using System.Threading.Tasks;
 using System.Xml;
-using AutomationTool;
 using EpicGames.Core;
 
 namespace AutomationTool.Tasks
@@ -29,37 +26,34 @@ namespace AutomationTool.Tasks
 	[TaskElement("Wait", typeof(WaitTaskParameters))]
 	public class WaitTask : BgTaskImpl
 	{
-		/// <summary>
-		/// Parameters for this task
-		/// </summary>
-		WaitTaskParameters Parameters;
+		readonly WaitTaskParameters _parameters;
 
 		/// <summary>
 		/// Construct a wait task
 		/// </summary>
-		/// <param name="InParameters">Parameters for the task</param>
-		public WaitTask(WaitTaskParameters InParameters)
+		/// <param name="parameters">Parameters for the task</param>
+		public WaitTask(WaitTaskParameters parameters)
 		{
-			Parameters = InParameters;
+			_parameters = parameters;
 		}
 
 		/// <summary>
-		/// Execute the task.
+		/// ExecuteAsync the task.
 		/// </summary>
-		/// <param name="Job">Information about the current job</param>
-		/// <param name="BuildProducts">Set of build products produced by this node.</param>
-		/// <param name="TagNameToFileSet">Mapping from tag names to the set of files they include</param>
-		public override async Task ExecuteAsync(JobContext Job, HashSet<FileReference> BuildProducts, Dictionary<string, HashSet<FileReference>> TagNameToFileSet)
+		/// <param name="job">Information about the current job</param>
+		/// <param name="buildProducts">Set of build products produced by this node.</param>
+		/// <param name="tagNameToFileSet">Mapping from tag names to the set of files they include</param>
+		public override async Task ExecuteAsync(JobContext job, HashSet<FileReference> buildProducts, Dictionary<string, HashSet<FileReference>> tagNameToFileSet)
 		{
-			await Task.Delay(TimeSpan.FromSeconds(Parameters.Seconds));
+			await Task.Delay(TimeSpan.FromSeconds(_parameters.Seconds));
 		}
 
 		/// <summary>
 		/// Output this task out to an XML writer.
 		/// </summary>
-		public override void Write(XmlWriter Writer)
+		public override void Write(XmlWriter writer)
 		{
-			Write(Writer, Parameters);
+			Write(writer, _parameters);
 		}
 
 		/// <summary>
diff --git a/Engine/Source/Programs/AutomationTool/BuildGraph/Tasks/WriteJsonValueTask.cs b/Engine/Source/Programs/AutomationTool/BuildGraph/Tasks/WriteJsonValueTask.cs
--- a/Engine/Source/Programs/AutomationTool/BuildGraph/Tasks/WriteJsonValueTask.cs
+++ b/Engine/Source/Programs/AutomationTool/BuildGraph/Tasks/WriteJsonValueTask.cs
@@ -1,16 +1,12 @@
 // Copyright Epic Games, Inc. All Rights Reserved.
 
-using EpicGames.Core;
-using IdentityModel.Client;
 using System;
 using System.Collections.Generic;
-using System.Linq;
-using System.Text;
 using System.Text.Json;
 using System.Text.Json.Nodes;
-using System.Text.RegularExpressions;
 using System.Threading.Tasks;
 using System.Xml;
+using EpicGames.Core;
 using UnrealBuildBase;
 
 #nullable enable
@@ -51,138 +47,138 @@ namespace AutomationTool.Tasks
 	[TaskElement("WriteJsonValue", typeof(WriteJsonValueTaskParameters))]
 	public class WriteJsonValueTask : BgTaskImpl
 	{
-		WriteJsonValueTaskParameters Parameters;
+		readonly WriteJsonValueTaskParameters _parameters;
 
 		/// <summary>
 		/// Create a new ModifyJsonValue.
 		/// </summary>
-		/// <param name="InParameters">Parameters for this task.</param>
-		public WriteJsonValueTask(WriteJsonValueTaskParameters InParameters)
+		/// <param name="parameters">Parameters for this task.</param>
+		public WriteJsonValueTask(WriteJsonValueTaskParameters parameters)
 		{
-			Parameters = InParameters;
+			_parameters = parameters;
 		}
 
 		/// <summary>
 		/// Placeholder comment
 		/// </summary>
-		public override async Task ExecuteAsync(JobContext Job, HashSet<FileReference> BuildProducts, Dictionary<string, HashSet<FileReference>> TagNameToFileSet)
+		public override async Task ExecuteAsync(JobContext job, HashSet<FileReference> buildProducts, Dictionary<string, HashSet<FileReference>> tagNameToFileSet)
 		{
-			HashSet<FileReference> Files = ResolveFilespec(Unreal.RootDirectory, Parameters.File, TagNameToFileSet);
+			HashSet<FileReference> files = ResolveFilespec(Unreal.RootDirectory, _parameters.File, tagNameToFileSet);
 
-			JsonNode? ValueNode;
+			JsonNode? valueNode;
 			try
 			{
-				ValueNode = String.IsNullOrEmpty(Parameters.Value) ? null : JsonNode.Parse(Parameters.Value);
+				valueNode = String.IsNullOrEmpty(_parameters.Value) ? null : JsonNode.Parse(_parameters.Value);
 			}
 			catch (Exception ex)
 			{
-				throw new AutomationException(ex, $"Unable to parse '{Parameters.Value}': {ex.Message}");
+				throw new AutomationException(ex, $"Unable to parse '{_parameters.Value}': {ex.Message}");
 			}
 
-			foreach (FileReference JsonFile in Files)
+			foreach (FileReference jsonFile in files)
 			{
-				string JsonText = FileReference.Exists(JsonFile) ? await FileReference.ReadAllTextAsync(JsonFile) : "{}";
+				string jsonText = FileReference.Exists(jsonFile) ? await FileReference.ReadAllTextAsync(jsonFile) : "{}";
 
-				if (!Parameters.Key.StartsWith("$", StringComparison.Ordinal))
+				if (!_parameters.Key.StartsWith("$", StringComparison.Ordinal))
 				{
 					throw new AutomationException("Key must be in JsonPath format (eg. $.Foo.Bar[123])");
 				}
 
-				JsonNode? RootNode;
+				JsonNode? rootNode;
 				try
 				{
-					RootNode = JsonNode.Parse(JsonText, documentOptions: new JsonDocumentOptions { CommentHandling = JsonCommentHandling.Skip });
+					rootNode = JsonNode.Parse(jsonText, documentOptions: new JsonDocumentOptions { CommentHandling = JsonCommentHandling.Skip });
 				}
 				catch (Exception ex)
 				{
-					throw new AutomationException($"Error parsing {JsonFile}: {ex.Message}");
+					throw new AutomationException($"Error parsing {jsonFile}: {ex.Message}");
 				}
-				RootNode = MergeValue(Parameters.Key, 1, RootNode, ValueNode);
+				rootNode = MergeValue(_parameters.Key, 1, rootNode, valueNode);
 
-				string NewJsonText = RootNode?.ToJsonString(new JsonSerializerOptions { WriteIndented = true }) ?? String.Empty;
+				string newJsonText = rootNode?.ToJsonString(new JsonSerializerOptions { WriteIndented = true }) ?? String.Empty;
 
-				DirectoryReference.CreateDirectory(JsonFile.Directory);
-				await FileReference.WriteAllTextAsync(JsonFile, NewJsonText);
+				DirectoryReference.CreateDirectory(jsonFile.Directory);
+				await FileReference.WriteAllTextAsync(jsonFile, newJsonText);
 			}
 		}
 
-		static JsonNode? MergeValue(string Key, int MinIdx, JsonNode? PrevValue, JsonNode? Value)
+		static JsonNode? MergeValue(string key, int minIdx, JsonNode? prevValue, JsonNode? value)
 		{
-			if (MinIdx == Key.Length)
+			if (minIdx == key.Length)
 			{
-				return Value;
+				return value;
 			}
 
 			// Find the length of the next token
-			int MaxIdx = MinIdx + 1;
-			while (MaxIdx < Key.Length && Key[MaxIdx] != '[' && Key[MaxIdx] != '.')
+			int maxIdx = minIdx + 1;
+			while (maxIdx < key.Length && key[maxIdx] != '[' && key[maxIdx] != '.')
 			{
-				MaxIdx++;
+				maxIdx++;
 			}
 
 			// Handle different types of element
-			if (Key[MinIdx] == '.')
+			if (key[minIdx] == '.')
 			{
-				JsonObject? Obj = PrevValue as JsonObject;
-				if (Obj != null)
+				JsonObject? obj = prevValue as JsonObject;
+				if (obj != null)
 				{
-					Obj = Obj.Deserialize<JsonObject>(); // Clone so we can reattach
+					obj = obj.Deserialize<JsonObject>(); // Clone so we can reattach
 				}
-				Obj ??= new JsonObject();
+				obj ??= new JsonObject();
 
-				string PropertyName = Key.Substring(MinIdx + 1, MaxIdx - (MinIdx + 1));
+				string propertyName = key.Substring(minIdx + 1, maxIdx - (minIdx + 1));
 
-				JsonNode? NextNode;
-				Obj.TryGetPropertyValue(PropertyName, out NextNode);
-				Obj[PropertyName] = MergeValue(Key, MaxIdx, NextNode, Value);
+				JsonNode? nextNode;
+				obj.TryGetPropertyValue(propertyName, out nextNode);
+				obj[propertyName] = MergeValue(key, maxIdx, nextNode, value);
 
-				return Obj;
+				return obj;
 			}
-			else if (Key[MinIdx] == '[')
+			else if (key[minIdx] == '[')
 			{
-				if (Key[MaxIdx - 1] != ']')
+				if (key[maxIdx - 1] != ']')
 				{
 					throw new AutomationException("Missing ']' in array subscript in Json path expression '{Key}'");
 				}
 
-				string IndexStr = Key.Substring(MinIdx + 1, (MaxIdx - 1) - (MinIdx + 1)).Trim();
+				string indexStr = key.Substring(minIdx + 1, (maxIdx - 1) - (minIdx + 1)).Trim();
 
-				int Index = int.MaxValue;
-				if (IndexStr.Length > 0)
+				int index = int.MaxValue;
+				if (indexStr.Length > 0)
 				{
-					Index = int.Parse(IndexStr);
+					index = int.Parse(indexStr);
 				}
 
-				JsonArray? Array = PrevValue as JsonArray;
-				if (Array != null)
+				JsonArray? array = prevValue as JsonArray;
+				if (array != null)
 				{
-					Array = Array.Deserialize<JsonArray>();
+					array = array.Deserialize<JsonArray>();
 				}
-				Array ??= new JsonArray();
+				array ??= new JsonArray();
 
-				if (Index < Array.Count)
+				if (index < array.Count)
 				{
-					Array[Index] = MergeValue(Key, MaxIdx, Array[Index], Value);
+					array[index] = MergeValue(key, maxIdx, array[index], value);
 				}
 				else
 				{
-					Array.Add(MergeValue(Key, MaxIdx, null, Value));
+					array.Add(MergeValue(key, maxIdx, null, value));
 				}
 
-				return Array;
+				return array;
 			}
 			else
 			{
-				throw new AutomationException($"Unable to parse JSON path after '{Key}'");
+				throw new AutomationException($"Unable to parse JSON path after '{key}'");
 			}
 		}
 
 		/// <summary>
 		/// Placeholder comment
 		/// </summary>
-		public override void Write(XmlWriter Writer)
+		public override void Write(XmlWriter writer)
 		{
-			Write(Writer, Parameters);
+			Write(writer, _parameters);
 		}
 
 		/// <summary>
@@ -190,9 +186,9 @@ namespace AutomationTool.Tasks
 		/// </summary>
 		public override IEnumerable<string> FindConsumedTagNames()
 		{
-			foreach (string TagName in FindTagNamesFromFilespec(Parameters.File))
+			foreach (string tagName in FindTagNamesFromFilespec(_parameters.File))
 			{
-				yield return TagName;
+				yield return tagName;
 			}
 		}
 
diff --git a/Engine/Source/Programs/AutomationTool/BuildGraph/Tasks/WriteTextFileTask.cs b/Engine/Source/Programs/AutomationTool/BuildGraph/Tasks/WriteTextFileTask.cs
--- a/Engine/Source/Programs/AutomationTool/BuildGraph/Tasks/WriteTextFileTask.cs
+++ b/Engine/Source/Programs/AutomationTool/BuildGraph/Tasks/WriteTextFileTask.cs
@@ -1,17 +1,13 @@
 // Copyright Epic Games, Inc. All Rights Reserved.
 
-using AutomationTool;
-using EpicGames.BuildGraph;
 using System;
 using System.Collections.Generic;
 using System.Linq;
+using System.Threading.Tasks;
 using System.Xml;
 using EpicGames.Core;
-using UnrealBuildBase;
-using System.Threading.Tasks;
 using Microsoft.Extensions.Logging;
-
-using static AutomationTool.CommandUtils;
+using UnrealBuildBase;
 
 namespace AutomationTool.Tasks
 {
@@ -60,75 +56,75 @@ namespace AutomationTool.Tasks
 		/// <summary>
 		/// Parameters for this task.
 		/// </summary>
-		WriteTextFileTaskParameters Parameters;
+		readonly WriteTextFileTaskParameters _parameters;
 
 		/// <summary>
 		/// Constructor.
 		/// </summary>
-		/// <param name="InParameters">Parameters for this task.</param>
-		public WriteTextFileTask(WriteTextFileTaskParameters InParameters)
+		/// <param name="parameters">Parameters for this task.</param>
+		public WriteTextFileTask(WriteTextFileTaskParameters parameters)
 		{
-			Parameters = InParameters;
+			_parameters = parameters;
 		}
 
 		/// <summary>
-		/// Execute the task.
+		/// ExecuteAsync the task.
 		/// </summary>
-		/// <param name="Job">Information about the current job.</param>
-		/// <param name="BuildProducts">Set of build products produced by this node.</param>
-		/// <param name="TagNameToFileSet">Mapping from tag names to the set of files they include.</param>
-		public override async Task ExecuteAsync(JobContext Job, HashSet<FileReference> BuildProducts, Dictionary<string, HashSet<FileReference>> TagNameToFileSet)
+		/// <param name="job">Information about the current job.</param>
+		/// <param name="buildProducts">Set of build products produced by this node.</param>
+		/// <param name="tagNameToFileSet">Mapping from tag names to the set of files they include.</param>
+		public override async Task ExecuteAsync(JobContext job, HashSet<FileReference> buildProducts, Dictionary<string, HashSet<FileReference>> tagNameToFileSet)
 		{
-			string FileText = Parameters.Text;
+			string fileText = _parameters.Text;
 
 			// If any files or tagsets are provided, add them to the text output.
-			if (!String.IsNullOrEmpty(Parameters.Files))
+			if (!String.IsNullOrEmpty(_parameters.Files))
 			{
-				if (!string.IsNullOrWhiteSpace(FileText))
+				if (!string.IsNullOrWhiteSpace(fileText))
 				{
-					FileText += Environment.NewLine;
+					fileText += Environment.NewLine;
 				}
 
-				HashSet<FileReference> Files = ResolveFilespec(Unreal.RootDirectory, Parameters.Files, TagNameToFileSet);
-				if (Files.Any())
+				HashSet<FileReference> files = ResolveFilespec(Unreal.RootDirectory, _parameters.Files, tagNameToFileSet);
+				if (files.Any())
 				{
-					FileText += string.Join(Environment.NewLine, Files.Select(f => f.FullName));
+					fileText += string.Join(Environment.NewLine, files.Select(f => f.FullName));
 				}
 			}
 
 			// Make sure output folder exists.
-			if (!DirectoryReference.Exists(Parameters.File.Directory))
+			if (!DirectoryReference.Exists(_parameters.File.Directory))
 			{
-				DirectoryReference.CreateDirectory(Parameters.File.Directory);
+				DirectoryReference.CreateDirectory(_parameters.File.Directory);
 			}
 
-			if (Parameters.Append)
+			if (_parameters.Append)
 			{
-				Logger.LogInformation("{Text}", string.Format("Appending text to file '{0}': {1}", Parameters.File, FileText));
-				await FileReference.AppendAllTextAsync(Parameters.File, Environment.NewLine + FileText);
+				Logger.LogInformation("{Text}", string.Format("Appending text to file '{0}': {1}", _parameters.File, fileText));
+				await FileReference.AppendAllTextAsync(_parameters.File, Environment.NewLine + fileText);
 			}
 			else
 			{
-				Logger.LogInformation("{Text}", string.Format("Writing text to file '{0}': {1}", Parameters.File, FileText));
-				await FileReference.WriteAllTextAsync(Parameters.File, FileText);
+				Logger.LogInformation("{Text}", string.Format("Writing text to file '{0}': {1}", _parameters.File, fileText));
+				await FileReference.WriteAllTextAsync(_parameters.File, fileText);
 			}
 
 			// Apply the optional tag to the build products
-			foreach (string TagName in FindTagNamesFromList(Parameters.Tag))
+			foreach (string tagName in FindTagNamesFromList(_parameters.Tag))
 			{
-				FindOrAddTagSet(TagNameToFileSet, TagName).Add(Parameters.File);
+				FindOrAddTagSet(tagNameToFileSet, tagName).Add(_parameters.File);
 			}
 
 			// Add them to the set of build products
-			BuildProducts.Add(Parameters.File);
+			buildProducts.Add(_parameters.File);
 		}
 
 		/// <summary>
 		/// Output this task out to an XML writer.
 		/// </summary>
-		public override void Write(XmlWriter Writer)
+		public override void Write(XmlWriter writer)
 		{
-			Write(Writer, Parameters);
+			Write(writer, _parameters);
 		}
 
 		/// <summary>
@@ -137,9 +133,9 @@ namespace AutomationTool.Tasks
 		/// <returns>The tag names which are read by this task</returns>
 		public override IEnumerable<string> FindConsumedTagNames()
 		{
-			foreach(string TagName in FindTagNamesFromFilespec(Parameters.Files))
+			foreach (string tagName in FindTagNamesFromFilespec(_parameters.Files))
 			{
-				yield return TagName;
+				yield return tagName;
 			}
 		}
 
@@ -149,7 +145,7 @@ namespace AutomationTool.Tasks
 		/// <returns>The tag names which are modified by this task</returns>
 		public override IEnumerable<string> FindProducedTagNames()
 		{
-			return FindTagNamesFromList(Parameters.Tag);
+			return FindTagNamesFromList(_parameters.Tag);
 		}
 	}
 }
