EpicGames.BuildGraph: Apply standard naming conventions to BuildGraph.Automation.

diff --git a/Engine/Source/Programs/AutomationTool/BuildGraph/Tasks/ZenExportSnapshotTask.cs b/Engine/Source/Programs/AutomationTool/BuildGraph/Tasks/ZenExportSnapshotTask.cs
--- a/Engine/Source/Programs/AutomationTool/BuildGraph/Tasks/ZenExportSnapshotTask.cs
+++ b/Engine/Source/Programs/AutomationTool/BuildGraph/Tasks/ZenExportSnapshotTask.cs
@@ -1,6 +1,5 @@
 // Copyright Epic Games, Inc. All Rights Reserved.
 
-using Microsoft.Extensions.Logging;
 using System;
 using System.Collections.Generic;
 using System.Diagnostics.CodeAnalysis;
@@ -14,8 +13,8 @@ using System.Threading.Tasks;
 using System.Xml;
 using EpicGames.Core;
 using EpicGames.Horde;
-using EpicGames.Serialization;
 using EpicGames.ProjectStore;
+using Microsoft.Extensions.Logging;
 
 namespace AutomationTool.Tasks
 {
@@ -176,7 +175,7 @@ namespace AutomationTool.Tasks
 			/// Target platform for this snapshot
 			/// </summary>
 			public string TargetPlatform { get; set; }
-			
+
 			/// <summary>
 			/// For cloud snapshots, the host they are stored on.
 			/// </summary>
@@ -196,7 +195,7 @@ namespace AutomationTool.Tasks
 			/// For cloud snapshots, the key they are stored in.
 			/// </summary>
 			public string Key { get; set; }
-			
+
 			/// <summary>
 			/// For file snapshots, the directory it is stored in.
 			/// </summary>
@@ -217,30 +216,30 @@ namespace AutomationTool.Tasks
 			/// </summary>
 			public List<SnapshotDescriptor> Snapshots { get; set; }
 		}
-		
+
 		private class ExportSourceData
 		{
-			public bool IsLocalHost;
-			public string HostName;
-			public int HostPort;
-			public string ProjectId;
-			public string OplogId;
-			public string TargetPlatform;
-			public SnapshotDescriptor SnapshotBaseDescriptor;
+			public bool _isLocalHost;
+			public string _hostName;
+			public int _hostPort;
+			public string _projectId;
+			public string _oplogId;
+			public string _targetPlatform;
+			public SnapshotDescriptor _snapshotBaseDescriptor;
 		}
 
 		/// <summary>
 		/// Parameters for the task
 		/// </summary>
-		ZenExportSnapshotTaskParameters Parameters;
+		readonly ZenExportSnapshotTaskParameters _parameters;
 
 		/// <summary>
 		/// Constructor.
 		/// </summary>
-		/// <param name="InParameters">Parameters for this task</param>
-		public ZenExportSnapshotTask(ZenExportSnapshotTaskParameters InParameters)
+		/// <param name="parameters">Parameters for this task</param>
+		public ZenExportSnapshotTask(ZenExportSnapshotTaskParameters parameters)
 		{
-			Parameters = InParameters;
+			_parameters = parameters;
 		}
 
 		/// <summary>
@@ -252,24 +251,23 @@ namespace AutomationTool.Tasks
 			return ResolveFile(String.Format("Engine/Binaries/{0}/zen{1}", HostPlatform.Current.HostEditorPlatform.ToString(), RuntimeInformation.IsOSPlatform(OSPlatform.Windows) ? ".exe" : ""));
 		}
 
-
 		/// <summary>
 		/// Ensures that ZenServer is running on this current machine. This is needed before running any oplog commands
 		/// This passes the sponsor'd process Id to launch zen.
 		/// This ensures that zen does not live longer than the lifetime of a particular a process that needs Zen to be running
 		/// </summary>
-		/// <param name="ProjectFile"></param>
-		public static void ZenLaunch(FileReference ProjectFile)
+		/// <param name="projectFile"></param>
+		public static void ZenLaunch(FileReference projectFile)
 		{
 			// Get the ZenLaunch executable path
-			FileReference ZenLaunchExe = ResolveFile(String.Format("Engine/Binaries/{0}/ZenLaunch{1}", HostPlatform.Current.HostEditorPlatform.ToString(), RuntimeInformation.IsOSPlatform(OSPlatform.Windows) ? ".exe" : ""));
+			FileReference zenLaunchExe = ResolveFile(String.Format("Engine/Binaries/{0}/ZenLaunch{1}", HostPlatform.Current.HostEditorPlatform.ToString(), RuntimeInformation.IsOSPlatform(OSPlatform.Windows) ? ".exe" : ""));
 
-			StringBuilder ZenLaunchCommandline = new StringBuilder();
-			ZenLaunchCommandline.AppendFormat("{0} -SponsorProcessID={1}", CommandUtils.MakePathSafeToUseWithCommandLine(ProjectFile.FullName), Environment.ProcessId);
+			StringBuilder zenLaunchCommandline = new StringBuilder();
+			zenLaunchCommandline.AppendFormat("{0} -SponsorProcessID={1}", CommandUtils.MakePathSafeToUseWithCommandLine(projectFile.FullName), Environment.ProcessId);
 
-			CommandUtils.RunAndLog(CommandUtils.CmdEnv, ZenLaunchExe.FullName, ZenLaunchCommandline.ToString(), Options: CommandUtils.ERunOptions.Default);
+			CommandUtils.RunAndLog(CommandUtils.CmdEnv, zenLaunchExe.FullName, zenLaunchCommandline.ToString(), Options: CommandUtils.ERunOptions.Default);
 		}
-		
+
 		static JsonSerializerOptions GetDefaultJsonSerializerOptions()
 		{
 			JsonSerializerOptions options = new JsonSerializerOptions();
@@ -294,109 +292,97 @@ namespace AutomationTool.Tasks
 				obj = LoadJson<T>(file);
 				return true;
 			}
-			catch(Exception)
+			catch (Exception)
 			{
 				obj = null;
 				return false;
 			}
 		}
 
-		static T? TryDeserializeJson<T>(byte[] data) where T : class
-		{
-			try
-			{
-				return JsonSerializer.Deserialize<T>(data, GetDefaultJsonSerializerOptions())!;
-			}
-			catch (Exception)
-			{
-				return null;
-			}
-		}
-
 		static T LoadJson<T>(FileReference file)
 		{
 			byte[] data = FileReference.ReadAllBytes(file);
 			return JsonSerializer.Deserialize<T>(data, GetDefaultJsonSerializerOptions())!;
 		}
 
-		static string SanitizeOplogName(string Name)
+		static string SanitizeOplogName(string name)
 		{
-			return Name.Replace('/','_').Replace(' ', '_').Replace('+','_').Replace('-','_');
+			return name.Replace('/', '_').Replace(' ', '_').Replace('+', '_').Replace('-', '_');
 		}
 
-		private void WriteExportSource(JsonWriter Writer, SnapshotStorageType DestinationStorageType, ExportSourceData ExportSource, string Name)
+		private void WriteExportSource(JsonWriter writer, SnapshotStorageType destinationStorageType, ExportSourceData exportSource, string name)
 		{
-			Writer.WriteObjectStart();
-			switch (DestinationStorageType)
+			writer.WriteObjectStart();
+			switch (destinationStorageType)
 			{
 				case SnapshotStorageType.Cloud:
-					string BucketName = Parameters.DestinationCloudBucket;
-					string ProjectNameAsBucketName = Parameters.Project.GetFileNameWithoutAnyExtensions().ToLowerInvariant();
-					if (string.IsNullOrEmpty(BucketName))
+					string bucketName = _parameters.DestinationCloudBucket;
+					string projectNameAsBucketName = _parameters.Project.GetFileNameWithoutAnyExtensions().ToLowerInvariant();
+					if (string.IsNullOrEmpty(bucketName))
 					{
-						BucketName = ProjectNameAsBucketName;
+						bucketName = projectNameAsBucketName;
 					}
-					BucketName = SanitizeBucketName(BucketName);
+					bucketName = SanitizeBucketName(bucketName);
 
-					string HostName = Parameters.SnapshotDescriptorCloudHost;
-					if (string.IsNullOrEmpty(HostName))
+					string hostName = _parameters.SnapshotDescriptorCloudHost;
+					if (string.IsNullOrEmpty(hostName))
 					{
-						HostName = Parameters.DestinationCloudHost;
+						hostName = _parameters.DestinationCloudHost;
 					}
 
-					string HttpVersion = Parameters.SnapshotDescriptorCloudHttpVersion;
-					if (string.IsNullOrEmpty(HttpVersion))
+					string httpVersion = _parameters.SnapshotDescriptorCloudHttpVersion;
+					if (string.IsNullOrEmpty(httpVersion))
 					{
-						HttpVersion = Parameters.DestinationCloudHttpVersion;
+						httpVersion = _parameters.DestinationCloudHttpVersion;
 					}
 
-					IoHash DestinationKeyHash = IoHash.Compute(Encoding.UTF8.GetBytes(Name));
-					Writer.WriteValue("name", Name);
-					Writer.WriteValue("type", "cloud");
-					Writer.WriteValue("targetplatform", ExportSource.TargetPlatform);
-					Writer.WriteValue("host", HostName);
-					if (!string.IsNullOrEmpty(HttpVersion) && !HttpVersion.Equals("None", StringComparison.OrdinalIgnoreCase))
+					IoHash destinationKeyHash = IoHash.Compute(Encoding.UTF8.GetBytes(name));
+					writer.WriteValue("name", name);
+					writer.WriteValue("type", "cloud");
+					writer.WriteValue("targetplatform", exportSource._targetPlatform);
+					writer.WriteValue("host", hostName);
+					if (!string.IsNullOrEmpty(httpVersion) && !httpVersion.Equals("None", StringComparison.OrdinalIgnoreCase))
 					{
-						Writer.WriteValue("httpversion", HttpVersion);
+						writer.WriteValue("httpversion", httpVersion);
 					}
-					Writer.WriteValue("namespace", Parameters.DestinationCloudNamespace);
-					Writer.WriteValue("bucket", BucketName);
-					Writer.WriteValue("key", DestinationKeyHash.ToString().ToLowerInvariant());
+					writer.WriteValue("namespace", _parameters.DestinationCloudNamespace);
+					writer.WriteValue("bucket", bucketName);
+					writer.WriteValue("key", destinationKeyHash.ToString().ToLowerInvariant());
 					break;
 				case SnapshotStorageType.Zen:
-					string ProjectName = Parameters.Project.GetFileNameWithoutAnyExtensions().ToLowerInvariant() + ".oplog";
-
-					Writer.WriteValue("name", Name);
-					Writer.WriteValue("type", "zen");
-					Writer.WriteValue("targetplatform", ExportSource.TargetPlatform);
-					Writer.WriteValue("host", Parameters.DestinationZenHost);
-					Writer.WriteValue("projectid", ProjectName);
-					Writer.WriteValue("oplogid", SanitizeOplogName(Name));
+					string projectName = _parameters.Project.GetFileNameWithoutAnyExtensions().ToLowerInvariant() + ".oplog";
+
+					writer.WriteValue("name", name);
+					writer.WriteValue("type", "zen");
+					writer.WriteValue("targetplatform", exportSource._targetPlatform);
+					writer.WriteValue("host", _parameters.DestinationZenHost);
+					writer.WriteValue("projectid", projectName);
+					writer.WriteValue("oplogid", SanitizeOplogName(name));
 					break;
 				case SnapshotStorageType.File:
-					Writer.WriteValue("name", Name);
-					Writer.WriteValue("type", "file");
-					Writer.WriteValue("targetplatform", ExportSource.TargetPlatform);
-					Writer.WriteValue("directory", Parameters.DestinationFileDir.FullName);
-					Writer.WriteValue("filename", Parameters.DestinationFileName);
+					writer.WriteValue("name", name);
+					writer.WriteValue("type", "file");
+					writer.WriteValue("targetplatform", exportSource._targetPlatform);
+					writer.WriteValue("directory", _parameters.DestinationFileDir.FullName);
+					writer.WriteValue("filename", _parameters.DestinationFileName);
 					break;
 			}
-			Writer.WriteObjectEnd();
+			writer.WriteObjectEnd();
 		}
 
-		private static bool TryRunAndLogWithoutSpew(string App, string CommandLine, bool IgnoreFailure)
+		private static bool TryRunAndLogWithoutSpew(string app, string commandLine, bool ignoreFailure)
 		{
-			ProcessResult.SpewFilterCallbackType SilentOutputFilter = new ProcessResult.SpewFilterCallbackType(Line =>
+			ProcessResult.SpewFilterCallbackType silentOutputFilter = new ProcessResult.SpewFilterCallbackType(line =>
 				{
 					return null;
 				});
 			try
 			{
-				CommandUtils.RunAndLog(CommandUtils.CmdEnv, App, CommandLine, MaxSuccessCode: 0, Options: CommandUtils.ERunOptions.Default, SpewFilterCallback: SilentOutputFilter);
+				CommandUtils.RunAndLog(CommandUtils.CmdEnv, app, commandLine, MaxSuccessCode: 0, Options: CommandUtils.ERunOptions.Default, SpewFilterCallback: silentOutputFilter);
 			}
 			catch (CommandUtils.CommandFailedException e)
 			{
-				if (!IgnoreFailure)
+				if (!ignoreFailure)
 				{
 					Logger.LogWarning("{Text}", e.ToString());
 				}
@@ -405,147 +391,147 @@ namespace AutomationTool.Tasks
 			return true;
 		}
 
-		private static bool TryExportOplogCommand(string App, string CommandLine)
+		private static bool TryExportOplogCommand(string app, string commandLine)
 		{
-			int AttemptLimit = 2;
-			int Attempt = 0;
-			while (Attempt < AttemptLimit)
+			int attemptLimit = 2;
+			int attempt = 0;
+			while (attempt < attemptLimit)
 			{
-				if (TryRunAndLogWithoutSpew(App, CommandLine, false))
+				if (TryRunAndLogWithoutSpew(app, commandLine, false))
 				{
 					return true;
 				}
-				Logger.LogWarning("Attempt {AttemptNum} of exporting the oplog failed, {Action}...", Attempt+1, Attempt < (AttemptLimit-1) ? "retrying" : "abandoning");
+				Logger.LogWarning("Attempt {AttemptNum} of exporting the oplog failed, {Action}...", attempt + 1, attempt < (attemptLimit - 1) ? "retrying" : "abandoning");
 
-				Attempt = Attempt + 1;
+				attempt = attempt + 1;
 			}
 			return false;
 		}
 
-		private static string SanitizeBucketName(string InString)
+		private static string SanitizeBucketName(string inString)
 		{
-			return StringId.Sanitize(InString).ToString();
+			return StringId.Sanitize(inString).ToString();
 		}
 
 		/// <summary>
-		/// Execute the task.
+		/// ExecuteAsync the task.
 		/// </summary>
-		/// <param name="Job">Information about the current job</param>
-		/// <param name="BuildProducts">Set of build products produced by this node.</param>
-		/// <param name="TagNameToFileSet">Mapping from tag names to the set of files they include</param>
-		public override Task ExecuteAsync(JobContext Job, HashSet<FileReference> BuildProducts, Dictionary<string, HashSet<FileReference>> TagNameToFileSet)
+		/// <param name="job">Information about the current job</param>
+		/// <param name="buildProducts">Set of build products produced by this node.</param>
+		/// <param name="tagNameToFileSet">Mapping from tag names to the set of files they include</param>
+		public override Task ExecuteAsync(JobContext job, HashSet<FileReference> buildProducts, Dictionary<string, HashSet<FileReference>> tagNameToFileSet)
 		{
-			SnapshotStorageType DestinationStorageType = SnapshotStorageType.Invalid;
-			if (!string.IsNullOrEmpty(Parameters.DestinationStorageType))
+			SnapshotStorageType destinationStorageType = SnapshotStorageType.Invalid;
+			if (!String.IsNullOrEmpty(_parameters.DestinationStorageType))
 			{
-				DestinationStorageType = (SnapshotStorageType)Enum.Parse(typeof(SnapshotStorageType), Parameters.DestinationStorageType);
+				destinationStorageType = (SnapshotStorageType)Enum.Parse(typeof(SnapshotStorageType), _parameters.DestinationStorageType);
 			}
 
-			FileReference ProjectFile = Parameters.Project;
-			if(!FileReference.Exists(ProjectFile))
+			FileReference projectFile = _parameters.Project;
+			if (!FileReference.Exists(projectFile))
 			{
-				throw new AutomationException("Missing project file - {0}", ProjectFile.FullName);
+				throw new AutomationException("Missing project file - {0}", projectFile.FullName);
 			}
 
-			ZenLaunch(ProjectFile);
+			ZenLaunch(projectFile);
 
-			List<ExportSourceData> ExportSources = new List<ExportSourceData>();
-			foreach (string Platform in Parameters.Platform.Split('+'))
+			List<ExportSourceData> exportSources = new List<ExportSourceData>();
+			foreach (string platform in _parameters.Platform.Split('+'))
 			{
-				DirectoryReference PlatformCookedDirectory;
-				if (string.IsNullOrEmpty(Parameters.OverridePlatformCookedDir))
+				DirectoryReference platformCookedDirectory;
+				if (string.IsNullOrEmpty(_parameters.OverridePlatformCookedDir))
 				{
-					PlatformCookedDirectory = DirectoryReference.Combine(ProjectFile.Directory, "Saved", "Cooked", Platform);
+					platformCookedDirectory = DirectoryReference.Combine(projectFile.Directory, "Saved", "Cooked", platform);
 				}
 				else
 				{
-					PlatformCookedDirectory = new DirectoryReference(Parameters.OverridePlatformCookedDir.Replace("{Platform}", Platform, StringComparison.InvariantCultureIgnoreCase));
+					platformCookedDirectory = new DirectoryReference(_parameters.OverridePlatformCookedDir.Replace("{Platform}", platform, StringComparison.InvariantCultureIgnoreCase));
 				}
-				if (!DirectoryReference.Exists(PlatformCookedDirectory))
+				if (!DirectoryReference.Exists(platformCookedDirectory))
 				{
-					throw new AutomationException("Cook output directory not found ({0})", PlatformCookedDirectory.FullName);
+					throw new AutomationException("Cook output directory not found ({0})", platformCookedDirectory.FullName);
 				}
 
-				FileReference ProjectStoreFile = FileReference.Combine(PlatformCookedDirectory, "ue.projectstore");
-				ProjectStoreData? ParsedProjectStore = null;
-				if (TryLoadJson(ProjectStoreFile, out ParsedProjectStore) && (ParsedProjectStore != null) && (ParsedProjectStore.ZenServer != null))
+				FileReference projectStoreFile = FileReference.Combine(platformCookedDirectory, "ue.projectstore");
+				ProjectStoreData? parsedProjectStore = null;
+				if (TryLoadJson(projectStoreFile, out parsedProjectStore) && (parsedProjectStore != null) && (parsedProjectStore.ZenServer != null))
 				{
-					ExportSourceData NewExportSource = new ExportSourceData();
-					NewExportSource.IsLocalHost = ParsedProjectStore.ZenServer.IsLocalHost;
-					NewExportSource.HostName = ParsedProjectStore.ZenServer.HostName;
-					NewExportSource.HostPort = ParsedProjectStore.ZenServer.HostPort;
-					NewExportSource.ProjectId = ParsedProjectStore.ZenServer.ProjectId;
-					NewExportSource.OplogId = ParsedProjectStore.ZenServer.OplogId;
-					NewExportSource.TargetPlatform = Platform;
-					NewExportSource.SnapshotBaseDescriptor = null;
-
-					if (Parameters.SnapshotBaseDescriptorFile != null)
+					ExportSourceData newExportSource = new ExportSourceData();
+					newExportSource._isLocalHost = parsedProjectStore.ZenServer.IsLocalHost;
+					newExportSource._hostName = parsedProjectStore.ZenServer.HostName;
+					newExportSource._hostPort = parsedProjectStore.ZenServer.HostPort;
+					newExportSource._projectId = parsedProjectStore.ZenServer.ProjectId;
+					newExportSource._oplogId = parsedProjectStore.ZenServer.OplogId;
+					newExportSource._targetPlatform = platform;
+					newExportSource._snapshotBaseDescriptor = null;
+
+					if (_parameters.SnapshotBaseDescriptorFile != null)
 					{
-						FileReference PlatformSnapshotBase = new FileReference(Parameters.SnapshotBaseDescriptorFile.FullName.Replace("{Platform}", Platform, StringComparison.InvariantCultureIgnoreCase));
+						FileReference platformSnapshotBase = new FileReference(_parameters.SnapshotBaseDescriptorFile.FullName.Replace("{Platform}", platform, StringComparison.InvariantCultureIgnoreCase));
 
-						SnapshotDescriptorCollection? ParsedDescriptorCollection = null;
-						if (TryLoadJson(PlatformSnapshotBase, out ParsedDescriptorCollection) && (ParsedDescriptorCollection != null) && (ParsedDescriptorCollection.Snapshots != null))
+						SnapshotDescriptorCollection? parsedDescriptorCollection = null;
+						if (TryLoadJson(platformSnapshotBase, out parsedDescriptorCollection) && (parsedDescriptorCollection != null) && (parsedDescriptorCollection.Snapshots != null))
 						{
-							foreach (SnapshotDescriptor ParsedDescriptor in ParsedDescriptorCollection.Snapshots)
+							foreach (SnapshotDescriptor parsedDescriptor in parsedDescriptorCollection.Snapshots)
 							{
-								if (ParsedDescriptor.TargetPlatform == Platform)
+								if (parsedDescriptor.TargetPlatform == platform)
 								{
-									NewExportSource.SnapshotBaseDescriptor = ParsedDescriptor;
+									newExportSource._snapshotBaseDescriptor = parsedDescriptor;
 									break;
 								}
 							}
 						}
 					}
 
-					ExportSources.Add(NewExportSource);
+					exportSources.Add(newExportSource);
 				}
 			}
-			int ExportIndex = 0;
-			string[] ExportNames = new string[ExportSources.Count];
-			List<ExportSourceData> SuccessfullyExportedSources = new List<ExportSourceData>();
+			int exportIndex;
+			string[] exportNames = new string[exportSources.Count];
+			List<ExportSourceData> successfullyExportedSources = new List<ExportSourceData>();
 
 			// Get the Zen executable path
-			FileReference ZenExe = ZenExeFileReference();
+			FileReference zenExe = ZenExeFileReference();
 
 			// Format the command line
-			StringBuilder OplogExportCommandline = new StringBuilder();
-			OplogExportCommandline.Append("oplog-export --embedloosefiles");
-			if (Parameters.Force)
+			StringBuilder oplogExportCommandline = new StringBuilder();
+			oplogExportCommandline.Append("oplog-export --embedloosefiles");
+			if (_parameters.Force)
 			{
-				OplogExportCommandline.Append(" --force");
+				oplogExportCommandline.Append(" --force");
 			}
 
-			switch (DestinationStorageType)
+			switch (destinationStorageType)
 			{
 				case SnapshotStorageType.Cloud:
-					if (string.IsNullOrEmpty(Parameters.DestinationCloudHost))
+					if (string.IsNullOrEmpty(_parameters.DestinationCloudHost))
 					{
 						throw new AutomationException("Missing destination cloud host");
 					}
-					if (string.IsNullOrEmpty(Parameters.DestinationCloudNamespace))
+					if (string.IsNullOrEmpty(_parameters.DestinationCloudNamespace))
 					{
 						throw new AutomationException("Missing destination cloud namespace");
 					}
-					if (string.IsNullOrEmpty(Parameters.DestinationIdentifier))
+					if (string.IsNullOrEmpty(_parameters.DestinationIdentifier))
 					{
 						throw new AutomationException("Missing destination identifier when exporting to cloud");
 					}
 
-					string BucketName = Parameters.DestinationCloudBucket;
-					string ProjectNameAsBucketName = ProjectFile.GetFileNameWithoutAnyExtensions().ToLowerInvariant();
-					if (string.IsNullOrEmpty(BucketName))
+					string bucketName = _parameters.DestinationCloudBucket;
+					string projectNameAsBucketName = projectFile.GetFileNameWithoutAnyExtensions().ToLowerInvariant();
+					if (string.IsNullOrEmpty(bucketName))
 					{
-						BucketName = ProjectNameAsBucketName;
+						bucketName = projectNameAsBucketName;
 					}
-					BucketName = SanitizeBucketName(BucketName);
+					bucketName = SanitizeBucketName(bucketName);
 
-					OplogExportCommandline.AppendFormat(" --cloud {0} --namespace {1} --bucket {2}", Parameters.DestinationCloudHost, Parameters.DestinationCloudNamespace, BucketName);
+					oplogExportCommandline.AppendFormat(" --cloud {0} --namespace {1} --bucket {2}", _parameters.DestinationCloudHost, _parameters.DestinationCloudNamespace, bucketName);
 
-					if (!string.IsNullOrEmpty(Parameters.DestinationCloudHttpVersion))
+					if (!string.IsNullOrEmpty(_parameters.DestinationCloudHttpVersion))
 					{
-						if (Parameters.DestinationCloudHttpVersion.Equals("http2-only", StringComparison.OrdinalIgnoreCase))
+						if (_parameters.DestinationCloudHttpVersion.Equals("http2-only", StringComparison.OrdinalIgnoreCase))
 						{
-							OplogExportCommandline.Append(" --assume-http2");
+							oplogExportCommandline.Append(" --assume-http2");
 						}
 						else
 						{
@@ -553,112 +539,112 @@ namespace AutomationTool.Tasks
 						}
 					}
 
-					ExportIndex = 0;
-					foreach (ExportSourceData ExportSource in ExportSources)
+					exportIndex = 0;
+					foreach (ExportSourceData exportSource in exportSources)
 					{
-						string HostUrlArg = string.Format("--hosturl http://{0}:{1}", ExportSource.IsLocalHost ? "localhost" : ExportSource.HostName, ExportSource.HostPort);
-						
-						string BaseKeyArg = string.Empty;
-						if ((ExportSource.SnapshotBaseDescriptor != null) && !string.IsNullOrEmpty(ExportSource.SnapshotBaseDescriptor.Key))
+						string hostUrlArg = string.Format("--hosturl http://{0}:{1}", exportSource._isLocalHost ? "localhost" : exportSource._hostName, exportSource._hostPort);
+
+						string baseKeyArg = string.Empty;
+						if ((exportSource._snapshotBaseDescriptor != null) && !string.IsNullOrEmpty(exportSource._snapshotBaseDescriptor.Key))
 						{
-							if (ExportSource.SnapshotBaseDescriptor.Type == SnapshotStorageType.Cloud)
+							if (exportSource._snapshotBaseDescriptor.Type == SnapshotStorageType.Cloud)
 							{
-								BaseKeyArg = " --basekey " + ExportSource.SnapshotBaseDescriptor.Key;
+								baseKeyArg = " --basekey " + exportSource._snapshotBaseDescriptor.Key;
 							}
 							else
 							{
-								Logger.LogWarning("Base snapshot descriptor was for a snapshot storage type {Type}, but we're producing a snapshot of type cloud.  Skipping use of base snapshot.", ExportSource.SnapshotBaseDescriptor.Type);
+								Logger.LogWarning("Base snapshot descriptor was for a snapshot storage type {Type}, but we're producing a snapshot of type cloud.  Skipping use of base snapshot.", exportSource._snapshotBaseDescriptor.Type);
 							}
 						}
 
-						StringBuilder ExportSingleSourceCommandline = new StringBuilder(OplogExportCommandline.Length);
-						ExportSingleSourceCommandline.Append(OplogExportCommandline);
+						StringBuilder exportSingleSourceCommandline = new StringBuilder(oplogExportCommandline.Length);
+						exportSingleSourceCommandline.Append(oplogExportCommandline);
 
-						StringBuilder DestinationKeyBuilder = new StringBuilder();
-						DestinationKeyBuilder.AppendFormat("{0}.{1}.{2}", ProjectNameAsBucketName, Parameters.DestinationIdentifier, ExportSource.OplogId);
-						ExportNames[ExportIndex] = DestinationKeyBuilder.ToString().ToLowerInvariant();
-						IoHash DestinationKeyHash = IoHash.Compute(Encoding.UTF8.GetBytes(ExportNames[ExportIndex]));
+						StringBuilder destinationKeyBuilder = new StringBuilder();
+						destinationKeyBuilder.AppendFormat("{0}.{1}.{2}", projectNameAsBucketName, _parameters.DestinationIdentifier, exportSource._oplogId);
+						exportNames[exportIndex] = destinationKeyBuilder.ToString().ToLowerInvariant();
+						IoHash destinationKeyHash = IoHash.Compute(Encoding.UTF8.GetBytes(exportNames[exportIndex]));
 
-						ExportSingleSourceCommandline.AppendFormat(" {0} --key {1} {2} {3} {4}", HostUrlArg, DestinationKeyHash.ToString().ToLowerInvariant(), BaseKeyArg, ExportSource.ProjectId, ExportSource.OplogId);
-						if (TryExportOplogCommand(ZenExe.FullName, ExportSingleSourceCommandline.ToString()))
+						exportSingleSourceCommandline.AppendFormat(" {0} --key {1} {2} {3} {4}", hostUrlArg, destinationKeyHash.ToString().ToLowerInvariant(), baseKeyArg, exportSource._projectId, exportSource._oplogId);
+						if (TryExportOplogCommand(zenExe.FullName, exportSingleSourceCommandline.ToString()))
 						{
-							SuccessfullyExportedSources.Add(ExportSource);
+							successfullyExportedSources.Add(exportSource);
 						}
 
-						ExportIndex = ExportIndex + 1;
+						exportIndex = exportIndex + 1;
 					}
 
 					break;
 				case SnapshotStorageType.Zen:
-					if (string.IsNullOrEmpty(Parameters.DestinationZenHost))
+					if (string.IsNullOrEmpty(_parameters.DestinationZenHost))
 					{
 						throw new AutomationException("Missing destination zen host");
 					}
-					if (string.IsNullOrEmpty(Parameters.DestinationIdentifier))
+					if (string.IsNullOrEmpty(_parameters.DestinationIdentifier))
 					{
 						throw new AutomationException("Missing destination identifier when exporting to zen");
 					}
 
-					string ProjectName = ProjectFile.GetFileNameWithoutAnyExtensions().ToLowerInvariant() + ".oplog";
+					string projectName = projectFile.GetFileNameWithoutAnyExtensions().ToLowerInvariant() + ".oplog";
 
-					StringBuilder CreateProjectCommandline = new StringBuilder();
-					CreateProjectCommandline.AppendFormat("project-create --hosturl {0} {1}", Parameters.DestinationZenHost, ProjectName);
-					TryRunAndLogWithoutSpew(ZenExe.FullName, CreateProjectCommandline.ToString(), true);
+					StringBuilder createProjectCommandline = new StringBuilder();
+					createProjectCommandline.AppendFormat("project-create --hosturl {0} {1}", _parameters.DestinationZenHost, projectName);
+					TryRunAndLogWithoutSpew(zenExe.FullName, createProjectCommandline.ToString(), true);
 
-					OplogExportCommandline.AppendFormat(" --zen {0}", Parameters.DestinationZenHost);
+					oplogExportCommandline.AppendFormat(" --zen {0}", _parameters.DestinationZenHost);
 
-					ExportIndex = 0;
-					foreach (ExportSourceData ExportSource in ExportSources)
+					exportIndex = 0;
+					foreach (ExportSourceData exportSource in exportSources)
 					{
-						string HostUrlArg = string.Format("--hosturl http://{0}:{1}", ExportSource.IsLocalHost ? "localhost" : ExportSource.HostName, ExportSource.HostPort);
+						string hostUrlArg = string.Format("--hosturl http://{0}:{1}", exportSource._isLocalHost ? "localhost" : exportSource._hostName, exportSource._hostPort);
 
-						StringBuilder ExportSingleSourceCommandline = new StringBuilder(OplogExportCommandline.Length);
-						ExportSingleSourceCommandline.Append(OplogExportCommandline);
+						StringBuilder exportSingleSourceCommandline = new StringBuilder(oplogExportCommandline.Length);
+						exportSingleSourceCommandline.Append(oplogExportCommandline);
 
-						StringBuilder DestinationKeyBuilder = new StringBuilder();
-						DestinationKeyBuilder.AppendFormat("{0}.{1}", Parameters.DestinationIdentifier, ExportSource.OplogId);
-						ExportNames[ExportIndex] = DestinationKeyBuilder.ToString().ToLowerInvariant();
-						string DestinationOplog = SanitizeOplogName(ExportNames[ExportIndex]);
+						StringBuilder destinationKeyBuilder = new StringBuilder();
+						destinationKeyBuilder.AppendFormat("{0}.{1}", _parameters.DestinationIdentifier, exportSource._oplogId);
+						exportNames[exportIndex] = destinationKeyBuilder.ToString().ToLowerInvariant();
+						string destinationOplog = SanitizeOplogName(exportNames[exportIndex]);
 
-						ExportSingleSourceCommandline.AppendFormat(" {0} --target-project {1} --target-oplog {2} {3} {4}", HostUrlArg, ProjectName, DestinationOplog, ExportSource.ProjectId, ExportSource.OplogId);
-						if (TryExportOplogCommand(ZenExe.FullName, ExportSingleSourceCommandline.ToString()))
+						exportSingleSourceCommandline.AppendFormat(" {0} --target-project {1} --target-oplog {2} {3} {4}", hostUrlArg, projectName, destinationOplog, exportSource._projectId, exportSource._oplogId);
+						if (TryExportOplogCommand(zenExe.FullName, exportSingleSourceCommandline.ToString()))
 						{
-							SuccessfullyExportedSources.Add(ExportSource);
+							successfullyExportedSources.Add(exportSource);
 						}
 
-						ExportIndex = ExportIndex + 1;
+						exportIndex = exportIndex + 1;
 					}
 
 					break;
 				case SnapshotStorageType.File:
-					string DefaultProjectId = ProjectUtils.GetProjectPathId(ProjectFile);
-					ExportIndex = 0;
-					foreach (ExportSourceData ExportSource in ExportSources)
+					string defaultProjectId = ProjectUtils.GetProjectPathId(projectFile);
+					exportIndex = 0;
+					foreach (ExportSourceData exportSource in exportSources)
 					{
-						StringBuilder ExportNameBuilder = new StringBuilder();
-						ExportNameBuilder.AppendFormat("{0}.{1}.{2}", ProjectFile.GetFileNameWithoutAnyExtensions().ToLowerInvariant(), Parameters.DestinationIdentifier, ExportSource.OplogId);
-						ExportNames[ExportIndex] = ExportNameBuilder.ToString().ToLowerInvariant();
+						StringBuilder exportNameBuilder = new StringBuilder();
+						exportNameBuilder.AppendFormat("{0}.{1}.{2}", projectFile.GetFileNameWithoutAnyExtensions().ToLowerInvariant(), _parameters.DestinationIdentifier, exportSource._oplogId);
+						exportNames[exportIndex] = exportNameBuilder.ToString().ToLowerInvariant();
 
-						StringBuilder ExportSingleSourceCommandline = new StringBuilder(OplogExportCommandline.Length);
-						ExportSingleSourceCommandline.Append(OplogExportCommandline);
+						StringBuilder exportSingleSourceCommandline = new StringBuilder(oplogExportCommandline.Length);
+						exportSingleSourceCommandline.Append(oplogExportCommandline);
 
-						string DestinationFileName = ExportSource.OplogId;
-						if (!string.IsNullOrEmpty(Parameters.DestinationFileName))
+						string destinationFileName = exportSource._oplogId;
+						if (!string.IsNullOrEmpty(_parameters.DestinationFileName))
 						{
-							DestinationFileName = Parameters.DestinationFileName.Replace("{Platform}", ExportSource.TargetPlatform, StringComparison.InvariantCultureIgnoreCase);
+							destinationFileName = _parameters.DestinationFileName.Replace("{Platform}", exportSource._targetPlatform, StringComparison.InvariantCultureIgnoreCase);
 						}
 
-						string ProjectId = string.IsNullOrEmpty(ExportSource.ProjectId) ? DefaultProjectId : ExportSource.ProjectId;
-						string BaseNameArg = string.Empty;
-						DirectoryReference PlatformDestinationFileDir = new DirectoryReference(Parameters.DestinationFileDir.FullName.Replace("{Platform}", ExportSource.TargetPlatform, StringComparison.InvariantCultureIgnoreCase));
-						if ((ExportSource.SnapshotBaseDescriptor != null) && !string.IsNullOrEmpty(ExportSource.SnapshotBaseDescriptor.Directory) && !string.IsNullOrEmpty(ExportSource.SnapshotBaseDescriptor.Filename))
+						string projectId = string.IsNullOrEmpty(exportSource._projectId) ? defaultProjectId : exportSource._projectId;
+						string baseNameArg = string.Empty;
+						DirectoryReference platformDestinationFileDir = new DirectoryReference(_parameters.DestinationFileDir.FullName.Replace("{Platform}", exportSource._targetPlatform, StringComparison.InvariantCultureIgnoreCase));
+						if ((exportSource._snapshotBaseDescriptor != null) && !string.IsNullOrEmpty(exportSource._snapshotBaseDescriptor.Directory) && !string.IsNullOrEmpty(exportSource._snapshotBaseDescriptor.Filename))
 						{
-							if (ExportSource.SnapshotBaseDescriptor.Type == SnapshotStorageType.File)
+							if (exportSource._snapshotBaseDescriptor.Type == SnapshotStorageType.File)
 							{
-								FileReference BaseSnapshotFile = new FileReference(Path.Combine(ExportSource.SnapshotBaseDescriptor.Directory, ExportSource.SnapshotBaseDescriptor.Filename));
-								if (FileReference.Exists(BaseSnapshotFile))
+								FileReference baseSnapshotFile = new FileReference(Path.Combine(exportSource._snapshotBaseDescriptor.Directory, exportSource._snapshotBaseDescriptor.Filename));
+								if (FileReference.Exists(baseSnapshotFile))
 								{
-									BaseNameArg = " --basename " + CommandUtils.MakePathSafeToUseWithCommandLine(BaseSnapshotFile.FullName);
+									baseNameArg = " --basename " + CommandUtils.MakePathSafeToUseWithCommandLine(baseSnapshotFile.FullName);
 								}
 								else
 								{
@@ -667,77 +653,75 @@ namespace AutomationTool.Tasks
 							}
 							else
 							{
-								Logger.LogWarning("Base snapshot descriptor was for a snapshot storage type {Type}, but we're producing a snapshot of type file.  Skipping use of base snapshot.", ExportSource.SnapshotBaseDescriptor.Type);
+								Logger.LogWarning("Base snapshot descriptor was for a snapshot storage type {Type}, but we're producing a snapshot of type file.  Skipping use of base snapshot.", exportSource._snapshotBaseDescriptor.Type);
 							}
 						}
-						ExportSingleSourceCommandline.AppendFormat(" --file {0} --name {1} {2} {3} {4}", CommandUtils.MakePathSafeToUseWithCommandLine(PlatformDestinationFileDir.FullName), DestinationFileName, BaseNameArg, ProjectId, ExportSource.OplogId);
+						exportSingleSourceCommandline.AppendFormat(" --file {0} --name {1} {2} {3} {4}", CommandUtils.MakePathSafeToUseWithCommandLine(platformDestinationFileDir.FullName), destinationFileName, baseNameArg, projectId, exportSource._oplogId);
 
-						if (TryExportOplogCommand(ZenExe.FullName, ExportSingleSourceCommandline.ToString()))
+						if (TryExportOplogCommand(zenExe.FullName, exportSingleSourceCommandline.ToString()))
 						{
-							SuccessfullyExportedSources.Add(ExportSource);
+							successfullyExportedSources.Add(exportSource);
 						}
 
-						ExportIndex = ExportIndex + 1;
+						exportIndex = exportIndex + 1;
 					}
 					break;
 				default:
-					throw new AutomationException("Unknown/invalid/unimplemented destination storage type - {0}", Parameters.DestinationStorageType);
+					throw new AutomationException("Unknown/invalid/unimplemented destination storage type - {0}", _parameters.DestinationStorageType);
 			}
-			
 
-			if ((Parameters.SnapshotDescriptorFile != null) && SuccessfullyExportedSources.Any())
+			if ((_parameters.SnapshotDescriptorFile != null) && successfullyExportedSources.Any())
 			{
-				if (Parameters.SnapshotDescriptorFile.FullName.Contains("{Platform}", StringComparison.OrdinalIgnoreCase))
+				if (_parameters.SnapshotDescriptorFile.FullName.Contains("{Platform}", StringComparison.OrdinalIgnoreCase))
 				{
 					// Separate descriptor file per platform
-					ExportIndex = 0;
-					foreach (ExportSourceData ExportSource in SuccessfullyExportedSources)
+					exportIndex = 0;
+					foreach (ExportSourceData exportSource in successfullyExportedSources)
 					{
-						FileReference PlatformSnapshotDescriptorFile = new FileReference(Parameters.SnapshotDescriptorFile.FullName.Replace("{Platform}", ExportSource.TargetPlatform, StringComparison.InvariantCultureIgnoreCase));
-						DirectoryReference.CreateDirectory(PlatformSnapshotDescriptorFile.Directory);
-						using (JsonWriter Writer = new JsonWriter(PlatformSnapshotDescriptorFile))
+						FileReference platformSnapshotDescriptorFile = new FileReference(_parameters.SnapshotDescriptorFile.FullName.Replace("{Platform}", exportSource._targetPlatform, StringComparison.InvariantCultureIgnoreCase));
+						DirectoryReference.CreateDirectory(platformSnapshotDescriptorFile.Directory);
+						using (JsonWriter writer = new JsonWriter(platformSnapshotDescriptorFile))
 						{
-							Writer.WriteObjectStart();
-							Writer.WriteArrayStart("snapshots");
-							WriteExportSource(Writer, DestinationStorageType, ExportSource, ExportNames[ExportIndex]);
-							Writer.WriteArrayEnd();
-							Writer.WriteObjectEnd();
+							writer.WriteObjectStart();
+							writer.WriteArrayStart("snapshots");
+							WriteExportSource(writer, destinationStorageType, exportSource, exportNames[exportIndex]);
+							writer.WriteArrayEnd();
+							writer.WriteObjectEnd();
 						}
-						ExportIndex = ExportIndex + 1;
+						exportIndex = exportIndex + 1;
 					}
 				}
 				else
 				{
 					// Write out a single snapshot descriptor with info about all snapshots
-					DirectoryReference.CreateDirectory(Parameters.SnapshotDescriptorFile.Directory);
-					using (JsonWriter Writer = new JsonWriter(Parameters.SnapshotDescriptorFile))
+					DirectoryReference.CreateDirectory(_parameters.SnapshotDescriptorFile.Directory);
+					using (JsonWriter writer = new JsonWriter(_parameters.SnapshotDescriptorFile))
 					{
-						Writer.WriteObjectStart();
-						Writer.WriteArrayStart("snapshots");
-								
-						ExportIndex = 0;
-						foreach (ExportSourceData ExportSource in SuccessfullyExportedSources)
+						writer.WriteObjectStart();
+						writer.WriteArrayStart("snapshots");
+
+						exportIndex = 0;
+						foreach (ExportSourceData exportSource in successfullyExportedSources)
 						{
-							WriteExportSource(Writer, DestinationStorageType, ExportSource, ExportNames[ExportIndex]);
-							ExportIndex = ExportIndex + 1;
+							WriteExportSource(writer, destinationStorageType, exportSource, exportNames[exportIndex]);
+							exportIndex = exportIndex + 1;
 						}
 
-						Writer.WriteArrayEnd();
-						Writer.WriteObjectEnd();
+						writer.WriteArrayEnd();
+						writer.WriteObjectEnd();
 					}
 				}
 			}
 
-
 			return Task.CompletedTask;
 		}
 
 		/// <summary>
 		/// Output this task out to an XML writer.
 		/// </summary>
-		public override void Write(XmlWriter Writer)
+		public override void Write(XmlWriter writer)
 		{
-			Write(Writer, Parameters);
+			Write(writer, _parameters);
 		}
 
 		/// <summary>
diff --git a/Engine/Source/Programs/AutomationTool/BuildGraph/Tasks/ZenImportOplogTask.cs b/Engine/Source/Programs/AutomationTool/BuildGraph/Tasks/ZenImportOplogTask.cs
--- a/Engine/Source/Programs/AutomationTool/BuildGraph/Tasks/ZenImportOplogTask.cs
+++ b/Engine/Source/Programs/AutomationTool/BuildGraph/Tasks/ZenImportOplogTask.cs
@@ -1,8 +1,5 @@
 // Copyright Epic Games, Inc. All Rights Reserved.
 
-using EpicGames.Core;
-using EpicGames.ProjectStore;
-using Microsoft.Extensions.Logging;
 using System;
 using System.Collections.Generic;
 using System.IO;
@@ -11,6 +8,9 @@ using System.Text.Json;
 using System.Text.Json.Serialization;
 using System.Threading.Tasks;
 using System.Xml;
+using EpicGames.Core;
+using EpicGames.ProjectStore;
+using Microsoft.Extensions.Logging;
 
 namespace AutomationTool.Tasks
 {
@@ -74,7 +74,6 @@ namespace AutomationTool.Tasks
 		/// </summary>
 		[TaskParameter(Optional = true)]
 		public string HostPort { get; set; } = "8558";
-		
 
 		/// <summary>
 		/// The cloud URL to import from
@@ -107,173 +106,170 @@ namespace AutomationTool.Tasks
 	[TaskElement("ZenImportOplog", typeof(ZenImportOplogTaskParameters))]
 	public class ZenImportOplogTask : BgTaskImpl
 	{
-		/// <summary>
-		/// Parameters for the task
-		/// </summary>
-		ZenImportOplogTaskParameters Parameters;
+		readonly ZenImportOplogTaskParameters _parameters;
 
-		FileReference ProjectFile;
+		FileReference _projectFile;
 
 		/// <summary>
 		/// Constructor.
 		/// </summary>
-		/// <param name="InParameters">Parameters for this task</param>
-		public ZenImportOplogTask(ZenImportOplogTaskParameters InParameters)
+		/// <param name="parameters">Parameters for this task</param>
+		public ZenImportOplogTask(ZenImportOplogTaskParameters parameters)
 		{
-			Parameters = InParameters;
+			_parameters = parameters;
 		}
 
 		/// <summary>
-		/// Execute the task.
+		/// ExecuteAsync the task.
 		/// </summary>
-		/// <param name="Job">Information about the current job</param>
-		/// <param name="BuildProducts">Set of build products produced by this node.</param>
-		/// <param name="TagNameToFileSet">Mapping from tag names to the set of files they include</param>
-		public override Task ExecuteAsync(JobContext Job, HashSet<FileReference> BuildProducts, Dictionary<string, HashSet<FileReference>> TagNameToFileSet)
+		/// <param name="job">Information about the current job</param>
+		/// <param name="buildProducts">Set of build products produced by this node.</param>
+		/// <param name="tagNameToFileSet">Mapping from tag names to the set of files they include</param>
+		public override Task ExecuteAsync(JobContext job, HashSet<FileReference> buildProducts, Dictionary<string, HashSet<FileReference>> tagNameToFileSet)
 		{
-			SnapshotStorageType ImportMethod = SnapshotStorageType.Invalid;
-			if (!string.IsNullOrEmpty(Parameters.ImportType))
+			SnapshotStorageType importMethod = SnapshotStorageType.Invalid;
+			if (!string.IsNullOrEmpty(_parameters.ImportType))
 			{
-				ImportMethod = (SnapshotStorageType)Enum.Parse(typeof(SnapshotStorageType), Parameters.ImportType);
+				importMethod = (SnapshotStorageType)Enum.Parse(typeof(SnapshotStorageType), _parameters.ImportType);
 			}
 
-			ProjectFile = Parameters.Project;
-			if (!FileReference.Exists(ProjectFile))
+			_projectFile = _parameters.Project;
+			if (!FileReference.Exists(_projectFile))
 			{
-				throw new AutomationException("Missing project file - {0}", ProjectFile.FullName);
+				throw new AutomationException("Missing project file - {0}", _projectFile.FullName);
 			}
 
-			ZenExportSnapshotTask.ZenLaunch(ProjectFile);
+			ZenExportSnapshotTask.ZenLaunch(_projectFile);
 
 			// Get the Zen executable path
-			FileReference ZenExe = ZenExportSnapshotTask.ZenExeFileReference();
+			FileReference zenExe = ZenExportSnapshotTask.ZenExeFileReference();
 			{
-				if (String.IsNullOrEmpty(Parameters.RootDir))
+				if (String.IsNullOrEmpty(_parameters.RootDir))
 				{
 					throw new AutomationException("RootDir was not specified");
 				}
-				if (String.IsNullOrEmpty(Parameters.ProjectName))
+				if (String.IsNullOrEmpty(_parameters.ProjectName))
 				{
 					throw new AutomationException("ProjectName was not specified");
 				}
-				
+
 				// Create a new project to import everything into.
-				string RootDir = Parameters.RootDir;
-				string EngineDir = System.IO.Path.Combine(Parameters.RootDir, "Engine");
-				string ProjectDir = System.IO.Path.Combine(Parameters.RootDir, ProjectFile.GetFileNameWithoutAnyExtensions());
-				string HostURLArg = string.Format("--hosturl http://{0}:{1}", Parameters.HostName, Parameters.HostPort);
-				StringBuilder OplogProjectCreateCommandline = new StringBuilder();
-				OplogProjectCreateCommandline.AppendFormat("project-create -p {0} --rootdir {1} --enginedir {2} --projectdir {3} --projectfile {4} {5}",
-					Parameters.ProjectName,
-					RootDir,
-					EngineDir,
-					ProjectDir,
-					ProjectFile.FullName,
-					HostURLArg);
-
-				Logger.LogInformation("Running '{Arg0} {Arg1}'", CommandUtils.MakePathSafeToUseWithCommandLine(ZenExe.FullName), OplogProjectCreateCommandline.ToString());
-				CommandUtils.RunAndLog(CommandUtils.CmdEnv, ZenExe.FullName, OplogProjectCreateCommandline.ToString(), Options: CommandUtils.ERunOptions.Default);
+				string rootDir = _parameters.RootDir;
+				string engineDir = System.IO.Path.Combine(_parameters.RootDir, "Engine");
+				string projectDir = System.IO.Path.Combine(_parameters.RootDir, _projectFile.GetFileNameWithoutAnyExtensions());
+				string hostUrlArg = string.Format("--hosturl http://{0}:{1}", _parameters.HostName, _parameters.HostPort);
+				StringBuilder oplogProjectCreateCommandline = new StringBuilder();
+				oplogProjectCreateCommandline.AppendFormat("project-create -p {0} --rootdir {1} --enginedir {2} --projectdir {3} --projectfile {4} {5}",
+					_parameters.ProjectName,
+					rootDir,
+					engineDir,
+					projectDir,
+					_projectFile.FullName,
+					hostUrlArg);
+
+				Logger.LogInformation("Running '{Arg0} {Arg1}'", CommandUtils.MakePathSafeToUseWithCommandLine(zenExe.FullName), oplogProjectCreateCommandline.ToString());
+				CommandUtils.RunAndLog(CommandUtils.CmdEnv, zenExe.FullName, oplogProjectCreateCommandline.ToString(), Options: CommandUtils.ERunOptions.Default);
 			}
 
-			switch (ImportMethod)
+			switch (importMethod)
 			{
 				case SnapshotStorageType.File:
-					ImportFromFile(ZenExe);
+					ImportFromFile(zenExe);
 					break;
 				case SnapshotStorageType.Cloud:
-					ImportFromCloud(ZenExe);
+					ImportFromCloud(zenExe);
 					break;
 				default:
-					throw new AutomationException("Unknown/invalid/unimplemented import type - {0}", Parameters.ImportType);
+					throw new AutomationException("Unknown/invalid/unimplemented import type - {0}", _parameters.ImportType);
 			}
 
 			WriteProjectStoreFile();
 			return Task.CompletedTask;
 		}
 
-		private void ImportFromFile(FileReference ZenExe)
+		private void ImportFromFile(FileReference zenExe)
 		{
-			if (String.IsNullOrEmpty(Parameters.OplogName))
+			if (String.IsNullOrEmpty(_parameters.OplogName))
 			{
 				throw new AutomationException("OplogName was not specified");
 			}
 
-			foreach (string FileToImport in Parameters.Files.Split(','))
+			foreach (string fileToImport in _parameters.Files.Split(','))
 			{
-				if (DirectoryReference.Exists(new DirectoryReference(FileToImport)))
+				if (DirectoryReference.Exists(new DirectoryReference(fileToImport)))
 				{
-					StringBuilder OplogImportCommandline = new StringBuilder();
-					OplogImportCommandline.AppendFormat("oplog-import --file {0} --oplog {1} -p {2}", FileToImport, Parameters.OplogName, Parameters.ProjectName);
+					StringBuilder oplogImportCommandline = new StringBuilder();
+					oplogImportCommandline.AppendFormat("oplog-import --file {0} --oplog {1} -p {2}", fileToImport, _parameters.OplogName, _parameters.ProjectName);
 
-					Logger.LogInformation("Running '{Arg0} {Arg1}'", CommandUtils.MakePathSafeToUseWithCommandLine(ZenExe.FullName), OplogImportCommandline.ToString());
-					CommandUtils.RunAndLog(CommandUtils.CmdEnv, ZenExe.FullName, OplogImportCommandline.ToString(), Options: CommandUtils.ERunOptions.Default);
+					Logger.LogInformation("Running '{Arg0} {Arg1}'", CommandUtils.MakePathSafeToUseWithCommandLine(zenExe.FullName), oplogImportCommandline.ToString());
+					CommandUtils.RunAndLog(CommandUtils.CmdEnv, zenExe.FullName, oplogImportCommandline.ToString(), Options: CommandUtils.ERunOptions.Default);
 				}
 			}
 		}
 
 		private void WriteProjectStoreFile()
 		{
-			DirectoryReference PlatformCookedDirectory = DirectoryReference.Combine(ProjectFile.Directory, "Saved", "Cooked", Parameters.Platform);
-			if (!DirectoryReference.Exists(PlatformCookedDirectory))
+			DirectoryReference platformCookedDirectory = DirectoryReference.Combine(_projectFile.Directory, "Saved", "Cooked", _parameters.Platform);
+			if (!DirectoryReference.Exists(platformCookedDirectory))
 			{
-				DirectoryReference.CreateDirectory(PlatformCookedDirectory);
+				DirectoryReference.CreateDirectory(platformCookedDirectory);
 			}
-			ProjectStoreData ProjectStore = new ProjectStoreData();
-			ProjectStore.ZenServer = new ZenServerStoreData
+			ProjectStoreData projectStore = new ProjectStoreData();
+			projectStore.ZenServer = new ZenServerStoreData
 			{
-				ProjectId = Parameters.ProjectName,
-				OplogId = Parameters.OplogName
+				ProjectId = _parameters.ProjectName,
+				OplogId = _parameters.OplogName
 			};
 
-			JsonSerializerOptions SerializerOptions = new JsonSerializerOptions
+			JsonSerializerOptions serializerOptions = new JsonSerializerOptions
 			{
 				AllowTrailingCommas = true,
 				ReadCommentHandling = JsonCommentHandling.Skip,
 				PropertyNameCaseInsensitive = true
 			};
-			SerializerOptions.Converters.Add(new JsonStringEnumConverter());
+			serializerOptions.Converters.Add(new JsonStringEnumConverter());
 
-			FileReference ProjectStoreFile = FileReference.Combine(PlatformCookedDirectory, "ue.projectstore");
-			File.WriteAllText(ProjectStoreFile.FullName, JsonSerializer.Serialize(ProjectStore, SerializerOptions), new UTF8Encoding(encoderShouldEmitUTF8Identifier: false));
+			FileReference projectStoreFile = FileReference.Combine(platformCookedDirectory, "ue.projectstore");
+			File.WriteAllText(projectStoreFile.FullName, JsonSerializer.Serialize(projectStore, serializerOptions), new UTF8Encoding(encoderShouldEmitUTF8Identifier: false));
 		}
 
-		private void ImportFromCloud(FileReference ZenExe)
+		private void ImportFromCloud(FileReference zenExe)
 		{
-			if (string.IsNullOrEmpty(Parameters.CloudURL))
+			if (string.IsNullOrEmpty(_parameters.CloudURL))
 			{
 				throw new AutomationException("Missing destination cloud host");
 			}
-			if (string.IsNullOrEmpty(Parameters.Namespace))
+			if (string.IsNullOrEmpty(_parameters.Namespace))
 			{
 				throw new AutomationException("Missing destination cloud namespace");
 			}
-			if (string.IsNullOrEmpty(Parameters.Key))
+			if (string.IsNullOrEmpty(_parameters.Key))
 			{
 				throw new AutomationException("Missing destination cloud storage key");
 			}
 
-			string BucketName = Parameters.Bucket;
-			string ProjectNameAsBucketName = ProjectFile.GetFileNameWithoutAnyExtensions().ToLowerInvariant();
-			if (string.IsNullOrEmpty(BucketName))
+			string bucketName = _parameters.Bucket;
+			string projectNameAsBucketName = _projectFile.GetFileNameWithoutAnyExtensions().ToLowerInvariant();
+			if (string.IsNullOrEmpty(bucketName))
 			{
-				BucketName = ProjectNameAsBucketName;
+				bucketName = projectNameAsBucketName;
 			}
 
-			string HostURLArg = string.Format("--hosturl http://{0}:{1}", Parameters.HostName, Parameters.HostPort);
-			StringBuilder OplogImportCommandline = new StringBuilder();
-			OplogImportCommandline.AppendFormat("oplog-import {0} --cloud {1} --namespace {2} --bucket {3}", HostURLArg, Parameters.CloudURL, Parameters.Namespace, BucketName);
-			OplogImportCommandline.AppendFormat(" {0}", Parameters.Key);
+			string hostUrlArg = string.Format("--hosturl http://{0}:{1}", _parameters.HostName, _parameters.HostPort);
+			StringBuilder oplogImportCommandline = new StringBuilder();
+			oplogImportCommandline.AppendFormat("oplog-import {0} --cloud {1} --namespace {2} --bucket {3}", hostUrlArg, _parameters.CloudURL, _parameters.Namespace, bucketName);
+			oplogImportCommandline.AppendFormat(" {0}", _parameters.Key);
 
-			Logger.LogInformation("Running '{Arg0} {Arg1}'", CommandUtils.MakePathSafeToUseWithCommandLine(ZenExe.FullName), OplogImportCommandline.ToString());
-			CommandUtils.RunAndLog(CommandUtils.CmdEnv, ZenExe.FullName, OplogImportCommandline.ToString(), Options: CommandUtils.ERunOptions.Default);
+			Logger.LogInformation("Running '{Arg0} {Arg1}'", CommandUtils.MakePathSafeToUseWithCommandLine(zenExe.FullName), oplogImportCommandline.ToString());
+			CommandUtils.RunAndLog(CommandUtils.CmdEnv, zenExe.FullName, oplogImportCommandline.ToString(), Options: CommandUtils.ERunOptions.Default);
 		}
 
 		/// <summary>
 		/// Output this task out to an XML writer.
 		/// </summary>
-		public override void Write(XmlWriter Writer)
+		public override void Write(XmlWriter writer)
 		{
-			Write(Writer, Parameters);
+			Write(writer, _parameters);
 		}
 
 		/// <summary>
diff --git a/Engine/Source/Programs/AutomationTool/BuildGraph/Tasks/ZipTask.cs b/Engine/Source/Programs/AutomationTool/BuildGraph/Tasks/ZipTask.cs
--- a/Engine/Source/Programs/AutomationTool/BuildGraph/Tasks/ZipTask.cs
+++ b/Engine/Source/Programs/AutomationTool/BuildGraph/Tasks/ZipTask.cs
@@ -1,18 +1,11 @@
 // Copyright Epic Games, Inc. All Rights Reserved.
 
-using AutomationTool;
-using EpicGames.BuildGraph;
-using System;
 using System.Collections.Generic;
-using System.IO.Compression;
 using System.Linq;
-using System.Text;
 using System.Threading.Tasks;
 using System.Xml;
 using EpicGames.Core;
-using UnrealBuildTool;
 using Microsoft.Extensions.Logging;
-using System.IO;
 
 namespace AutomationTool.Tasks
 {
@@ -61,68 +54,68 @@ namespace AutomationTool.Tasks
 		/// <summary>
 		/// Parameters for this task
 		/// </summary>
-		ZipTaskParameters Parameters;
+		readonly ZipTaskParameters _parameters;
 
 		/// <summary>
 		/// Constructor
 		/// </summary>
-		/// <param name="InParameters">Parameters for this task</param>
-		public ZipTask(ZipTaskParameters InParameters)
+		/// <param name="parameters">Parameters for this task</param>
+		public ZipTask(ZipTaskParameters parameters)
 		{
-			Parameters = InParameters;
+			_parameters = parameters;
 		}
 
 		/// <summary>
-		/// Execute the task.
+		/// ExecuteAsync the task.
 		/// </summary>
-		/// <param name="Job">Information about the current job</param>
-		/// <param name="BuildProducts">Set of build products produced by this node.</param>
-		/// <param name="TagNameToFileSet">Mapping from tag names to the set of files they include</param>
-		public override Task ExecuteAsync(JobContext Job, HashSet<FileReference> BuildProducts, Dictionary<string, HashSet<FileReference>> TagNameToFileSet)
+		/// <param name="job">Information about the current job</param>
+		/// <param name="buildProducts">Set of build products produced by this node.</param>
+		/// <param name="tagNameToFileSet">Mapping from tag names to the set of files they include</param>
+		public override Task ExecuteAsync(JobContext job, HashSet<FileReference> buildProducts, Dictionary<string, HashSet<FileReference>> tagNameToFileSet)
 		{
 			// Find all the input files
-			List<FileReference> Files;
-			if(Parameters.Files == null)
+			List<FileReference> files;
+			if (_parameters.Files == null)
 			{
-				Files = DirectoryReference.EnumerateFiles(Parameters.FromDir, "*", System.IO.SearchOption.AllDirectories).ToList();
+				files = DirectoryReference.EnumerateFiles(_parameters.FromDir, "*", System.IO.SearchOption.AllDirectories).ToList();
 			}
 			else
 			{
-				Files = ResolveFilespec(Parameters.FromDir, Parameters.Files, TagNameToFileSet).ToList();
+				files = ResolveFilespec(_parameters.FromDir, _parameters.Files, tagNameToFileSet).ToList();
 			}
 
 			// Create the zip file
-			Logger.LogInformation("Adding {NumFiles} files to {ZipFile}...", Files.Count, Parameters.ZipFile);
+			Logger.LogInformation("Adding {NumFiles} files to {ZipFile}...", files.Count, _parameters.ZipFile);
 
-			HashSet<FileReference> ExecutableFiles = null;
-			if (Parameters.ExecutableFiles != null)
+			HashSet<FileReference> executableFiles = null;
+			if (_parameters.ExecutableFiles != null)
 			{
-				ExecutableFiles = ResolveFilespec(Parameters.FromDir, Parameters.ExecutableFiles, TagNameToFileSet);
-				foreach (FileReference ExecutableFile in Files.Intersect(ExecutableFiles))
+				executableFiles = ResolveFilespec(_parameters.FromDir, _parameters.ExecutableFiles, tagNameToFileSet);
+				foreach (FileReference executableFile in files.Intersect(executableFiles))
 				{
-					Logger.LogInformation("  Executable file: {File}", ExecutableFile);
+					Logger.LogInformation("  Executable file: {File}", executableFile);
 				}
 			}
 
-			CommandUtils.ZipFiles(Parameters.ZipFile, Parameters.FromDir, Files, ExecutableFiles);
+			CommandUtils.ZipFiles(_parameters.ZipFile, _parameters.FromDir, files, executableFiles);
 
 			// Apply the optional tag to the produced archive
-			foreach (string TagName in FindTagNamesFromList(Parameters.Tag))
+			foreach (string tagName in FindTagNamesFromList(_parameters.Tag))
 			{
-				FindOrAddTagSet(TagNameToFileSet, TagName).Add(Parameters.ZipFile);
+				FindOrAddTagSet(tagNameToFileSet, tagName).Add(_parameters.ZipFile);
 			}
 
 			// Add the archive to the set of build products
-			BuildProducts.Add(Parameters.ZipFile);
+			buildProducts.Add(_parameters.ZipFile);
 			return Task.CompletedTask;
 		}
 
 		/// <summary>
 		/// Output this task out to an XML writer.
 		/// </summary>
-		public override void Write(XmlWriter Writer)
+		public override void Write(XmlWriter writer)
 		{
-			Write(Writer, Parameters);
+			Write(writer, _parameters);
 		}
 
 		/// <summary>
@@ -131,7 +124,7 @@ namespace AutomationTool.Tasks
 		/// <returns>The tag names which are read by this task</returns>
 		public override IEnumerable<string> FindConsumedTagNames()
 		{
-			return FindTagNamesFromFilespec(Parameters.Files);
+			return FindTagNamesFromFilespec(_parameters.Files);
 		}
 
 		/// <summary>
@@ -140,7 +133,7 @@ namespace AutomationTool.Tasks
 		/// <returns>The tag names which are modified by this task</returns>
 		public override IEnumerable<string> FindProducedTagNames()
 		{
-			return FindTagNamesFromList(Parameters.Tag);
+			return FindTagNamesFromList(_parameters.Tag);
 		}
 	}
 }
