EpicGames.BuildGraph: Apply standard naming conventions to BuildGraph.Automation.

diff --git a/Engine/Source/Programs/AutomationTool/BuildGraph/TempStorage.cs b/Engine/Source/Programs/AutomationTool/BuildGraph/TempStorage.cs
--- a/Engine/Source/Programs/AutomationTool/BuildGraph/TempStorage.cs
+++ b/Engine/Source/Programs/AutomationTool/BuildGraph/TempStorage.cs
@@ -1,12 +1,10 @@
 // Copyright Epic Games, Inc. All Rights Reserved.
 
-using EpicGames.Core;
-using OpenTracing;
-using OpenTracing.Util;
 using System;
 using System.Collections.Concurrent;
 using System.Collections.Generic;
 using System.Diagnostics;
+using System.Diagnostics.CodeAnalysis;
 using System.IO;
 using System.IO.Compression;
 using System.Linq;
@@ -15,12 +13,12 @@ using System.Threading;
 using System.Threading.Tasks;
 using System.Xml;
 using System.Xml.Serialization;
-using UnrealBuildTool;
+using EpicGames.Core;
 using Microsoft.Extensions.Logging;
-using System.Text;
-
+using OpenTracing;
+using OpenTracing.Util;
+using UnrealBuildTool;
 using static AutomationTool.CommandUtils;
-using System.Diagnostics.CodeAnalysis;
 
 #nullable enable
 
@@ -55,23 +53,23 @@ namespace AutomationTool
 		/// <summary>
 		/// Construct a temp storage block
 		/// </summary>
-		/// <param name="InNodeName">Name of the node</param>
-		/// <param name="InOutputName">Name of the node's output</param>
-		public TempStorageBlock(string InNodeName, string InOutputName)
+		/// <param name="nodeName">Name of the node</param>
+		/// <param name="outputName">Name of the node's output</param>
+		public TempStorageBlock(string nodeName, string outputName)
 		{
-			NodeName = InNodeName;
-			OutputName = InOutputName;
+			NodeName = nodeName;
+			OutputName = outputName;
 		}
 
 		/// <summary>
 		/// Tests whether two temp storage blocks are equal
 		/// </summary>
-		/// <param name="Other">The object to compare against</param>
+		/// <param name="other">The object to compare against</param>
 		/// <returns>True if the blocks are equivalent</returns>
-		public override bool Equals(object? Other)
+		public override bool Equals(object? other)
 		{
-			TempStorageBlock? OtherBlock = Other as TempStorageBlock;
-			return OtherBlock != null && NodeName == OtherBlock.NodeName && OutputName == OtherBlock.OutputName;
+			TempStorageBlock? otherBlock = other as TempStorageBlock;
+			return otherBlock != null && NodeName == otherBlock.NodeName && OutputName == otherBlock.OutputName;
 		}
 
 		/// <summary>
@@ -134,52 +132,52 @@ namespace AutomationTool
 		/// <summary>
 		/// Constructor
 		/// </summary>
-		/// <param name="FileInfo">File to be added</param>
-		/// <param name="RootDir">Root directory to store paths relative to</param>
-		public TempStorageFile(FileInfo FileInfo, DirectoryReference RootDir)
+		/// <param name="fileInfo">File to be added</param>
+		/// <param name="rootDir">Root directory to store paths relative to</param>
+		public TempStorageFile(FileInfo fileInfo, DirectoryReference rootDir)
 		{
 			// Check the file exists and is in the right location
-			FileReference File = new FileReference(FileInfo);
-			if(!File.IsUnderDirectory(RootDir))
+			FileReference file = new FileReference(fileInfo);
+			if (!file.IsUnderDirectory(rootDir))
 			{
-				throw new AutomationException("Attempt to add file to temp storage manifest that is outside the root directory ({0})", File.FullName);
+				throw new AutomationException("Attempt to add file to temp storage manifest that is outside the root directory ({0})", file.FullName);
 			}
-			if(!FileInfo.Exists)
+			if (!fileInfo.Exists)
 			{
-				throw new AutomationException("Attempt to add file to temp storage manifest that does not exist ({0})", File.FullName);
+				throw new AutomationException("Attempt to add file to temp storage manifest that does not exist ({0})", file.FullName);
 			}
 
-			RelativePath = File.MakeRelativeTo(RootDir).Replace(Path.DirectorySeparatorChar, '/');
-			LastWriteTimeUtcTicks = FileInfo.LastWriteTimeUtc.Ticks;
-			Length = FileInfo.Length;
+			RelativePath = file.MakeRelativeTo(rootDir).Replace(Path.DirectorySeparatorChar, '/');
+			LastWriteTimeUtcTicks = fileInfo.LastWriteTimeUtc.Ticks;
+			Length = fileInfo.Length;
 
 			if (GenerateDigest())
 			{
-				Digest = ContentHash.MD5(File).ToString();
+				Digest = ContentHash.MD5(file).ToString();
 			}
 		}
 
 		/// <summary>
 		/// Compare stored for this file with the one on disk, and output an error if they differ.
 		/// </summary>
-		/// <param name="RootDir">Root directory for this branch</param>
+		/// <param name="rootDir">Root directory for this branch</param>
 		/// <returns>True if the files are identical, false otherwise</returns>
-		public bool Compare(DirectoryReference RootDir)
+		public bool Compare(DirectoryReference rootDir)
 		{
-			string? Message;
-			if(Compare(RootDir, out Message))
+			string? message;
+			if (Compare(rootDir, out message))
 			{
-				if(Message != null)
+				if (message != null)
 				{
-					Logger.LogInformation("{Text}", Message);
+					Logger.LogInformation("{Text}", message);
 				}
 				return true;
 			}
 			else
 			{
-				if(Message != null)
+				if (message != null)
 				{
-					Logger.LogError("{Text}", Message);
+					Logger.LogError("{Text}", message);
 				}
 				return false;
 			}
@@ -188,70 +186,70 @@ namespace AutomationTool
 		/// <summary>
 		/// Compare stored for this file with the one on disk, and output an error if they differ.
 		/// </summary>
-		/// <param name="RootDir">Root directory for this branch</param>
-		/// <param name="Message">Message describing the difference</param>
+		/// <param name="rootDir">Root directory for this branch</param>
+		/// <param name="message">Message describing the difference</param>
 		/// <returns>True if the files are identical, false otherwise</returns>
-		public bool Compare(DirectoryReference RootDir, [NotNullWhen(false)] out string? Message)
+		public bool Compare(DirectoryReference rootDir, [NotNullWhen(false)] out string? message)
 		{
-			FileReference LocalFile = ToFileReference(RootDir);
+			FileReference localFile = ToFileReference(rootDir);
 
 			// Get the local file info, and check it exists
-			FileInfo Info = new FileInfo(LocalFile.FullName);
-			if(!Info.Exists)
+			FileInfo info = new FileInfo(localFile.FullName);
+			if (!info.Exists)
 			{
-				Message = String.Format("Missing file from manifest - {0}", RelativePath);
+				message = String.Format("Missing file from manifest - {0}", RelativePath);
 				return false;
 			}
 
 			// Check the size matches
-			if(Info.Length != Length)
+			if (info.Length != Length)
 			{
-				if(TempStorage.IsDuplicateBuildProduct(LocalFile))
+				if (TempStorage.IsDuplicateBuildProduct(localFile))
 				{
-					Message = String.Format("Ignored file size mismatch for {0} - was {1} bytes, expected {2} bytes", RelativePath, Info.Length, Length);
+					message = String.Format("Ignored file size mismatch for {0} - was {1} bytes, expected {2} bytes", RelativePath, info.Length, Length);
 					return true;
 				}
 				else
 				{
-					Message = String.Format("File size differs from manifest - {0} is {1} bytes, expected {2} bytes", RelativePath, Info.Length, Length);
+					message = String.Format("File size differs from manifest - {0} is {1} bytes, expected {2} bytes", RelativePath, info.Length, Length);
 					return false;
 				}
 			}
 
 			// Check the timestamp of the file matches. On FAT filesystems writetime has a two seconds resolution (see http://msdn.microsoft.com/en-us/library/windows/desktop/ms724290%28v=vs.85%29.aspx)
-			TimeSpan TimeDifference = new TimeSpan(Info.LastWriteTimeUtc.Ticks - LastWriteTimeUtcTicks);
-			if (TimeDifference.TotalSeconds >= -2 && TimeDifference.TotalSeconds <= +2)
+			TimeSpan timeDifference = new TimeSpan(info.LastWriteTimeUtc.Ticks - LastWriteTimeUtcTicks);
+			if (timeDifference.TotalSeconds >= -2 && timeDifference.TotalSeconds <= +2)
 			{
-				Message = null;
+				message = null;
 				return true;
 			}
 
 			// Check if the files have been modified
-			DateTime ExpectedLocal = new DateTime(LastWriteTimeUtcTicks, DateTimeKind.Utc).ToLocalTime();
+			DateTime expectedLocal = new DateTime(LastWriteTimeUtcTicks, DateTimeKind.Utc).ToLocalTime();
 			if (Digest != null)
 			{
-				string LocalDigest = ContentHash.MD5(LocalFile).ToString();
-				if (Digest.Equals(LocalDigest, StringComparison.Ordinal))
+				string localDigest = ContentHash.MD5(localFile).ToString();
+				if (Digest.Equals(localDigest, StringComparison.Ordinal))
 				{
-					Message = null;
+					message = null;
 					return true;
 				}
 				else
 				{
-					Message = String.Format("Digest mismatch for {0} - was {1} ({2}), expected {3} ({4}), TimeDifference {5}", RelativePath, LocalDigest, Info.LastWriteTime, Digest, ExpectedLocal, TimeDifference);
+					message = String.Format("Digest mismatch for {0} - was {1} ({2}), expected {3} ({4}), TimeDifference {5}", RelativePath, localDigest, info.LastWriteTime, Digest, expectedLocal, timeDifference);
 					return false;
 				}
 			}
 			else
 			{
-				if (RequireMatchingTimestamps() && !TempStorage.IsDuplicateBuildProduct(LocalFile))
+				if (RequireMatchingTimestamps() && !TempStorage.IsDuplicateBuildProduct(localFile))
 				{
-					Message = String.Format("File date/time mismatch for {0} - was {1}, expected {2}, TimeDifference {3}", RelativePath, Info.LastWriteTime, ExpectedLocal, TimeDifference);
+					message = String.Format("File date/time mismatch for {0} - was {1}, expected {2}, TimeDifference {3}", RelativePath, info.LastWriteTime, expectedLocal, timeDifference);
 					return false;
 				}
 				else
 				{
-					Message = String.Format("Ignored file date/time mismatch for {0} - was {1}, expected {2}, TimeDifference {3}", RelativePath, Info.LastWriteTime, ExpectedLocal, TimeDifference);
+					message = String.Format("Ignored file date/time mismatch for {0} - was {1}, expected {2}, TimeDifference {3}", RelativePath, info.LastWriteTime, expectedLocal, timeDifference);
 					return true;
 				}
 			}
@@ -287,11 +285,11 @@ namespace AutomationTool
 		/// <summary>
 		/// Gets a local file reference for this file, given a root directory to base it from.
 		/// </summary>
-		/// <param name="RootDir">The local root directory</param>
+		/// <param name="rootDir">The local root directory</param>
 		/// <returns>Reference to the file</returns>
-		public FileReference ToFileReference(DirectoryReference RootDir)
+		public FileReference ToFileReference(DirectoryReference rootDir)
 		{
-			return FileReference.Combine(RootDir, RelativePath.Replace('/', Path.DirectorySeparatorChar));
+			return FileReference.Combine(rootDir, RelativePath.Replace('/', Path.DirectorySeparatorChar));
 		}
 	}
 
@@ -324,11 +322,11 @@ namespace AutomationTool
 		/// <summary>
 		/// Constructor
 		/// </summary>
-		/// <param name="Info">FileInfo for the zip file</param>
-		public TempStorageZipFile(FileInfo Info)
+		/// <param name="info">FileInfo for the zip file</param>
+		public TempStorageZipFile(FileInfo info)
 		{
-			Name = Info.Name;
-			Length = Info.Length;
+			Name = info.Name;
+			Length = info.Length;
 		}
 	}
 
@@ -354,7 +352,7 @@ namespace AutomationTool
 		/// <summary>
 		/// Construct a static Xml serializer to avoid throwing an exception searching for the reflection info at runtime
 		/// </summary>
-		static XmlSerializer Serializer = XmlSerializer.FromTypes(new Type[]{ typeof(TempStorageManifest) })[0]!;
+		static readonly XmlSerializer s_serializer = XmlSerializer.FromTypes(new Type[] { typeof(TempStorageManifest) })[0]!;
 
 		/// <summary>
 		/// Construct an empty temp storage manifest
@@ -368,11 +366,11 @@ namespace AutomationTool
 		/// <summary>
 		/// Creates a manifest from a flat list of files (in many folders) and a BaseFolder from which they are rooted.
 		/// </summary>
-		/// <param name="InFiles">List of full file paths</param>
-		/// <param name="RootDir">Root folder for all the files. All files must be relative to this RootDir.</param>
-		public TempStorageManifest(FileInfo[] InFiles, DirectoryReference RootDir)
+		/// <param name="files">List of full file paths</param>
+		/// <param name="rootDir">Root folder for all the files. All files must be relative to this RootDir.</param>
+		public TempStorageManifest(FileInfo[] files, DirectoryReference rootDir)
 		{
-			Files = InFiles.Select(x => new TempStorageFile(x, RootDir)).ToArray();
+			Files = files.Select(x => new TempStorageFile(x, rootDir)).ToArray();
 			ZipFiles = Array.Empty<TempStorageZipFile>();
 		}
 
@@ -382,12 +380,12 @@ namespace AutomationTool
 		/// <returns>The total size of all files</returns>
 		public long GetTotalSize()
 		{
-			long Result = 0;
-			foreach(TempStorageFile File in Files)
+			long result = 0;
+			foreach (TempStorageFile file in Files)
 			{
-				Result += File.Length;
+				result += file.Length;
 			}
-			return Result;
+			return result;
 		}
 
 		/// <summary>
@@ -396,24 +394,24 @@ namespace AutomationTool
 		/// <param name="file">File to load</param>
 		static public TempStorageManifest Load(FileReference file)
 		{
-			using (StreamReader Reader = new(file.FullName))
+			using (StreamReader reader = new(file.FullName))
 			{
-				return (TempStorageManifest)(Serializer.Deserialize(Reader) ?? throw new InvalidOperationException());
+				return (TempStorageManifest)(s_serializer.Deserialize(reader) ?? throw new InvalidOperationException());
 			}
 		}
 
 		/// <summary>
 		/// Saves a manifest to disk
 		/// </summary>
-		/// <param name="File">File to save</param>
-		public void Save(FileReference File)
+		/// <param name="file">File to save</param>
+		public void Save(FileReference file)
 		{
-			using(StreamWriter Writer = new StreamWriter(File.FullName))
+			using (StreamWriter writer = new StreamWriter(file.FullName))
 			{
-				XmlWriterSettings WriterSettings = new() { Indent = true };
-				using (XmlWriter XMLWriter = XmlWriter.Create(Writer, WriterSettings))
+				XmlWriterSettings writerSettings = new() { Indent = true };
+				using (XmlWriter xmlWriter = XmlWriter.Create(writer, writerSettings))
 				{
-					Serializer.Serialize(XMLWriter, this);
+					s_serializer.Serialize(xmlWriter, this);
 				}
 			}
 		}
@@ -456,7 +454,7 @@ namespace AutomationTool
 		/// <summary>
 		/// Construct a static Xml serializer to avoid throwing an exception searching for the reflection info at runtime
 		/// </summary>
-		static XmlSerializer Serializer = XmlSerializer.FromTypes(new Type[]{ typeof(TempStorageTagManifest) })[0]!;
+		static readonly XmlSerializer s_serializer = XmlSerializer.FromTypes(new Type[] { typeof(TempStorageTagManifest) })[0]!;
 
 		/// <summary>
 		/// Construct an empty file list for deserialization
@@ -472,56 +470,56 @@ namespace AutomationTool
 		/// <summary>
 		/// Creates a manifest from a flat list of files (in many folders) and a BaseFolder from which they are rooted.
 		/// </summary>
-		/// <param name="InFiles">List of full file paths</param>
-		/// <param name="RootDir">Root folder for all the files. All files must be relative to this RootDir.</param>
-		/// <param name="InBlocks">Referenced storage blocks required for these files</param>
-		/// <param name="InArtifactKeys">Keys for published artifacts</param>
-		public TempStorageTagManifest(IEnumerable<FileReference> InFiles, DirectoryReference RootDir, IEnumerable<TempStorageBlock> InBlocks, IEnumerable<string> InArtifactKeys)
+		/// <param name="files">List of full file paths</param>
+		/// <param name="rootDir">Root folder for all the files. All files must be relative to this RootDir.</param>
+		/// <param name="blocks">Referenced storage blocks required for these files</param>
+		/// <param name="artifactKeys">Keys for published artifacts</param>
+		public TempStorageTagManifest(IEnumerable<FileReference> files, DirectoryReference rootDir, IEnumerable<TempStorageBlock> blocks, IEnumerable<string> artifactKeys)
 		{
-			List<string> NewLocalFiles = new List<string>();
-			List<string> NewExternalFiles = new List<string>();
-			foreach(FileReference File in InFiles)
+			List<string> newLocalFiles = new List<string>();
+			List<string> newExternalFiles = new List<string>();
+			foreach (FileReference file in files)
 			{
-				if(File.IsUnderDirectory(RootDir))
+				if (file.IsUnderDirectory(rootDir))
 				{
-					NewLocalFiles.Add(File.MakeRelativeTo(RootDir).Replace(Path.DirectorySeparatorChar, '/'));
+					newLocalFiles.Add(file.MakeRelativeTo(rootDir).Replace(Path.DirectorySeparatorChar, '/'));
 				}
 				else
 				{
-					NewExternalFiles.Add(File.FullName.Replace(Path.DirectorySeparatorChar, '/'));
+					newExternalFiles.Add(file.FullName.Replace(Path.DirectorySeparatorChar, '/'));
 				}
 			}
-			LocalFiles = NewLocalFiles.ToArray();
-			ExternalFiles = NewExternalFiles.ToArray();
+			LocalFiles = newLocalFiles.ToArray();
+			ExternalFiles = newExternalFiles.ToArray();
 
-			Blocks = InBlocks.ToArray();
-			ArtifactKeys = InArtifactKeys.ToArray();
+			Blocks = blocks.ToArray();
+			ArtifactKeys = artifactKeys.ToArray();
 		}
 
 		/// <summary>
 		/// Load this list of files from disk
 		/// </summary>
-		/// <param name="File">File to load</param>
-		static public TempStorageTagManifest Load(FileReference File)
+		/// <param name="file">File to load</param>
+		static public TempStorageTagManifest Load(FileReference file)
 		{
-			using(StreamReader Reader = new StreamReader(File.FullName))
+			using (StreamReader reader = new StreamReader(file.FullName))
 			{
-				return (TempStorageTagManifest)Serializer.Deserialize(Reader)!;
+				return (TempStorageTagManifest)s_serializer.Deserialize(reader)!;
 			}
 		}
 
 		/// <summary>
 		/// Saves this list of files to disk
 		/// </summary>
-		/// <param name="File">File to save</param>
-		public void Save(FileReference File)
+		/// <param name="file">File to save</param>
+		public void Save(FileReference file)
 		{
-			using (StreamWriter Writer = new StreamWriter(File.FullName))
+			using (StreamWriter writer = new StreamWriter(file.FullName))
 			{
-				XmlWriterSettings WriterSettings = new() { Indent = true };
-				using (XmlWriter XMLWriter = XmlWriter.Create(Writer, WriterSettings))
+				XmlWriterSettings writerSettings = new() { Indent = true };
+				using (XmlWriter xmlWriter = XmlWriter.Create(writer, writerSettings))
 				{
-					Serializer.Serialize(XMLWriter, this);
+					s_serializer.Serialize(xmlWriter, this);
 				}
 			}
 		}
@@ -529,14 +527,14 @@ namespace AutomationTool
 		/// <summary>
 		/// Converts this file list into a set of FileReference objects
 		/// </summary>
-		/// <param name="RootDir">The root directory to rebase local files</param>
+		/// <param name="rootDir">The root directory to rebase local files</param>
 		/// <returns>Set of files</returns>
-		public HashSet<FileReference> ToFileSet(DirectoryReference RootDir)
+		public HashSet<FileReference> ToFileSet(DirectoryReference rootDir)
 		{
-			HashSet<FileReference> Files = new HashSet<FileReference>();
-			Files.UnionWith(LocalFiles.Select(x => FileReference.Combine(RootDir, x)));
-			Files.UnionWith(ExternalFiles.Select(x => new FileReference(x)));
-			return Files;
+			HashSet<FileReference> files = new HashSet<FileReference>();
+			files.UnionWith(LocalFiles.Select(x => FileReference.Combine(rootDir, x)));
+			files.UnionWith(ExternalFiles.Select(x => new FileReference(x)));
+			return files;
 		}
 	}
 
@@ -555,36 +553,36 @@ namespace AutomationTool
 		/// <summary>
 		/// Root directory for this branch.
 		/// </summary>
-		DirectoryReference RootDir;
+		readonly DirectoryReference _rootDir;
 
 		/// <summary>
 		/// The local temp storage directory (typically somewhere under /Engine/Saved directory).
 		/// </summary>
-		DirectoryReference LocalDir;
+		readonly DirectoryReference _localDir;
 
 		/// <summary>
 		/// The shared temp storage directory; typically a network location. May be null.
 		/// </summary>
-		DirectoryReference? SharedDir;
+		readonly DirectoryReference? _sharedDir;
 
 		/// <summary>
 		/// Whether to allow writes to shared storage
 		/// </summary>
-		bool bWriteToSharedStorage;
+		readonly bool _writeToSharedStorage;
 
 		/// <summary>
 		/// Constructor
 		/// </summary>
-		/// <param name="InRootDir">Root directory for this branch</param>
-		/// <param name="InLocalDir">The local temp storage directory.</param>
-		/// <param name="InSharedDir">The shared temp storage directory. May be null.</param>
-		/// <param name="bInWriteToSharedStorage">Whether to write to shared storage, or only permit reads from it</param>
-		public TempStorage(DirectoryReference InRootDir, DirectoryReference InLocalDir, DirectoryReference? InSharedDir, bool bInWriteToSharedStorage)
+		/// <param name="rootDir">Root directory for this branch</param>
+		/// <param name="localDir">The local temp storage directory.</param>
+		/// <param name="sharedDir">The shared temp storage directory. May be null.</param>
+		/// <param name="writeToSharedStorage">Whether to write to shared storage, or only permit reads from it</param>
+		public TempStorage(DirectoryReference rootDir, DirectoryReference localDir, DirectoryReference? sharedDir, bool writeToSharedStorage)
 		{
-			RootDir = InRootDir;
-			LocalDir = InLocalDir;
-			SharedDir = InSharedDir;
-			bWriteToSharedStorage = bInWriteToSharedStorage;
+			_rootDir = rootDir;
+			_localDir = localDir;
+			_sharedDir = sharedDir;
+			_writeToSharedStorage = writeToSharedStorage;
 		}
 
 		/// <summary>
@@ -592,42 +590,42 @@ namespace AutomationTool
 		/// </summary>
 		public void CleanLocal()
 		{
-			CommandUtils.DeleteDirectoryContents(LocalDir.FullName);
+			CommandUtils.DeleteDirectoryContents(_localDir.FullName);
 		}
 
 		/// <summary>
 		/// Cleans local build products for a given node. Does not modify shared storage.
 		/// </summary>
-		/// <param name="NodeName">Name of the node</param>
-		public void CleanLocalNode(string NodeName)
+		/// <param name="nodeName">Name of the node</param>
+		public void CleanLocalNode(string nodeName)
 		{
-			DirectoryReference NodeDir = GetDirectoryForNode(LocalDir, NodeName);
-			if(DirectoryReference.Exists(NodeDir))
+			DirectoryReference nodeDir = GetDirectoryForNode(_localDir, nodeName);
+			if (DirectoryReference.Exists(nodeDir))
 			{
-				CommandUtils.DeleteDirectoryContents(NodeDir.FullName);
-				CommandUtils.DeleteDirectory_NoExceptions(NodeDir.FullName);
+				CommandUtils.DeleteDirectoryContents(nodeDir.FullName);
+				CommandUtils.DeleteDirectory_NoExceptions(nodeDir.FullName);
 			}
 		}
 
 		/// <summary>
 		/// Check whether the given node is complete
 		/// </summary>
-		/// <param name="NodeName">Name of the node</param>
+		/// <param name="nodeName">Name of the node</param>
 		/// <returns>True if the node is complete</returns>
-		public bool IsComplete(string NodeName)
+		public bool IsComplete(string nodeName)
 		{
 			// Check if it already exists locally
-			FileReference LocalFile = GetCompleteMarkerFile(LocalDir, NodeName);
-			if(FileReference.Exists(LocalFile))
+			FileReference localFile = GetCompleteMarkerFile(_localDir, nodeName);
+			if (FileReference.Exists(localFile))
 			{
 				return true;
 			}
-			
+
 			// Check if it exists in shared storage
-			if(SharedDir != null)
+			if (_sharedDir != null)
 			{
-				FileReference SharedFile = GetCompleteMarkerFile(SharedDir, NodeName);
-				if(FileReference.Exists(SharedFile))
+				FileReference sharedFile = GetCompleteMarkerFile(_sharedDir, nodeName);
+				if (FileReference.Exists(sharedFile))
 				{
 					return true;
 				}
@@ -640,108 +638,108 @@ namespace AutomationTool
 		/// <summary>
 		/// Mark the given node as complete
 		/// </summary>
-		/// <param name="NodeName">Name of the node</param>
-		public void MarkAsComplete(string NodeName)
+		/// <param name="nodeName">Name of the node</param>
+		public void MarkAsComplete(string nodeName)
 		{
 			// Create the marker locally
-			FileReference LocalFile = GetCompleteMarkerFile(LocalDir, NodeName);
-			DirectoryReference.CreateDirectory(LocalFile.Directory);
-			File.OpenWrite(LocalFile.FullName).Close();
+			FileReference localFile = GetCompleteMarkerFile(_localDir, nodeName);
+			DirectoryReference.CreateDirectory(localFile.Directory);
+			File.OpenWrite(localFile.FullName).Close();
 
 			// Create the marker in the shared directory
-			if(SharedDir != null && bWriteToSharedStorage)
+			if (_sharedDir != null && _writeToSharedStorage)
 			{
-				FileReference SharedFile = GetCompleteMarkerFile(SharedDir, NodeName);
-				DirectoryReference.CreateDirectory(SharedFile.Directory);
-				File.OpenWrite(SharedFile.FullName).Close();
+				FileReference sharedFile = GetCompleteMarkerFile(_sharedDir, nodeName);
+				DirectoryReference.CreateDirectory(sharedFile.Directory);
+				File.OpenWrite(sharedFile.FullName).Close();
 			}
 		}
 
 		/// <summary>
 		/// Checks the integrity of the give node's local build products.
 		/// </summary>
-		/// <param name="NodeName">The node to retrieve build products for</param>
-		/// <param name="TagNames">List of tag names from this node.</param>
-		/// <param name="IgnoreModified">Filter for files to ignore when performing integrity check. Specified per node. </param>
+		/// <param name="nodeName">The node to retrieve build products for</param>
+		/// <param name="tagNames">List of tag names from this node.</param>
+		/// <param name="ignoreModified">Filter for files to ignore when performing integrity check. Specified per node. </param>
 		/// <returns>True if the node is complete and valid, false if not (and typically followed by a call to CleanNode()).</returns>
-		public bool CheckLocalIntegrity(string NodeName, IEnumerable<string> TagNames, FileFilter IgnoreModified)
+		public bool CheckLocalIntegrity(string nodeName, IEnumerable<string> tagNames, FileFilter ignoreModified)
 		{
 			// If the node is not locally complete, fail immediately.
-			FileReference CompleteMarkerFile = GetCompleteMarkerFile(LocalDir, NodeName);
-			if(!FileReference.Exists(CompleteMarkerFile))
+			FileReference completeMarkerFile = GetCompleteMarkerFile(_localDir, nodeName);
+			if (!FileReference.Exists(completeMarkerFile))
 			{
 				return false;
 			}
 
 			// Check that each of the tags exist
-			HashSet<TempStorageBlock> Blocks = new HashSet<TempStorageBlock>();
-			foreach(string TagName in TagNames)
+			HashSet<TempStorageBlock> blocks = new HashSet<TempStorageBlock>();
+			foreach (string tagName in tagNames)
 			{
 				// Check the local manifest exists
-				FileReference LocalFileListLocation = GetTaggedFileListLocation(LocalDir, NodeName, TagName);
-				if(!FileReference.Exists(LocalFileListLocation))
+				FileReference localFileListLocation = GetTaggedFileListLocation(_localDir, nodeName, tagName);
+				if (!FileReference.Exists(localFileListLocation))
 				{
 					return false;
 				}
 
 				// Check the local manifest matches the shared manifest
-				if(SharedDir != null)
+				if (_sharedDir != null)
 				{
 					// Check the shared manifest exists
-					FileReference SharedFileListLocation = GetManifestLocation(SharedDir, NodeName, TagName);
-					if(!FileReference.Exists(SharedFileListLocation))
+					FileReference sharedFileListLocation = GetManifestLocation(_sharedDir, nodeName, tagName);
+					if (!FileReference.Exists(sharedFileListLocation))
 					{
 						return false;
 					}
 
 					// Check the manifests are identical, byte by byte
-					byte[] LocalManifestBytes = File.ReadAllBytes(LocalFileListLocation.FullName);
-					byte[] SharedManifestBytes = Array.Empty<byte>();
-					PerformActionWithRetries(() => SharedManifestBytes = File.ReadAllBytes(SharedFileListLocation.FullName), 3, TimeSpan.FromSeconds(1));
-					if (!LocalManifestBytes.SequenceEqual(SharedManifestBytes))
+					byte[] localManifestBytes = File.ReadAllBytes(localFileListLocation.FullName);
+					byte[] sharedManifestBytes = Array.Empty<byte>();
+					PerformActionWithRetries(() => sharedManifestBytes = File.ReadAllBytes(sharedFileListLocation.FullName), 3, TimeSpan.FromSeconds(1));
+					if (!localManifestBytes.SequenceEqual(sharedManifestBytes))
 					{
 						return false;
 					}
 				}
 
 				// Read the manifest and add the referenced blocks to be checked
-				TempStorageTagManifest LocalFileList = TempStorageTagManifest.Load(LocalFileListLocation);
-				Blocks.UnionWith(LocalFileList.Blocks);
+				TempStorageTagManifest localFileList = TempStorageTagManifest.Load(localFileListLocation);
+				blocks.UnionWith(localFileList.Blocks);
 			}
 
 			// Check that each of the outputs match
-			foreach(TempStorageBlock Block in Blocks)
+			foreach (TempStorageBlock block in blocks)
 			{
 				// Check the local manifest exists
-				FileReference LocalManifestFile = GetManifestLocation(LocalDir, Block.NodeName, Block.OutputName);
-				if(!FileReference.Exists(LocalManifestFile))
+				FileReference localManifestFile = GetManifestLocation(_localDir, block.NodeName, block.OutputName);
+				if (!FileReference.Exists(localManifestFile))
 				{
 					return false;
 				}
 
 				// Check the local manifest matches the shared manifest
-				if(SharedDir != null)
+				if (_sharedDir != null)
 				{
 					// Check the shared manifest exists
-					FileReference SharedManifestFile = GetManifestLocation(SharedDir, Block.NodeName, Block.OutputName);
-					if(!FileReference.Exists(SharedManifestFile))
+					FileReference sharedManifestFile = GetManifestLocation(_sharedDir, block.NodeName, block.OutputName);
+					if (!FileReference.Exists(sharedManifestFile))
 					{
 						return false;
 					}
 
 					// Check the manifests are identical, byte by byte
-					byte[] LocalManifestBytes = File.ReadAllBytes(LocalManifestFile.FullName);
-					byte[] SharedManifestBytes = Array.Empty<byte>();
-					PerformActionWithRetries(() => SharedManifestBytes = File.ReadAllBytes(SharedManifestFile.FullName), 3, TimeSpan.FromSeconds(1));
-					if(!LocalManifestBytes.SequenceEqual(SharedManifestBytes))
+					byte[] localManifestBytes = File.ReadAllBytes(localManifestFile.FullName);
+					byte[] sharedManifestBytes = Array.Empty<byte>();
+					PerformActionWithRetries(() => sharedManifestBytes = File.ReadAllBytes(sharedManifestFile.FullName), 3, TimeSpan.FromSeconds(1));
+					if (!localManifestBytes.SequenceEqual(sharedManifestBytes))
 					{
 						return false;
 					}
 				}
 
 				// Read the manifest and check the files
-				TempStorageManifest LocalManifest = TempStorageManifest.Load(LocalManifestFile);
-				if(LocalManifest.Files.Any(x => !IgnoreModified.Matches(x.ToFileReference(RootDir).FullName) && !x.Compare(RootDir)))
+				TempStorageManifest localManifest = TempStorageManifest.Load(localManifestFile);
+				if (localManifest.Files.Any(x => !ignoreModified.Matches(x.ToFileReference(_rootDir).FullName) && !x.Compare(_rootDir)))
 				{
 					return false;
 				}
@@ -752,258 +750,258 @@ namespace AutomationTool
 		/// <summary>
 		/// Reads a set of tagged files from disk
 		/// </summary>
-		/// <param name="NodeName">Name of the node which produced the tag set</param>
-		/// <param name="TagName">Name of the tag, with a '#' prefix</param>
+		/// <param name="nodeName">Name of the node which produced the tag set</param>
+		/// <param name="tagName">Name of the tag, with a '#' prefix</param>
 		/// <returns>The set of files</returns>
-		public TempStorageTagManifest? ReadFileList(string NodeName, string TagName)
+		public TempStorageTagManifest? ReadFileList(string nodeName, string tagName)
 		{
 #pragma warning disable CA1508 // False positive; FileList is always null
-			TempStorageTagManifest? FileList = null;
+			TempStorageTagManifest? fileList = null;
 
 			// Try to read the tag set from the local directory
-			FileReference LocalFileListLocation = GetTaggedFileListLocation(LocalDir, NodeName, TagName);
-			if(FileReference.Exists(LocalFileListLocation))
+			FileReference localFileListLocation = GetTaggedFileListLocation(_localDir, nodeName, tagName);
+			if (FileReference.Exists(localFileListLocation))
 			{
-				Logger.LogInformation("Reading local file list from {Arg0}", LocalFileListLocation.FullName);
-				FileList = TempStorageTagManifest.Load(LocalFileListLocation);
+				Logger.LogInformation("Reading local file list from {Arg0}", localFileListLocation.FullName);
+				fileList = TempStorageTagManifest.Load(localFileListLocation);
 			}
 			else
 			{
 				// Check we have shared storage
-				if(SharedDir == null)
+				if (_sharedDir == null)
 				{
-					throw new AutomationException("Missing local file list - {0}", LocalFileListLocation.FullName);
+					throw new AutomationException("Missing local file list - {0}", localFileListLocation.FullName);
 				}
 
 				// Make sure the manifest exists. Try up to 5 times with a 5s wait between to harden against network hiccups.
-				int Attempts = 5;
-				FileReference SharedFileListLocation;
-				SharedFileListLocation = GetTaggedFileListLocation(SharedDir, NodeName, TagName);
+				int attempts = 5;
+				FileReference sharedFileListLocation;
+				sharedFileListLocation = GetTaggedFileListLocation(_sharedDir, nodeName, tagName);
 				PerformActionWithRetries(() =>
 				{
-					if (!FileReference.Exists(SharedFileListLocation))
+					if (!FileReference.Exists(sharedFileListLocation))
 					{
-						throw new AutomationException("Missing local or shared file list - {0}", SharedFileListLocation.FullName);
+						throw new AutomationException("Missing local or shared file list - {0}", sharedFileListLocation.FullName);
 					}
-				}, Attempts, TimeSpan.FromSeconds(5));
+				}, attempts, TimeSpan.FromSeconds(5));
 
 				try
 				{
 					PerformActionWithRetries(() =>
 					{
 						// Read the shared manifest
-						Logger.LogInformation("Copying shared tag set from {Arg0} to {Arg1}", SharedFileListLocation.FullName, LocalFileListLocation.FullName);
-						FileList = TempStorageTagManifest.Load(SharedFileListLocation);
-					}, Attempts, TimeSpan.FromSeconds(5));
+						Logger.LogInformation("Copying shared tag set from {Arg0} to {Arg1}", sharedFileListLocation.FullName, localFileListLocation.FullName);
+						fileList = TempStorageTagManifest.Load(sharedFileListLocation);
+					}, attempts, TimeSpan.FromSeconds(5));
 				}
 				catch
 				{
-					throw new AutomationException("Local or shared file list {0} was found but failed to be read", SharedFileListLocation.FullName);
+					throw new AutomationException("Local or shared file list {0} was found but failed to be read", sharedFileListLocation.FullName);
 				}
 
 				// Save the manifest locally
-				DirectoryReference.CreateDirectory(LocalFileListLocation.Directory);
-				FileList?.Save(LocalFileListLocation);
+				DirectoryReference.CreateDirectory(localFileListLocation.Directory);
+				fileList?.Save(localFileListLocation);
 			}
-			return FileList;
+			return fileList;
 #pragma warning restore CA1508
 		}
 
 		/// <summary>
 		/// Writes a list of tagged files to disk
 		/// </summary>
-		/// <param name="NodeName">Name of the node which produced the tag set</param>
-		/// <param name="TagName">Name of the tag, with a '#' prefix</param>
-		/// <param name="Files">List of files in this set</param>
-		/// <param name="Blocks">List of referenced storage blocks</param>
-		/// <param name="ArtifactKeys">Keys for published artifacts</param>
+		/// <param name="nodeName">Name of the node which produced the tag set</param>
+		/// <param name="tagName">Name of the tag, with a '#' prefix</param>
+		/// <param name="files">List of files in this set</param>
+		/// <param name="blocks">List of referenced storage blocks</param>
+		/// <param name="artifactKeys">Keys for published artifacts</param>
 		/// <returns>The set of files</returns>
-		public void WriteFileList(string NodeName, string TagName, IEnumerable<FileReference> Files, IEnumerable<TempStorageBlock> Blocks, IEnumerable<string> ArtifactKeys)
+		public void WriteFileList(string nodeName, string tagName, IEnumerable<FileReference> files, IEnumerable<TempStorageBlock> blocks, IEnumerable<string> artifactKeys)
 		{
 			// Create the file list
-			TempStorageTagManifest FileList = new TempStorageTagManifest(Files, RootDir, Blocks, ArtifactKeys);
+			TempStorageTagManifest fileList = new TempStorageTagManifest(files, _rootDir, blocks, artifactKeys);
 
 			// Save the set of files to the local and shared locations
-			FileReference LocalFileListLocation = GetTaggedFileListLocation(LocalDir, NodeName, TagName);
-			if(SharedDir != null && bWriteToSharedStorage)
+			FileReference localFileListLocation = GetTaggedFileListLocation(_localDir, nodeName, tagName);
+			if (_sharedDir != null && _writeToSharedStorage)
 			{
-				FileReference SharedFileListLocation = GetTaggedFileListLocation(SharedDir, NodeName, TagName);
+				FileReference sharedFileListLocation = GetTaggedFileListLocation(_sharedDir, nodeName, tagName);
 
 				try
 				{
 					PerformActionWithRetries(() =>
 					{
-						Logger.LogInformation("Saving file list to {Arg0} and {Arg1}", LocalFileListLocation.FullName, SharedFileListLocation.FullName);
-						DirectoryReference.CreateDirectory(SharedFileListLocation.Directory);
-						FileList.Save(SharedFileListLocation);
+						Logger.LogInformation("Saving file list to {Arg0} and {Arg1}", localFileListLocation.FullName, sharedFileListLocation.FullName);
+						DirectoryReference.CreateDirectory(sharedFileListLocation.Directory);
+						fileList.Save(sharedFileListLocation);
 					}, 3, TimeSpan.FromSeconds(5));
 				}
 				catch (Exception ex)
 				{
-					throw new AutomationException("Failed to save file list {0} to {1}, exception: {2}", LocalFileListLocation, SharedFileListLocation, ex);
+					throw new AutomationException("Failed to save file list {0} to {1}, exception: {2}", localFileListLocation, sharedFileListLocation, ex);
 				}
 			}
 			else
 			{
-				Logger.LogInformation("Saving file list to {Arg0}", LocalFileListLocation.FullName);
+				Logger.LogInformation("Saving file list to {Arg0}", localFileListLocation.FullName);
 			}
 
 			// Save the local file list
-			DirectoryReference.CreateDirectory(LocalFileListLocation.Directory);
-			FileList.Save(LocalFileListLocation);
+			DirectoryReference.CreateDirectory(localFileListLocation.Directory);
+			fileList.Save(localFileListLocation);
 		}
 
 		/// <summary>
 		/// Saves the given files (that should be rooted at the branch root) to a shared temp storage manifest with the given temp storage node and game.
 		/// </summary>
-		/// <param name="NodeName">The node which created the storage block</param>
-		/// <param name="BlockName">Name of the block to retrieve. May be null or empty.</param>
-		/// <param name="BuildProducts">Array of build products to be archived</param>
-		/// <param name="bPushToRemote">Allow skipping the copying of this manifest to shared storage, because it's not required by any other agent</param>
+		/// <param name="nodeName">The node which created the storage block</param>
+		/// <param name="blockName">Name of the block to retrieve. May be null or empty.</param>
+		/// <param name="buildProducts">Array of build products to be archived</param>
+		/// <param name="pushToRemote">Allow skipping the copying of this manifest to shared storage, because it's not required by any other agent</param>
 		/// <returns>The created manifest instance (which has already been saved to disk).</returns>
-		public TempStorageManifest Archive(string NodeName, string? BlockName, FileReference[] BuildProducts, bool bPushToRemote = true)
+		public TempStorageManifest Archive(string nodeName, string? blockName, FileReference[] buildProducts, bool pushToRemote = true)
 		{
-			using (IScope Scope = GlobalTracer.Instance.BuildSpan("StoreToTempStorage").StartActive())
+			using (IScope scope = GlobalTracer.Instance.BuildSpan("StoreToTempStorage").StartActive())
 			{
 				// Create a manifest for the given build products
-				FileInfo[] Files = BuildProducts.Select(x => new FileInfo(x.FullName)).ToArray();
-				TempStorageManifest Manifest = new TempStorageManifest(Files, RootDir);
+				FileInfo[] files = buildProducts.Select(x => new FileInfo(x.FullName)).ToArray();
+				TempStorageManifest manifest = new TempStorageManifest(files, _rootDir);
 
 				// Create the local directory for this node
-				DirectoryReference LocalNodeDir = GetDirectoryForNode(LocalDir, NodeName);
-				DirectoryReference.CreateDirectory(LocalNodeDir);
+				DirectoryReference localNodeDir = GetDirectoryForNode(_localDir, nodeName);
+				DirectoryReference.CreateDirectory(localNodeDir);
 
 				// Compress the files and copy to shared storage if necessary
-				bool bRemote = SharedDir != null && bPushToRemote && bWriteToSharedStorage;
-				if(bRemote)
+				bool remote = _sharedDir != null && pushToRemote && _writeToSharedStorage;
+				if (remote)
 				{
 					// Create the shared directory for this node
-					FileReference SharedManifestFile = GetManifestLocation(SharedDir!, NodeName, BlockName);
-					DirectoryReference.CreateDirectory(SharedManifestFile.Directory);
+					FileReference sharedManifestFile = GetManifestLocation(_sharedDir!, nodeName, blockName);
+					DirectoryReference.CreateDirectory(sharedManifestFile.Directory);
 
 					// Zip all the build products
-					FileInfo[] ZipFiles = ParallelZipFiles(Files, RootDir, SharedManifestFile.Directory, LocalNodeDir, SharedManifestFile.GetFileNameWithoutExtension());
-					Manifest.ZipFiles = ZipFiles.Select(x => new TempStorageZipFile(x)).ToArray();
+					FileInfo[] zipFiles = ParallelZipFiles(files, _rootDir, sharedManifestFile.Directory, localNodeDir, sharedManifestFile.GetFileNameWithoutExtension());
+					manifest.ZipFiles = zipFiles.Select(x => new TempStorageZipFile(x)).ToArray();
 
 					// Save the shared manifest
-					Logger.LogInformation("Saving shared manifest to {Arg0}", SharedManifestFile.FullName);
-					PerformActionWithRetries(() => Manifest.Save(SharedManifestFile), 3, TimeSpan.FromSeconds(5));
+					Logger.LogInformation("Saving shared manifest to {Arg0}", sharedManifestFile.FullName);
+					PerformActionWithRetries(() => manifest.Save(sharedManifestFile), 3, TimeSpan.FromSeconds(5));
 				}
 
 				// Save the local manifest
-				FileReference LocalManifestFile = GetManifestLocation(LocalDir, NodeName, BlockName);
-				Logger.LogInformation("Saving local manifest to {Arg0}", LocalManifestFile.FullName);
-				Manifest.Save(LocalManifestFile);
+				FileReference localManifestFile = GetManifestLocation(_localDir, nodeName, blockName);
+				Logger.LogInformation("Saving local manifest to {Arg0}", localManifestFile.FullName);
+				manifest.Save(localManifestFile);
 
 				// Update the stats
-				long ZipFilesTotalSize = (Manifest.ZipFiles == null)? 0 : Manifest.ZipFiles.Sum(x => x.Length);
-				Scope.Span.SetTag("numFiles", Files.Length);
-				Scope.Span.SetTag("manifestSize", Manifest.GetTotalSize());
-				Scope.Span.SetTag("manifestZipFilesSize", ZipFilesTotalSize);
-				Scope.Span.SetTag("isRemote", bRemote);
-				Scope.Span.SetTag("blockName", BlockName);
-				return Manifest;
+				long zipFilesTotalSize = (manifest.ZipFiles == null) ? 0 : manifest.ZipFiles.Sum(x => x.Length);
+				scope.Span.SetTag("numFiles", files.Length);
+				scope.Span.SetTag("manifestSize", manifest.GetTotalSize());
+				scope.Span.SetTag("manifestZipFilesSize", zipFilesTotalSize);
+				scope.Span.SetTag("isRemote", remote);
+				scope.Span.SetTag("blockName", blockName);
+				return manifest;
 			}
 		}
 
 		/// <summary>
 		/// Retrieve an output of the given node. Fetches and decompresses the files from shared storage if necessary, or validates the local files.
 		/// </summary>
-		/// <param name="NodeName">The node which created the storage block</param>
-		/// <param name="OutputName">Name of the block to retrieve. May be null or empty.</param>
-		/// <param name="IgnoreModified">Filter for files to ignore</param>
+		/// <param name="nodeName">The node which created the storage block</param>
+		/// <param name="outputName">Name of the block to retrieve. May be null or empty.</param>
+		/// <param name="ignoreModified">Filter for files to ignore</param>
 		/// <returns>Manifest of the files retrieved</returns>
-		public TempStorageManifest Retrieve(string NodeName, string? OutputName, FileFilter IgnoreModified)
+		public TempStorageManifest Retrieve(string nodeName, string? outputName, FileFilter ignoreModified)
 		{
-			using (IScope Scope = GlobalTracer.Instance.BuildSpan("RetrieveFromTempStorage").StartActive())
+			using (IScope scope = GlobalTracer.Instance.BuildSpan("RetrieveFromTempStorage").StartActive())
 			{
 				// Get the path to the local manifest
-				FileReference LocalManifestFile = GetManifestLocation(LocalDir, NodeName, OutputName);
-				bool bLocal = FileReference.Exists(LocalManifestFile);
+				FileReference localManifestFile = GetManifestLocation(_localDir, nodeName, outputName);
+				bool local = FileReference.Exists(localManifestFile);
 
 				// Read the manifest, either from local storage or shared storage
-				TempStorageManifest? Manifest = null;
-				if(bLocal)
+				TempStorageManifest? manifest = null;
+				if (local)
 				{
-					Logger.LogInformation("Reading shared manifest from {Arg0}", LocalManifestFile.FullName);
-					Manifest = TempStorageManifest.Load(LocalManifestFile);
+					Logger.LogInformation("Reading shared manifest from {Arg0}", localManifestFile.FullName);
+					manifest = TempStorageManifest.Load(localManifestFile);
 				}
 				else
 				{
 					// Check we have shared storage
-					if(SharedDir == null)
+					if (_sharedDir == null)
 					{
-						throw new AutomationException("Missing local manifest for node - {0}", LocalManifestFile.FullName);
+						throw new AutomationException("Missing local manifest for node - {0}", localManifestFile.FullName);
 					}
 
 					// Get the shared directory for this node
-					FileReference SharedManifestFile = GetManifestLocation(SharedDir, NodeName, OutputName);
+					FileReference sharedManifestFile = GetManifestLocation(_sharedDir, nodeName, outputName);
 
 					// Make sure the manifest exists
-					if(!FileReference.Exists(SharedManifestFile))
+					if (!FileReference.Exists(sharedManifestFile))
 					{
-						throw new AutomationException("Missing local or shared manifest for node - {0}", SharedManifestFile.FullName);
+						throw new AutomationException("Missing local or shared manifest for node - {0}", sharedManifestFile.FullName);
 					}
 
 					// Read the shared manifest
-					Logger.LogInformation("Copying shared manifest from {Arg0} to {Arg1}", SharedManifestFile.FullName, LocalManifestFile.FullName);
-					PerformActionWithRetries(() => Manifest = TempStorageManifest.Load(SharedManifestFile), 3, TimeSpan.FromSeconds(5));
+					Logger.LogInformation("Copying shared manifest from {Arg0} to {Arg1}", sharedManifestFile.FullName, localManifestFile.FullName);
+					PerformActionWithRetries(() => manifest = TempStorageManifest.Load(sharedManifestFile), 3, TimeSpan.FromSeconds(5));
 
 					// Unzip all the build products
-					DirectoryReference SharedNodeDir = GetDirectoryForNode(SharedDir, NodeName);
-					FileInfo[] ZipFiles = Manifest!.ZipFiles.Select(x => new FileInfo(FileReference.Combine(SharedNodeDir, x.Name).FullName)).ToArray();
-					string Result = ParallelUnzipFiles(ZipFiles, RootDir);
-					if (!string.IsNullOrWhiteSpace(Result))
+					DirectoryReference sharedNodeDir = GetDirectoryForNode(_sharedDir, nodeName);
+					FileInfo[] zipFiles = manifest!.ZipFiles.Select(x => new FileInfo(FileReference.Combine(sharedNodeDir, x.Name).FullName)).ToArray();
+					string result = ParallelUnzipFiles(zipFiles, _rootDir);
+					if (!string.IsNullOrWhiteSpace(result))
 					{
-						string LogPath = CommandUtils.CombinePaths(RootDir.FullName, "Engine/Programs/AutomationTool/Saved/Logs", $"Copy Manifest - {NodeName}.log");
-						Logger.LogInformation("Saving copy log to {LogPath}", LogPath);
-						File.WriteAllText(LogPath, Result);
+						string logPath = CommandUtils.CombinePaths(_rootDir.FullName, "Engine/Programs/AutomationTool/Saved/Logs", $"Copy Manifest - {nodeName}.log");
+						Logger.LogInformation("Saving copy log to {LogPath}", logPath);
+						File.WriteAllText(logPath, result);
 					}
 
 					// Update the timestamps to match the manifest. Zip files only use local time, and there's no guarantee it matches the local clock.
-					foreach(TempStorageFile ManifestFile in Manifest.Files)
+					foreach (TempStorageFile manifestFile in manifest.Files)
 					{
-						FileReference File = ManifestFile.ToFileReference(RootDir);
-						System.IO.File.SetLastWriteTimeUtc(File.FullName, new DateTime(ManifestFile.LastWriteTimeUtcTicks, DateTimeKind.Utc));
+						FileReference file = manifestFile.ToFileReference(_rootDir);
+						System.IO.File.SetLastWriteTimeUtc(file.FullName, new DateTime(manifestFile.LastWriteTimeUtcTicks, DateTimeKind.Utc));
 					}
 
 					// Save the manifest locally
-					DirectoryReference.CreateDirectory(LocalManifestFile.Directory);
-					Manifest.Save(LocalManifestFile);
+					DirectoryReference.CreateDirectory(localManifestFile.Directory);
+					manifest.Save(localManifestFile);
 				}
 
 				// Check all the local files are as expected
-				List<string> ModifiedFileMessages = new List<string>();
-				foreach(TempStorageFile File in Manifest.Files)
+				List<string> modifiedFileMessages = new List<string>();
+				foreach (TempStorageFile file in manifest.Files)
 				{
-					string? Message;
-					if (!IgnoreModified.Matches(File.ToFileReference(RootDir).FullName) && !File.Compare(RootDir, out Message))
+					string? message;
+					if (!ignoreModified.Matches(file.ToFileReference(_rootDir).FullName) && !file.Compare(_rootDir, out message))
 					{
-						ModifiedFileMessages.Add(Message);
+						modifiedFileMessages.Add(message);
 					}
 				}
-				if(ModifiedFileMessages.Count > 0)
+				if (modifiedFileMessages.Count > 0)
 				{
 					string modifiedFileList = "";
-					if (ModifiedFileMessages.Count < 100)
+					if (modifiedFileMessages.Count < 100)
 					{
-						modifiedFileList = String.Join("\n", ModifiedFileMessages.Select(x => $"  {x}"));
+						modifiedFileList = String.Join("\n", modifiedFileMessages.Select(x => $"  {x}"));
 					}
 					else
 					{
-						modifiedFileList = String.Join("\n", ModifiedFileMessages.Take(100).Select(x => $"  {x}"));
-						modifiedFileList += $"\n  ...and {ModifiedFileMessages.Count - 100} more.";
+						modifiedFileList = String.Join("\n", modifiedFileMessages.Take(100).Select(x => $"  {x}"));
+						modifiedFileList += $"\n  ...and {modifiedFileMessages.Count - 100} more.";
 					}
 					throw new AutomationException($"Files have been modified:\n{modifiedFileList}");
 				}
 
 				// Update the stats and return
-				Scope.Span.SetTag("numFiles", Manifest.Files.Length);
-				Scope.Span.SetTag("manifestSize", Manifest.Files.Sum(x => x.Length));
-				Scope.Span.SetTag("manifestZipFilesSize", bLocal? 0 : Manifest.ZipFiles.Sum(x => x.Length));
-				Scope.Span.SetTag("isRemote", !bLocal);
-				Scope.Span.SetTag("outputName", OutputName);
-				return Manifest;
+				scope.Span.SetTag("numFiles", manifest.Files.Length);
+				scope.Span.SetTag("manifestSize", manifest.Files.Sum(x => x.Length));
+				scope.Span.SetTag("manifestZipFilesSize", local ? 0 : manifest.ZipFiles.Sum(x => x.Length));
+				scope.Span.SetTag("isRemote", !local);
+				scope.Span.SetTag("outputName", outputName);
+				return manifest;
 			}
 		}
 
@@ -1031,54 +1029,54 @@ namespace AutomationTool
 		/// <summary>
 		/// Zips a set of files (that must be rooted at the given RootDir) to a set of zip files in the given OutputDir. The files will be prefixed with the given basename.
 		/// </summary>
-		/// <param name="InputFiles">Fully qualified list of files to zip (must be rooted at RootDir).</param>
-		/// <param name="RootDir">Root Directory where all files will be extracted.</param>
-		/// <param name="OutputDir">Location to place the set of zip files created.</param>
-		/// <param name="StagingDir">Location to create zip files before copying them to the OutputDir. If the OutputDir is on a remote file share, staging may be more efficient. Use null to avoid using a staging copy.</param>
-		/// <param name="ZipBaseName">The basename of the set of zip files.</param>
+		/// <param name="inputFiles">Fully qualified list of files to zip (must be rooted at RootDir).</param>
+		/// <param name="rootDir">Root Directory where all files will be extracted.</param>
+		/// <param name="outputDir">Location to place the set of zip files created.</param>
+		/// <param name="stagingDir">Location to create zip files before copying them to the OutputDir. If the OutputDir is on a remote file share, staging may be more efficient. Use null to avoid using a staging copy.</param>
+		/// <param name="zipBaseName">The basename of the set of zip files.</param>
 		/// <returns>Some metrics about the zip process.</returns>
 		/// <remarks>
 		/// This function tries to zip the files in parallel as fast as it can. It makes no guarantees about how many zip files will be created or which files will be in which zip,
 		/// but it does try to reasonably balance the file sizes.
 		/// </remarks>
-		static FileInfo[] ParallelZipFiles(FileInfo[] InputFiles, DirectoryReference RootDir, DirectoryReference OutputDir, DirectoryReference StagingDir, string ZipBaseName)
+		static FileInfo[] ParallelZipFiles(FileInfo[] inputFiles, DirectoryReference rootDir, DirectoryReference outputDir, DirectoryReference stagingDir, string zipBaseName)
 		{
 			// First get the sizes of all the files. We won't parallelize if there isn't enough data to keep the number of zips down.
-			var FilesInfo = InputFiles
-				.Select(InputFile => new { File = new FileReference(InputFile), FileSize = InputFile.Length })
+			var filesInfo = inputFiles
+				.Select(inputFile => new { File = new FileReference(inputFile), FileSize = inputFile.Length })
 				.ToList();
 
 			// Profiling results show that we can zip 100MB quite fast and it is not worth parallelizing that case and creating a bunch of zips that are relatively small.
 			const long MinFileSizeToZipInParallel = 1024 * 1024 * 100L;
-			bool bZipInParallel = FilesInfo.Sum(FileInfo => FileInfo.FileSize) >= MinFileSizeToZipInParallel;
+			bool zipInParallel = filesInfo.Sum(fileInfo => fileInfo.FileSize) >= MinFileSizeToZipInParallel;
 
 			// order the files in descending order so our threads pick up the biggest ones first.
 			// We want to end with the smaller files to more effectively fill in the gaps
-			ConcurrentQueue<FileReference> FilesToZip = new(FilesInfo.OrderByDescending(FileInfo => FileInfo.FileSize).Select(FileInfo => FileInfo.File));
+			ConcurrentQueue<FileReference> filesToZip = new(filesInfo.OrderByDescending(fileInfo => fileInfo.FileSize).Select(fileInfo => fileInfo.File));
 
-			ConcurrentBag<FileInfo> ZipFiles = new ConcurrentBag<FileInfo>();
+			ConcurrentBag<FileInfo> zipFiles = new ConcurrentBag<FileInfo>();
 
-			DirectoryReference ZipDir = StagingDir ?? OutputDir;
+			DirectoryReference zipDir = stagingDir ?? outputDir;
 
 			// We deliberately avoid Parallel.ForEach here because profiles have shown that dynamic partitioning creates
 			// too many zip files, and they can be of too varying size, creating uneven work when unzipping later,
 			// as ZipFile cannot unzip files in parallel from a single archive.
 			// We can safely assume the build system will not be doing more important things at the same time, so we simply use all our logical cores,
 			// which has shown to be optimal via profiling, and limits the number of resulting zip files to the number of logical cores.
-			List<Thread> ZipThreads = (
-				from CoreNum in Enumerable.Range(0, bZipInParallel ? Environment.ProcessorCount : 1)
+			List<Thread> zipThreads = (
+				from CoreNum in Enumerable.Range(0, zipInParallel ? Environment.ProcessorCount : 1)
 				select new Thread((object? indexObject) =>
 				{
 					int index = (int)indexObject!;
-					FileReference ZipFileName = FileReference.Combine(ZipDir, string.Format("{0}{1}.zip", ZipBaseName, bZipInParallel ? "-" + index.ToString("00") : ""));
+					FileReference zipFileName = FileReference.Combine(zipDir, string.Format("{0}{1}.zip", zipBaseName, zipInParallel ? "-" + index.ToString("00") : ""));
 					// don't create the zip unless we have at least one file to add
-					FileReference? File;
-					if (FilesToZip.TryDequeue(out File))
+					FileReference? file;
+					if (filesToZip.TryDequeue(out file))
 					{
 						try
 						{
 							// Create one zip per thread using the given basename
-							using (ZipArchive ZipArchive = ZipFile.Open(ZipFileName.FullName, ZipArchiveMode.Create))
+							using (ZipArchive zipArchive = ZipFile.Open(zipFileName.FullName, ZipArchiveMode.Create))
 							{
 								// pull from the queue until we are out of files.
 								do
@@ -1086,125 +1084,124 @@ namespace AutomationTool
 									// use fastest compression. In our best case we are CPU bound, so this is a good tradeoff,
 									// cutting overall time by 2/3 while only modestly increasing the compression ratio (22.7% -> 23.8% for RootEditor PDBs).
 									// This is in cases of a super hot cache, so the operation was largely CPU bound.
-									ZipArchiveExtensions.CreateEntryFromFile_CrossPlatform(ZipArchive, File.FullName, CommandUtils.ConvertSeparators(PathSeparator.Slash, File.MakeRelativeTo(RootDir)), CompressionLevel.Fastest);
-								} while (FilesToZip.TryDequeue(out File));
+									ZipArchiveExtensions.CreateEntryFromFile_CrossPlatform(zipArchive, file.FullName, CommandUtils.ConvertSeparators(PathSeparator.Slash, file.MakeRelativeTo(rootDir)), CompressionLevel.Fastest);
+								} while (filesToZip.TryDequeue(out file));
 							}
 						}
 						catch (IOException)
 						{
-							Logger.LogError("Unable to open file for TempStorage zip: \"{Arg0}\"", ZipFileName.FullName);
-							throw new AutomationException("Unable to open file {0}", ZipFileName.FullName);
+							Logger.LogError("Unable to open file for TempStorage zip: \"{Arg0}\"", zipFileName.FullName);
+							throw new AutomationException("Unable to open file {0}", zipFileName.FullName);
 						}
 
-						ZipFiles.Add(new FileInfo(ZipFileName.FullName));
+						zipFiles.Add(new FileInfo(zipFileName.FullName));
 					}
 				})).ToList();
 
-			for (int index = 0; index < ZipThreads.Count; index++)
+			for (int index = 0; index < zipThreads.Count; index++)
 			{
-				Thread thread = ZipThreads[index];
+				Thread thread = zipThreads[index];
 				thread.Start(index);
 			}
 
-			ZipThreads.ForEach(thread => thread.Join());
+			zipThreads.ForEach(thread => thread.Join());
 
-			if (ZipFiles.Any() && !string.IsNullOrWhiteSpace(StagingDir?.FullName))
+			if (zipFiles.Any() && !string.IsNullOrWhiteSpace(stagingDir?.FullName))
 			{
 				try
 				{
-					string CopyResult = CopyDirectory(ZipDir, OutputDir);
-					string LogPath = CommandUtils.CombinePaths(RootDir.FullName, "Engine/Programs/AutomationTool/Saved/Logs", $"Copy Files to Temp Storage.log");
-					Logger.LogInformation("Saving copy log to {LogPath}", LogPath);
-					Directory.CreateDirectory(Path.GetDirectoryName(LogPath)!);
-					File.WriteAllText(LogPath, string.Join(Environment.NewLine, CopyResult));
-					Parallel.ForEach(ZipFiles, (z) => z.Delete());
-					ZipFiles = new ConcurrentBag<FileInfo>(ZipFiles.Select(z => new FileInfo(CommandUtils.MakeRerootedFilePath(z.FullName, StagingDir.FullName, OutputDir.FullName))));
+					string copyResult = CopyDirectory(zipDir, outputDir);
+					string logPath = CommandUtils.CombinePaths(rootDir.FullName, "Engine/Programs/AutomationTool/Saved/Logs", $"Copy Files to Temp Storage.log");
+					Logger.LogInformation("Saving copy log to {LogPath}", logPath);
+					Directory.CreateDirectory(Path.GetDirectoryName(logPath)!);
+					File.WriteAllText(logPath, string.Join(Environment.NewLine, copyResult));
+					Parallel.ForEach(zipFiles, (z) => z.Delete());
+					zipFiles = new ConcurrentBag<FileInfo>(zipFiles.Select(z => new FileInfo(CommandUtils.MakeRerootedFilePath(z.FullName, stagingDir.FullName, outputDir.FullName))));
 				}
-				catch (IOException Ex)
+				catch (IOException ex)
 				{
-					Logger.LogError("Unable to copy staging directory {Arg0} to {Arg1}, Ex: {Ex}", ZipDir.FullName, OutputDir.FullName, Ex);
-					throw new AutomationException(Ex, "Unable to copy staging directory {0} to {1}", ZipDir.FullName, OutputDir.FullName);
+					Logger.LogError("Unable to copy staging directory {Arg0} to {Arg1}, Ex: {Ex}", zipDir.FullName, outputDir.FullName, ex);
+					throw new AutomationException(ex, "Unable to copy staging directory {0} to {1}", zipDir.FullName, outputDir.FullName);
 				}
 			}
 
-			return ZipFiles.OrderBy(x => x.Name).ToArray();
+			return zipFiles.OrderBy(x => x.Name).ToArray();
 		}
 
 		/// <summary>
 		/// Unzips a set of zip files with a given basename in a given folder to a given RootDir.
 		/// </summary>
-		/// <param name="ZipFiles">Files to extract</param>
-		/// <param name="RootDir">Root Directory where all files will be extracted.</param>
+		/// <param name="zipFiles">Files to extract</param>
+		/// <param name="rootDir">Root Directory where all files will be extracted.</param>
 		/// <returns>Some metrics about the unzip process.</returns>
 		/// <remarks>
 		/// The code is expected to be the used as the symmetrical inverse of <see cref="ParallelZipFiles"/>, but could be used independently, as long as the files in the zip do not overlap.
 		/// </remarks>
-		private static string ParallelUnzipFiles(FileInfo[] ZipFiles, DirectoryReference RootDir)
+		private static string ParallelUnzipFiles(FileInfo[] zipFiles, DirectoryReference rootDir)
 		{
-			ConcurrentBag<string> CopyResults = new ConcurrentBag<string>();
-			Parallel.ForEach(ZipFiles,
-				(ZipFile) =>
+			ConcurrentBag<string> copyResults = new ConcurrentBag<string>();
+			Parallel.ForEach(zipFiles,
+				(zipFile) =>
 				{
 					// Copy the ZIP to the local drive before unzipping to harden against network issues.
-					CopyResults.Add(CopyFile(ZipFile, RootDir));
-					string LocalZipFile = CommandUtils.CombinePaths(RootDir.FullName, ZipFile.Name);
+					copyResults.Add(CopyFile(zipFile, rootDir));
+					string localZipFile = CommandUtils.CombinePaths(rootDir.FullName, zipFile.Name);
 
 					// unzip the files manually instead of caling ZipFile.ExtractToDirectory() because we need to overwrite readonly files. Because of this, creating the directories is up to us as well.
-					List<string> ExtractedPaths = new List<string>();
-					int UnzipFileAttempts = 3;
-					while (UnzipFileAttempts-- > 0)
+					List<string> extractedPaths = new List<string>();
+					int unzipFileAttempts = 3;
+					while (unzipFileAttempts-- > 0)
 					{
 						try
 						{
-							using (ZipArchive ZipArchive = System.IO.Compression.ZipFile.OpenRead(LocalZipFile))
+							using (ZipArchive zipArchive = System.IO.Compression.ZipFile.OpenRead(localZipFile))
 							{
-								foreach (ZipArchiveEntry Entry in ZipArchive.Entries)
+								foreach (ZipArchiveEntry entry in zipArchive.Entries)
 								{
 									// Use CommandUtils.CombinePaths to ensure directory separators get converted correctly.
-									string ExtractedFilename = CommandUtils.CombinePaths(RootDir.FullName, Entry.FullName);
+									string extractedFilename = CommandUtils.CombinePaths(rootDir.FullName, entry.FullName);
 
 									// Skip this if it's already been extracted.
-									if (ExtractedPaths.Contains(ExtractedFilename))
+									if (extractedPaths.Contains(extractedFilename))
 									{
 										continue;
 									}
 
 									// Zips can contain empty dirs. Ours usually don't have them, but we should support it.
-									if (Path.GetFileName(ExtractedFilename).Length == 0)
+									if (Path.GetFileName(extractedFilename).Length == 0)
 									{
-										Directory.CreateDirectory(ExtractedFilename);
-										ExtractedPaths.Add(ExtractedFilename);
+										Directory.CreateDirectory(extractedFilename);
+										extractedPaths.Add(extractedFilename);
 									}
 									else
 									{
 										// We must delete any existing file, even if it's readonly. .Net does not do this by default.
-										if (File.Exists(ExtractedFilename))
+										if (File.Exists(extractedFilename))
 										{
-											InternalUtils.SafeDeleteFile(ExtractedFilename, true);
+											InternalUtils.SafeDeleteFile(extractedFilename, true);
 										}
 										else
 										{
-											Directory.CreateDirectory(Path.GetDirectoryName(ExtractedFilename)!);
+											Directory.CreateDirectory(Path.GetDirectoryName(extractedFilename)!);
 										}
 
-
-										int UnzipEntryAttempts = 3;
-										while (UnzipEntryAttempts-- > 0)
+										int unzipEntryAttempts = 3;
+										while (unzipEntryAttempts-- > 0)
 										{
 											try
 											{
-												Entry.ExtractToFile_CrossPlatform(ExtractedFilename, true);
-												ExtractedPaths.Add(ExtractedFilename);
+												entry.ExtractToFile_CrossPlatform(extractedFilename, true);
+												extractedPaths.Add(extractedFilename);
 												break;
 											}
-											catch (IOException IOEx)
+											catch (IOException ioEx)
 											{
-												if (UnzipEntryAttempts == 0)
+												if (unzipEntryAttempts == 0)
 												{
 													throw;
 												}
 
-												Log.Logger.LogWarning(IOEx, "Failed to unzip '{File}' from '{LocalZipFile}' to '{ExtractedFilename}', retrying.. (Error: {Message})", Entry.FullName, LocalZipFile, ExtractedFilename, IOEx.Message);
+												Log.Logger.LogWarning(ioEx, "Failed to unzip '{File}' from '{LocalZipFile}' to '{ExtractedFilename}', retrying.. (Error: {Message})", entry.FullName, localZipFile, extractedFilename, ioEx.Message);
 											}
 										}
 									}
@@ -1213,184 +1210,184 @@ namespace AutomationTool
 
 							break;
 						}
-						catch (Exception Ex)
+						catch (Exception ex)
 						{
-							if (UnzipFileAttempts == 0)
+							if (unzipFileAttempts == 0)
 							{
-								Log.Logger.LogError(Ex, "All retries exhausted attempting to unzip entries from '{LocalZipFile}'. Terminating.", LocalZipFile);
-								string LogPath = CommandUtils.CombinePaths(RootDir.FullName, "Engine/Programs/AutomationTool/Saved/Logs", $"Copy Manifest - {ZipFile.Name}.log");
-								Logger.LogInformation("Saving copy log to {LogPath}", LogPath);
-								File.WriteAllText(LogPath, string.Join(Environment.NewLine, CopyResults));
+								Log.Logger.LogError(ex, "All retries exhausted attempting to unzip entries from '{LocalZipFile}'. Terminating.", localZipFile);
+								string logPath = CommandUtils.CombinePaths(rootDir.FullName, "Engine/Programs/AutomationTool/Saved/Logs", $"Copy Manifest - {zipFile.Name}.log");
+								Logger.LogInformation("Saving copy log to {LogPath}", logPath);
+								File.WriteAllText(logPath, string.Join(Environment.NewLine, copyResults));
 								throw;
 							}
 
 							// Some exceptions may be caused by networking hiccups. We want to retry in those cases.
-							if ((Ex is IOException || Ex is InvalidDataException))
+							if ((ex is IOException || ex is InvalidDataException))
 							{
-								Log.Logger.LogWarning(Ex, "Failed to unzip entries from '{LocalZipFile}' to '{TargetDir}', retrying.. (Error: {Message})", LocalZipFile, RootDir.FullName, Ex.Message);
+								Log.Logger.LogWarning(ex, "Failed to unzip entries from '{LocalZipFile}' to '{TargetDir}', retrying.. (Error: {Message})", localZipFile, rootDir.FullName, ex.Message);
 							}
 						}
 						finally
 						{
-							if (File.Exists(LocalZipFile))
+							if (File.Exists(localZipFile))
 							{
-								File.Delete(LocalZipFile);
+								File.Delete(localZipFile);
 							}
 						}
 					}
 				});
 
-			return string.Join(Environment.NewLine, CopyResults);
+			return string.Join(Environment.NewLine, copyResults);
 		}
 
 		/// <summary>
 		/// Gets the directory used to store data for the given node
 		/// </summary>
-		/// <param name="BaseDir">A local or shared temp storage root directory.</param>
-		/// <param name="NodeName">Name of the node</param>
+		/// <param name="baseDir">A local or shared temp storage root directory.</param>
+		/// <param name="nodeName">Name of the node</param>
 		/// <returns>Directory to contain a node's data</returns>
-		static DirectoryReference GetDirectoryForNode(DirectoryReference BaseDir, string NodeName)
+		static DirectoryReference GetDirectoryForNode(DirectoryReference baseDir, string nodeName)
 		{
-			return DirectoryReference.Combine(BaseDir, NodeName);
+			return DirectoryReference.Combine(baseDir, nodeName);
 		}
 
 		/// <summary>
 		/// Gets the path to the manifest created for a node's output.
 		/// </summary>
-		/// <param name="BaseDir">A local or shared temp storage root directory.</param>
-		/// <param name="NodeName">Name of the node to get the file for</param>
-		/// <param name="BlockName">Name of the output block to get the manifest for</param>
-		static FileReference GetManifestLocation(DirectoryReference BaseDir, string NodeName, string? BlockName)
+		/// <param name="baseDir">A local or shared temp storage root directory.</param>
+		/// <param name="nodeName">Name of the node to get the file for</param>
+		/// <param name="blockName">Name of the output block to get the manifest for</param>
+		static FileReference GetManifestLocation(DirectoryReference baseDir, string nodeName, string? blockName)
 		{
-			return FileReference.Combine(BaseDir, NodeName, String.IsNullOrEmpty(BlockName)? "Manifest.xml" : String.Format("Manifest-{0}.xml", BlockName));
+			return FileReference.Combine(baseDir, nodeName, String.IsNullOrEmpty(blockName) ? "Manifest.xml" : String.Format("Manifest-{0}.xml", blockName));
 		}
 
 		/// <summary>
 		/// Gets the path to the file created to store a tag manifest for a node
 		/// </summary>
-		/// <param name="BaseDir">A local or shared temp storage root directory.</param>
-		/// <param name="NodeName">Name of the node to get the file for</param>
-		/// <param name="TagName">Name of the tag to get the manifest for</param>
-		static FileReference GetTaggedFileListLocation(DirectoryReference BaseDir, string NodeName, string TagName)
+		/// <param name="baseDir">A local or shared temp storage root directory.</param>
+		/// <param name="nodeName">Name of the node to get the file for</param>
+		/// <param name="tagName">Name of the tag to get the manifest for</param>
+		static FileReference GetTaggedFileListLocation(DirectoryReference baseDir, string nodeName, string tagName)
 		{
-			Debug.Assert(TagName.StartsWith("#", StringComparison.Ordinal));
-			return FileReference.Combine(BaseDir, NodeName, String.Format("Tag-{0}.xml", TagName.Substring(1)));
+			Debug.Assert(tagName.StartsWith("#", StringComparison.Ordinal));
+			return FileReference.Combine(baseDir, nodeName, String.Format("Tag-{0}.xml", tagName.Substring(1)));
 		}
 
 		/// <summary>
 		/// Gets the path to a file created to indicate that a node is complete, under the given base directory.
 		/// </summary>
-		/// <param name="BaseDir">A local or shared temp storage root directory.</param>
-		/// <param name="NodeName">Name of the node to get the file for</param>
-		static FileReference GetCompleteMarkerFile(DirectoryReference BaseDir, string NodeName)
+		/// <param name="baseDir">A local or shared temp storage root directory.</param>
+		/// <param name="nodeName">Name of the node to get the file for</param>
+		static FileReference GetCompleteMarkerFile(DirectoryReference baseDir, string nodeName)
 		{
-			return FileReference.Combine(GetDirectoryForNode(BaseDir, NodeName), "Complete");
+			return FileReference.Combine(GetDirectoryForNode(baseDir, nodeName), "Complete");
 		}
 
 		/// <summary>
 		/// Checks whether the given path is allowed as a build product that can be produced by more than one node (timestamps may be modified, etc..). Used to suppress
 		/// warnings about build products being overwritten.
 		/// </summary>
-		/// <param name="LocalFile">File name to check</param>
+		/// <param name="localFile">File name to check</param>
 		/// <returns>True if the given path may be output by multiple build products</returns>
-		public static bool IsDuplicateBuildProduct(FileReference LocalFile)
+		public static bool IsDuplicateBuildProduct(FileReference localFile)
 		{
-			string FileName = LocalFile.GetFileName();
-			if (FileName.Equals("AgentInterface.dll", StringComparison.OrdinalIgnoreCase) || FileName.Equals("AgentInterface.pdb", StringComparison.OrdinalIgnoreCase))
+			string fileName = localFile.GetFileName();
+			if (fileName.Equals("AgentInterface.dll", StringComparison.OrdinalIgnoreCase) || fileName.Equals("AgentInterface.pdb", StringComparison.OrdinalIgnoreCase))
 			{
 				return true;
 			}
-			if (FileName.Equals("dxil.dll", StringComparison.OrdinalIgnoreCase))
+			if (fileName.Equals("dxil.dll", StringComparison.OrdinalIgnoreCase))
 			{
 				return true;
 			}
-			if (FileName.Equals("dxcompiler.dll", StringComparison.OrdinalIgnoreCase))
+			if (fileName.Equals("dxcompiler.dll", StringComparison.OrdinalIgnoreCase))
 			{
 				return true;
 			}
-			if (FileName.Equals("embree.2.14.0.dll", StringComparison.OrdinalIgnoreCase) || FileName.Equals("libembree.2.14.0.dylib", StringComparison.OrdinalIgnoreCase))
+			if (fileName.Equals("embree.2.14.0.dll", StringComparison.OrdinalIgnoreCase) || fileName.Equals("libembree.2.14.0.dylib", StringComparison.OrdinalIgnoreCase))
 			{
 				return true;
 			}
-			if (FileName.Equals("tbb.dll", StringComparison.OrdinalIgnoreCase) || FileName.Equals("tbb.pdb", StringComparison.OrdinalIgnoreCase) || FileName.Equals("libtbb.dylib", StringComparison.OrdinalIgnoreCase) || FileName.Equals("tbb.psym", StringComparison.OrdinalIgnoreCase))
+			if (fileName.Equals("tbb.dll", StringComparison.OrdinalIgnoreCase) || fileName.Equals("tbb.pdb", StringComparison.OrdinalIgnoreCase) || fileName.Equals("libtbb.dylib", StringComparison.OrdinalIgnoreCase) || fileName.Equals("tbb.psym", StringComparison.OrdinalIgnoreCase))
 			{
 				return true;
 			}
-			if (FileName.Equals("tbbmalloc.dll", StringComparison.OrdinalIgnoreCase) || FileName.Equals("tbbmalloc.pdb", StringComparison.OrdinalIgnoreCase) || FileName.Equals("libtbbmalloc.dylib", StringComparison.OrdinalIgnoreCase) || FileName.Equals("tbbmalloc.psym", StringComparison.OrdinalIgnoreCase))
+			if (fileName.Equals("tbbmalloc.dll", StringComparison.OrdinalIgnoreCase) || fileName.Equals("tbbmalloc.pdb", StringComparison.OrdinalIgnoreCase) || fileName.Equals("libtbbmalloc.dylib", StringComparison.OrdinalIgnoreCase) || fileName.Equals("tbbmalloc.psym", StringComparison.OrdinalIgnoreCase))
 			{
 				return true;
 			}
-			if (FileName.EndsWith(".dll", StringComparison.OrdinalIgnoreCase))
+			if (fileName.EndsWith(".dll", StringComparison.OrdinalIgnoreCase))
 			{
 				return true;
 			}
-			if (FileName.StartsWith("lib", StringComparison.OrdinalIgnoreCase) && FileName.EndsWith(".dylib", StringComparison.OrdinalIgnoreCase))
+			if (fileName.StartsWith("lib", StringComparison.OrdinalIgnoreCase) && fileName.EndsWith(".dylib", StringComparison.OrdinalIgnoreCase))
 			{
 				return true;
 			}
-			if (FileName.StartsWith("lib", StringComparison.OrdinalIgnoreCase) && FileName.EndsWith(".so", StringComparison.OrdinalIgnoreCase))
+			if (fileName.StartsWith("lib", StringComparison.OrdinalIgnoreCase) && fileName.EndsWith(".so", StringComparison.OrdinalIgnoreCase))
 			{
 				return true;
 			}
-			if (FileName.StartsWith("lib", StringComparison.OrdinalIgnoreCase) && FileName.Contains(".so.", StringComparison.OrdinalIgnoreCase))
+			if (fileName.StartsWith("lib", StringComparison.OrdinalIgnoreCase) && fileName.Contains(".so.", StringComparison.OrdinalIgnoreCase))
 			{
 				// e.g. a Unix shared library with a version number suffix.
 				return true;
 			}
-			if (FileName.Equals("plugInfo.json", StringComparison.OrdinalIgnoreCase))
+			if (fileName.Equals("plugInfo.json", StringComparison.OrdinalIgnoreCase))
 			{
 				return true;
 			}
-			if ((FileName.Equals("info.plist", StringComparison.OrdinalIgnoreCase) || FileName.Equals("coderesources", StringComparison.OrdinalIgnoreCase)) && LocalFile.FullName.Contains(".app/", StringComparison.OrdinalIgnoreCase))
+			if ((fileName.Equals("info.plist", StringComparison.OrdinalIgnoreCase) || fileName.Equals("coderesources", StringComparison.OrdinalIgnoreCase)) && localFile.FullName.Contains(".app/", StringComparison.OrdinalIgnoreCase))
 			{
 				// xcode can generate plist files and coderesources differently in different stages of compile/cook/stage/package/etc. only allow ones inside a .app bundle
 				return true;
 			}
 			return false;
 		}
-		
+
 		/// <summary>
 		/// Copy a temp storage .zip file to directory.
 		/// Uses Robocopy on Windows, rsync on Linux/macOS and native .NET API when under Wine.
 		/// </summary>
-		/// <param name="ZipFile">.zip file to copy</param>
-		/// <param name="RootDir">Destination dir</param>
+		/// <param name="zipFile">.zip file to copy</param>
+		/// <param name="rootDir">Destination dir</param>
 		/// <returns>Output from the copy operation</returns>
-		static string CopyFile(FileInfo ZipFile, DirectoryReference RootDir)
+		static string CopyFile(FileInfo zipFile, DirectoryReference rootDir)
 		{
 			if (BuildHostPlatform.Current.IsRunningOnWine())
 			{
-				string sourceFile = ZipFile.FullName;
-				string destFile = Path.Join(RootDir.FullName, ZipFile.Name);
+				string sourceFile = zipFile.FullName;
+				string destFile = Path.Join(rootDir.FullName, zipFile.Name);
 				File.Copy(sourceFile, destFile, true);
 				return $".NET file copy. Source='{sourceFile}' Dest='{destFile}'";
 			}
 			else if (HostPlatform.Current.HostEditorPlatform == UnrealTargetPlatform.Win64)
 			{
-				return CommandUtils.RunAndLog(GetRoboCopyExe(), $"\"{ZipFile.DirectoryName}\" \"{RootDir}\" \"{ZipFile.Name}\" /w:5 /r:10", MaxSuccessCode: 3, Options: CommandUtils.ERunOptions.AppMustExist | CommandUtils.ERunOptions.NoLoggingOfRunCommand);
+				return CommandUtils.RunAndLog(GetRoboCopyExe(), $"\"{zipFile.DirectoryName}\" \"{rootDir}\" \"{zipFile.Name}\" /w:5 /r:10", MaxSuccessCode: 3, Options: CommandUtils.ERunOptions.AppMustExist | CommandUtils.ERunOptions.NoLoggingOfRunCommand);
 			}
 			else
 			{
-				return CommandUtils.RunAndLog("rsync", $"-v \"{ZipFile.DirectoryName}/{ZipFile.Name}\" \"{RootDir}/{ZipFile.Name}\"", Options: CommandUtils.ERunOptions.AppMustExist | CommandUtils.ERunOptions.NoLoggingOfRunCommand);
+				return CommandUtils.RunAndLog("rsync", $"-v \"{zipFile.DirectoryName}/{zipFile.Name}\" \"{rootDir}/{zipFile.Name}\"", Options: CommandUtils.ERunOptions.AppMustExist | CommandUtils.ERunOptions.NoLoggingOfRunCommand);
 			}
 		}
 
 		/// <summary>
 		/// Copy a directory using the most suitable option
 		/// </summary>
-		/// <param name="SourceDir">Source directory</param>
-		/// <param name="DestinationDir">Directory directory</param>
+		/// <param name="sourceDir">Source directory</param>
+		/// <param name="destinationDir">Directory directory</param>
 		/// <returns>Output from directory copy operation</returns>
 		/// <exception cref="DirectoryNotFoundException">If source directory wasn't found</exception>
-		static string CopyDirectory(DirectoryReference SourceDir, DirectoryReference DestinationDir)
+		static string CopyDirectory(DirectoryReference sourceDir, DirectoryReference destinationDir)
 		{
 			if (BuildHostPlatform.Current.IsRunningOnWine())
 			{
-				return CopyDirectoryDotNet(SourceDir, DestinationDir);
+				return CopyDirectoryDotNet(sourceDir, destinationDir);
 			}
 			else
 			{
-				return CopyDirectoryExternalTool(SourceDir, DestinationDir);
+				return CopyDirectoryExternalTool(sourceDir, destinationDir);
 			}
 		}
 
@@ -1398,32 +1395,32 @@ namespace AutomationTool
 		/// Copy a directory using an external tool native to OS
 		/// Robocopy Windows and rsync for Linux and macOS.
 		/// </summary>
-		/// <param name="SourceDir">Source directory</param>
-		/// <param name="DestinationDir">Directory directory</param>
+		/// <param name="sourceDir">Source directory</param>
+		/// <param name="destinationDir">Directory directory</param>
 		/// <returns>Output from directory copy operation</returns>
 		/// <exception cref="DirectoryNotFoundException">If source directory wasn't found</exception>
-		static string CopyDirectoryExternalTool(DirectoryReference SourceDir, DirectoryReference DestinationDir)
+		static string CopyDirectoryExternalTool(DirectoryReference sourceDir, DirectoryReference destinationDir)
 		{
 			if (HostPlatform.Current.HostEditorPlatform == UnrealTargetPlatform.Win64)
 			{
-				return CommandUtils.RunAndLog(GetRoboCopyExe(), $"\"{SourceDir}\" \"{DestinationDir}\" * /S /w:5 /r:10", MaxSuccessCode: 3, Options: CommandUtils.ERunOptions.AppMustExist | CommandUtils.ERunOptions.NoLoggingOfRunCommand);
+				return CommandUtils.RunAndLog(GetRoboCopyExe(), $"\"{sourceDir}\" \"{destinationDir}\" * /S /w:5 /r:10", MaxSuccessCode: 3, Options: CommandUtils.ERunOptions.AppMustExist | CommandUtils.ERunOptions.NoLoggingOfRunCommand);
 			}
 			else
 			{
-				return CommandUtils.RunAndLog("rsync", $"-vam --include=\"**\" \"{SourceDir}/\" \"{DestinationDir}/\"", Options: CommandUtils.ERunOptions.AppMustExist | CommandUtils.ERunOptions.NoLoggingOfRunCommand);
+				return CommandUtils.RunAndLog("rsync", $"-vam --include=\"**\" \"{sourceDir}/\" \"{destinationDir}/\"", Options: CommandUtils.ERunOptions.AppMustExist | CommandUtils.ERunOptions.NoLoggingOfRunCommand);
 			}
 		}
-		
+
 		/// <summary>
 		/// Copy a directory using .NET SDK directory and file copy API
 		/// </summary>
-		/// <param name="SourceDir">Source directory</param>
-		/// <param name="DestinationDir">Directory directory</param>
+		/// <param name="sourceDir">Source directory</param>
+		/// <param name="destinationDir">Directory directory</param>
 		/// <returns>Output from directory copy operation</returns>
 		/// <exception cref="DirectoryNotFoundException">If source directory wasn't found</exception>
-		static string CopyDirectoryDotNet(DirectoryReference SourceDir, DirectoryReference DestinationDir)
+		static string CopyDirectoryDotNet(DirectoryReference sourceDir, DirectoryReference destinationDir)
 		{
-			DirectoryInfo dir = new (SourceDir.FullName);
+			DirectoryInfo dir = new(sourceDir.FullName);
 
 			if (!dir.Exists)
 			{
@@ -1431,17 +1428,17 @@ namespace AutomationTool
 			}
 
 			DirectoryInfo[] dirs = dir.GetDirectories();
-			Directory.CreateDirectory(DestinationDir.FullName);
+			Directory.CreateDirectory(destinationDir.FullName);
 
 			foreach (FileInfo file in dir.GetFiles())
 			{
-				string targetFilePath = Path.Combine(DestinationDir.FullName, file.Name);
+				string targetFilePath = Path.Combine(destinationDir.FullName, file.Name);
 				file.CopyTo(targetFilePath);
 			}
 
 			foreach (DirectoryInfo subDir in dirs)
 			{
-				string newDestinationDir = Path.Combine(DestinationDir.FullName, subDir.Name);
+				string newDestinationDir = Path.Combine(destinationDir.FullName, subDir.Name);
 				CopyDirectoryDotNet(new DirectoryReference(subDir.FullName), new DirectoryReference(newDestinationDir));
 			}
 
@@ -1450,18 +1447,18 @@ namespace AutomationTool
 
 		static string GetRoboCopyExe()
 		{
-			var Result = CommandUtils.CombinePaths(Environment.SystemDirectory, "robocopy.exe");
-			if (!CommandUtils.FileExists(Result))
+			var result = CommandUtils.CombinePaths(Environment.SystemDirectory, "robocopy.exe");
+			if (!CommandUtils.FileExists(result))
 			{
 				// Use Regex.Replace so we can do a case-insensitive replacement of System32
-				var SysNativeDirectory = Regex.Replace(Environment.SystemDirectory, "System32", "Sysnative", RegexOptions.IgnoreCase);
-				var SysNativeExe = CommandUtils.CombinePaths(SysNativeDirectory, "robocopy.exe");
-				if (CommandUtils.FileExists(SysNativeExe))
+				var sysNativeDirectory = Regex.Replace(Environment.SystemDirectory, "System32", "Sysnative", RegexOptions.IgnoreCase);
+				var sysNativeExe = CommandUtils.CombinePaths(sysNativeDirectory, "robocopy.exe");
+				if (CommandUtils.FileExists(sysNativeExe))
 				{
-					Result = SysNativeExe;
+					result = sysNativeExe;
 				}
 			}
-			return Result;
+			return result;
 		}
 	}
 
@@ -1476,122 +1473,122 @@ namespace AutomationTool
 		public override void ExecuteBuild()
 		{
 			// Get all the shared directories
-			DirectoryReference RootDir = new DirectoryReference(CommandUtils.CmdEnv.LocalRoot);
+			DirectoryReference rootDir = new DirectoryReference(CommandUtils.CmdEnv.LocalRoot);
 
-			DirectoryReference LocalDir = DirectoryReference.Combine(RootDir, "Engine", "Saved", "TestTempStorage-Local");
-			CommandUtils.CreateDirectory(LocalDir);
-			CommandUtils.DeleteDirectoryContents(LocalDir.FullName);
+			DirectoryReference localDir = DirectoryReference.Combine(rootDir, "Engine", "Saved", "TestTempStorage-Local");
+			CommandUtils.CreateDirectory(localDir);
+			CommandUtils.DeleteDirectoryContents(localDir.FullName);
 
-			DirectoryReference SharedDir = DirectoryReference.Combine(RootDir, "Engine", "Saved", "TestTempStorage-Shared");
-			CommandUtils.CreateDirectory(SharedDir);
-			CommandUtils.DeleteDirectoryContents(SharedDir.FullName);
+			DirectoryReference sharedDir = DirectoryReference.Combine(rootDir, "Engine", "Saved", "TestTempStorage-Shared");
+			CommandUtils.CreateDirectory(sharedDir);
+			CommandUtils.DeleteDirectoryContents(sharedDir.FullName);
 
-			DirectoryReference WorkingDir = DirectoryReference.Combine(RootDir, "Engine", "Saved", "TestTempStorage-Working");
-			CommandUtils.CreateDirectory(WorkingDir);
-			CommandUtils.DeleteDirectoryContents(WorkingDir.FullName);
+			DirectoryReference workingDir = DirectoryReference.Combine(rootDir, "Engine", "Saved", "TestTempStorage-Working");
+			CommandUtils.CreateDirectory(workingDir);
+			CommandUtils.DeleteDirectoryContents(workingDir.FullName);
 
 			// Create the temp storage object
-			TempStorage TempStore = new TempStorage(WorkingDir, LocalDir, SharedDir, true);
+			TempStorage tempStore = new TempStorage(workingDir, localDir, sharedDir, true);
 
 			// Create a working directory, and copy some source files into it
-			DirectoryReference SourceDir = DirectoryReference.Combine(RootDir, "Engine", "Source", "Runtime");
-			if(!CommandUtils.CopyDirectory_NoExceptions(SourceDir.FullName, WorkingDir.FullName, true))
+			DirectoryReference sourceDir = DirectoryReference.Combine(rootDir, "Engine", "Source", "Runtime");
+			if (!CommandUtils.CopyDirectory_NoExceptions(sourceDir.FullName, workingDir.FullName, true))
 			{
-				throw new AutomationException("Couldn't copy {0} to {1}", SourceDir.FullName, WorkingDir.FullName);
+				throw new AutomationException("Couldn't copy {0} to {1}", sourceDir.FullName, workingDir.FullName);
 			}
 
 			// Save the default output
-			Dictionary<FileReference, DateTime> DefaultOutput = SelectFiles(WorkingDir, 'a', 'f');
-			TempStore.Archive("TestNode", null, DefaultOutput.Keys.ToArray(), false);
+			Dictionary<FileReference, DateTime> defaultOutput = SelectFiles(workingDir, 'a', 'f');
+			tempStore.Archive("TestNode", null, defaultOutput.Keys.ToArray(), false);
+
+			Dictionary<FileReference, DateTime> namedOutput = SelectFiles(workingDir, 'g', 'i');
+			tempStore.Archive("TestNode", "NamedOutput", namedOutput.Keys.ToArray(), true);
 
-			Dictionary<FileReference, DateTime> NamedOutput = SelectFiles(WorkingDir, 'g', 'i');
-			TempStore.Archive("TestNode", "NamedOutput", NamedOutput.Keys.ToArray(), true);
-			
 			// Check both outputs are still ok
-			TempStorageManifest DefaultManifest = TempStore.Retrieve("TestNode", null, new FileFilter());
-			CheckManifest(WorkingDir, DefaultManifest, DefaultOutput);
+			TempStorageManifest defaultManifest = tempStore.Retrieve("TestNode", null, new FileFilter());
+			CheckManifest(workingDir, defaultManifest, defaultOutput);
 
-			TempStorageManifest NamedManifest = TempStore.Retrieve("TestNode", "NamedOutput", new FileFilter());
-			CheckManifest(WorkingDir, NamedManifest, NamedOutput);
+			TempStorageManifest namedManifest = tempStore.Retrieve("TestNode", "NamedOutput", new FileFilter());
+			CheckManifest(workingDir, namedManifest, namedOutput);
 
 			// Delete local temp storage and the working directory and try again
 			Logger.LogInformation("Clearing local folders...");
-			CommandUtils.DeleteDirectoryContents(WorkingDir.FullName);
-			CommandUtils.DeleteDirectoryContents(LocalDir.FullName);
+			CommandUtils.DeleteDirectoryContents(workingDir.FullName);
+			CommandUtils.DeleteDirectoryContents(localDir.FullName);
 
 			// First output should fail
 			Logger.LogInformation("Checking default manifest is now unavailable...");
-			bool bGotManifest = false;
+			bool gotManifest;
 			try
 			{
-				TempStore.Retrieve("TestNode", null, new FileFilter());
-				bGotManifest = true;
+				tempStore.Retrieve("TestNode", null, new FileFilter());
+				gotManifest = true;
 			}
 			catch
 			{
-				bGotManifest = false;
+				gotManifest = false;
 			}
-			if(bGotManifest)
+			if (gotManifest)
 			{
 				throw new AutomationException("Did not expect shared temp storage manifest to exist");
 			}
 
 			// Second one should be fine
-			TempStorageManifest NamedManifestFromShared = TempStore.Retrieve("TestNode", "NamedOutput", new FileFilter());
-			CheckManifest(WorkingDir, NamedManifestFromShared, NamedOutput);
+			TempStorageManifest namedManifestFromShared = tempStore.Retrieve("TestNode", "NamedOutput", new FileFilter());
+			CheckManifest(workingDir, namedManifestFromShared, namedOutput);
 		}
 
 		/// <summary>
 		/// Enumerate all the files beginning with a letter within a certain range
 		/// </summary>
-		/// <param name="SourceDir">The directory to read from</param>
-		/// <param name="CharRangeBegin">First character in the range to files to return</param>
-		/// <param name="CharRangeEnd">Last character (inclusive) in the range of files to return</param>
+		/// <param name="sourceDir">The directory to read from</param>
+		/// <param name="charRangeBegin">First character in the range to files to return</param>
+		/// <param name="charRangeEnd">Last character (inclusive) in the range of files to return</param>
 		/// <returns>Mapping from filename to timestamp</returns>
-		static Dictionary<FileReference, DateTime> SelectFiles(DirectoryReference SourceDir, char CharRangeBegin, char CharRangeEnd)
+		static Dictionary<FileReference, DateTime> SelectFiles(DirectoryReference sourceDir, char charRangeBegin, char charRangeEnd)
 		{
-			Dictionary<FileReference, DateTime> ArchiveFileToTime = new Dictionary<FileReference,DateTime>();
-			foreach(FileInfo FileInfo in new DirectoryInfo(SourceDir.FullName).EnumerateFiles("*", SearchOption.AllDirectories))
+			Dictionary<FileReference, DateTime> archiveFileToTime = new Dictionary<FileReference, DateTime>();
+			foreach (FileInfo fileInfo in new DirectoryInfo(sourceDir.FullName).EnumerateFiles("*", SearchOption.AllDirectories))
 			{
-				char FirstCharacter = Char.ToLower(FileInfo.Name[0]);
-				if(FirstCharacter >= CharRangeBegin && FirstCharacter <= CharRangeEnd)
+				char firstCharacter = Char.ToLower(fileInfo.Name[0]);
+				if (firstCharacter >= charRangeBegin && firstCharacter <= charRangeEnd)
 				{
-					ArchiveFileToTime.Add(new FileReference(FileInfo), FileInfo.LastWriteTimeUtc);
+					archiveFileToTime.Add(new FileReference(fileInfo), fileInfo.LastWriteTimeUtc);
 				}
 			}
-			return ArchiveFileToTime;
+			return archiveFileToTime;
 		}
 
 		/// <summary>
 		/// Checks that a manifest matches the files on disk
 		/// </summary>
-		/// <param name="RootDir">Root directory for relative paths in the manifest</param>
-		/// <param name="Manifest">Manifest to check</param>
-		/// <param name="Files">Mapping of filename to timestamp as expected in the manifest</param>
-		static void CheckManifest(DirectoryReference RootDir, TempStorageManifest Manifest, Dictionary<FileReference, DateTime> Files)
+		/// <param name="rootDir">Root directory for relative paths in the manifest</param>
+		/// <param name="manifest">Manifest to check</param>
+		/// <param name="files">Mapping of filename to timestamp as expected in the manifest</param>
+		static void CheckManifest(DirectoryReference rootDir, TempStorageManifest manifest, Dictionary<FileReference, DateTime> files)
 		{
-			if(Files.Count != Manifest.Files.Length)
+			if (files.Count != manifest.Files.Length)
 			{
 				throw new AutomationException("Number of files in manifest does not match");
 			}
-			foreach(TempStorageFile ManifestFile in Manifest.Files)
+			foreach (TempStorageFile manifestFile in manifest.Files)
 			{
-				FileReference File = ManifestFile.ToFileReference(RootDir);
-				if(!FileReference.Exists(File))
+				FileReference file = manifestFile.ToFileReference(rootDir);
+				if (!FileReference.Exists(file))
 				{
 					throw new AutomationException("File in manifest does not exist");
 				}
 
-				DateTime OriginalTime;
-				if(!Files.TryGetValue(File, out OriginalTime))
+				DateTime originalTime;
+				if (!files.TryGetValue(file, out originalTime))
 				{
 					throw new AutomationException("File in manifest did not exist previously");
 				}
 
-				double DiffSeconds = (new FileInfo(File.FullName).LastWriteTimeUtc - OriginalTime).TotalSeconds;
-				if(Math.Abs(DiffSeconds) > 2)
+				double diffSeconds = (new FileInfo(file.FullName).LastWriteTimeUtc - originalTime).TotalSeconds;
+				if (Math.Abs(diffSeconds) > 2)
 				{
-					throw new AutomationException("Incorrect timestamp for {0}", ManifestFile.RelativePath);
+					throw new AutomationException("Incorrect timestamp for {0}", manifestFile.RelativePath);
 				}
 			}
 		}
@@ -1609,120 +1606,120 @@ namespace AutomationTool
 		/// </summary>
 		public override void ExecuteBuild()
 		{
-			string TempStorageDir = ParseParamValue("TempStorageDir", null);
-			if (TempStorageDir == null)
+			string tempStorageDir = ParseParamValue("TempStorageDir", null);
+			if (tempStorageDir == null)
 			{
 				throw new AutomationException("Missing -TempStorageDir parameter");
 			}
 
-			if (!Directory.Exists(TempStorageDir))
+			if (!Directory.Exists(tempStorageDir))
 			{
-				Logger.LogInformation("Temp Storage folder '{TempStorageDir}' does not exist, no work to do.", TempStorageDir);
+				Logger.LogInformation("Temp Storage folder '{TempStorageDir}' does not exist, no work to do.", tempStorageDir);
 				return;
 			}
 
-			string Days = ParseParamValue("Days", null);
-			if (Days == null)
+			string days = ParseParamValue("Days", null);
+			if (days == null)
 			{
 				throw new AutomationException("Missing -Days parameter");
 			}
 
-			double DaysValue;
-			if (!Double.TryParse(Days, out DaysValue))
+			double daysValue;
+			if (!Double.TryParse(days, out daysValue))
 			{
-				throw new AutomationException("'{0}' is not a valid value for the -Days parameter", Days);
+				throw new AutomationException("'{0}' is not a valid value for the -Days parameter", days);
 			}
 
-			DateTime RetainTime = DateTime.UtcNow - TimeSpan.FromDays(DaysValue);
+			DateTime retainTime = DateTime.UtcNow - TimeSpan.FromDays(daysValue);
 
 			// Enumerate all the build directories
-			Logger.LogInformation("Scanning {TempStorageDir}...", TempStorageDir);
-			int NumBuilds = 0;
-			List<DirectoryInfo> BuildsToDelete = new List<DirectoryInfo>();
-			foreach (DirectoryInfo StreamDirectory in new DirectoryInfo(TempStorageDir).EnumerateDirectories().OrderBy(x => x.Name))
+			Logger.LogInformation("Scanning {TempStorageDir}...", tempStorageDir);
+			int numBuilds = 0;
+			List<DirectoryInfo> buildsToDelete = new List<DirectoryInfo>();
+			foreach (DirectoryInfo streamDirectory in new DirectoryInfo(tempStorageDir).EnumerateDirectories().OrderBy(x => x.Name))
 			{
-				Logger.LogInformation("Scanning {Path}...", StreamDirectory.FullName);
+				Logger.LogInformation("Scanning {Path}...", streamDirectory.FullName);
 				try
 				{
-					foreach (DirectoryInfo BuildDirectory in StreamDirectory.EnumerateDirectories())
+					foreach (DirectoryInfo buildDirectory in streamDirectory.EnumerateDirectories())
 					{
 						try
 						{
-							if (!BuildDirectory.EnumerateFiles("*", SearchOption.AllDirectories).Any(x => x.LastWriteTimeUtc > RetainTime) &&
-								!BuildDirectory.EnumerateDirectories("*", SearchOption.AllDirectories).Any(x => x.LastWriteTimeUtc > RetainTime))
+							if (!buildDirectory.EnumerateFiles("*", SearchOption.AllDirectories).Any(x => x.LastWriteTimeUtc > retainTime) &&
+								!buildDirectory.EnumerateDirectories("*", SearchOption.AllDirectories).Any(x => x.LastWriteTimeUtc > retainTime))
 							{
-								BuildsToDelete.Add(BuildDirectory);
+								buildsToDelete.Add(buildDirectory);
 							}
-							NumBuilds++;
+							numBuilds++;
 						}
-						catch (Exception Ex)
+						catch (Exception ex)
 						{
-							Logger.LogError(Ex, "Exception while trying to scan files under {BuildDirectory}: {Ex}", BuildDirectory, Ex);
+							Logger.LogError(ex, "Exception while trying to scan files under {BuildDirectory}: {Ex}", buildDirectory, ex);
 						}
 					}
 				}
-				catch (Exception Ex)
+				catch (Exception ex)
 				{
-					Logger.LogError(Ex, "Exception while trying to scan {StreamDirectory}: {Ex}", StreamDirectory, Ex);
+					Logger.LogError(ex, "Exception while trying to scan {StreamDirectory}: {Ex}", streamDirectory, ex);
 				}
 			}
-			Logger.LogInformation("Found {NumBuilds} builds; {Count} to delete.", NumBuilds, BuildsToDelete.Count);
+			Logger.LogInformation("Found {NumBuilds} builds; {Count} to delete.", numBuilds, buildsToDelete.Count);
 
 			// Loop through them all, checking for files older than the delete time
-			int Idx = BuildsToDelete.Count;
-			while (Idx-- > 0)
+			int idx = buildsToDelete.Count;
+			while (idx-- > 0)
 			{
 				try
 				{
 					// Done if something already cleaned up this folder.
-					if (!BuildsToDelete[Idx].Exists)
+					if (!buildsToDelete[idx].Exists)
 					{
 						continue;
 					}
 
 					// Check if there is a marker file, if so skip this folder unless it has been twenty minutes.
-					FileInfo? DeleteInProgressFile = BuildsToDelete[Idx].GetFiles("DeleteInProgress.tmp").FirstOrDefault();
-					if (DeleteInProgressFile != null && DeleteInProgressFile.LastWriteTimeUtc < (DateTimeOffset.UtcNow - TimeSpan.FromMinutes(20)))
+					FileInfo? deleteInProgressFile = buildsToDelete[idx].GetFiles("DeleteInProgress.tmp").FirstOrDefault();
+					if (deleteInProgressFile != null && deleteInProgressFile.LastWriteTimeUtc < (DateTimeOffset.UtcNow - TimeSpan.FromMinutes(20)))
 					{
-						Logger.LogInformation("[{Index}/{Total}] {Path} flagged as delete in progress, skipping...", BuildsToDelete.Count - Idx, BuildsToDelete.Count, BuildsToDelete[Idx].FullName);
+						Logger.LogInformation("[{Index}/{Total}] {Path} flagged as delete in progress, skipping...", buildsToDelete.Count - idx, buildsToDelete.Count, buildsToDelete[idx].FullName);
 						continue;
 					}
 
-					File.WriteAllBytes(Path.Combine(BuildsToDelete[Idx].FullName, "DeleteInProgress.tmp"), Array.Empty<byte>());
-					Logger.LogInformation("[{Index}/{Total}] Deleting {Path}...", BuildsToDelete.Count - Idx, BuildsToDelete.Count, BuildsToDelete[Idx].FullName);
-					BuildsToDelete[Idx].Delete(true);
+					File.WriteAllBytes(Path.Combine(buildsToDelete[idx].FullName, "DeleteInProgress.tmp"), Array.Empty<byte>());
+					Logger.LogInformation("[{Index}/{Total}] Deleting {Path}...", buildsToDelete.Count - idx, buildsToDelete.Count, buildsToDelete[idx].FullName);
+					buildsToDelete[idx].Delete(true);
 				}
-				catch (Exception Ex)
+				catch (Exception ex)
 				{
-					Logger.LogWarning("Failed to delete old manifest folder; will try one file at a time: {Ex}", Ex);
-					CommandUtils.DeleteDirectory_NoExceptions(true, BuildsToDelete[Idx].FullName);
+					Logger.LogWarning("Failed to delete old manifest folder; will try one file at a time: {Ex}", ex);
+					CommandUtils.DeleteDirectory_NoExceptions(true, buildsToDelete[idx].FullName);
 				}
 			}
 
 			// Try to delete any empty branch folders
-			foreach (DirectoryInfo StreamDirectory in new DirectoryInfo(TempStorageDir).EnumerateDirectories())
+			foreach (DirectoryInfo streamDirectory in new DirectoryInfo(tempStorageDir).EnumerateDirectories())
 			{
 				try
 				{
-					if (!StreamDirectory.EnumerateDirectories().Any() && !StreamDirectory.EnumerateFiles().Any())
+					if (!streamDirectory.EnumerateDirectories().Any() && !streamDirectory.EnumerateFiles().Any())
 					{
 						try
 						{
-							StreamDirectory.Delete();
+							streamDirectory.Delete();
 						}
 						catch (IOException)
 						{
 							// only catch "directory is not empty type exceptions, if possible. Best we can do is check for IOException.
 						}
-						catch (Exception Ex)
+						catch (Exception ex)
 						{
-							Logger.LogWarning("Unexpected failure trying to delete (potentially empty) stream directory {Path}: {Ex}", StreamDirectory.FullName, Ex);
+							Logger.LogWarning("Unexpected failure trying to delete (potentially empty) stream directory {Path}: {Ex}", streamDirectory.FullName, ex);
 						}
 					}
 				}
-				catch (Exception Ex)
+				catch (Exception ex)
 				{
-					Logger.LogError(Ex, "Exception while trying to delete {StreamDirectory}: {Ex}", StreamDirectory, Ex);
+					Logger.LogError(ex, "Exception while trying to delete {StreamDirectory}: {Ex}", streamDirectory, ex);
 				}
 			}
 		}
