BuildGraph: Add a Multiline attribute to property attributes, allowing the inner text to be captured as a single multi-line string.

#rnx

diff --git a/Engine/Source/Programs/AutomationTool/BuildGraph/BgSchema.cs b/Engine/Source/Programs/AutomationTool/BuildGraph/BgSchema.cs
--- a/Engine/Source/Programs/AutomationTool/BuildGraph/BgSchema.cs
+++ b/Engine/Source/Programs/AutomationTool/BuildGraph/BgSchema.cs
@@ -796,6 +796,7 @@ namespace AutomationTool
 			extension.Attributes.Add(CreateSchemaAttribute("Name", ScriptSchemaStandardType.Name, XmlSchemaUse.Required));
 			extension.Attributes.Add(CreateSchemaAttribute("Value", s_stringTypeName, XmlSchemaUse.Optional));
 			extension.Attributes.Add(CreateSchemaAttribute("Separator", s_stringTypeName, XmlSchemaUse.Optional));
+			extension.Attributes.Add(CreateSchemaAttribute("Multiline", s_stringTypeName, XmlSchemaUse.Optional));
 			extension.Attributes.Add(CreateSchemaAttribute("If", ScriptSchemaStandardType.BalancedString, XmlSchemaUse.Optional));
 			extension.Attributes.Add(CreateSchemaAttribute("CreateInParentScope", ScriptSchemaStandardType.Boolean, XmlSchemaUse.Optional));
 
diff --git a/Engine/Source/Programs/AutomationTool/BuildGraph/BgScriptReader.cs b/Engine/Source/Programs/AutomationTool/BuildGraph/BgScriptReader.cs
--- a/Engine/Source/Programs/AutomationTool/BuildGraph/BgScriptReader.cs
+++ b/Engine/Source/Programs/AutomationTool/BuildGraph/BgScriptReader.cs
@@ -814,25 +814,58 @@ namespace AutomationTool
 					string value = ReadAttribute(element, "Value");
 					if (element.HasChildNodes)
 					{
-						// Get the separator character
-						string separator = ";";
-						if (element.HasAttribute("Separator"))
+						StringBuilder builder = new StringBuilder(value);
+						if (ReadBooleanAttribute(element, "Multiline", false))
 						{
-							separator = ReadAttribute(element, "Separator");
-						}
+							// Get a hanging indent from the first line, and remove that whitespace from subsequent lines
+							int prefixLen = -1;
 
-						// Read the element content, and append each line to the value as a semicolon delimited list
-						StringBuilder builder = new StringBuilder(value);
-						foreach (string line in element.InnerText.Split('\n'))
+							List<string> lines = new List<string>(element.InnerText.Split('\n'));
+							foreach(string line in lines)
+							{
+								int whitespaceLen = 0;
+								while (whitespaceLen < line.Length && Char.IsWhiteSpace(line[whitespaceLen]))
+								{
+									whitespaceLen++;
+								}
+
+								if (prefixLen == -1)
+								{
+									if (whitespaceLen == line.Length)
+									{
+										continue;
+									}
+									else
+									{
+										prefixLen = whitespaceLen;
+									}
+								}
+
+								whitespaceLen = Math.Min(whitespaceLen, prefixLen);
+								builder.AppendLine(line.Substring(whitespaceLen));
+							}
+						}
+						else
 						{
-							string trimLine = ExpandProperties(element, line.Trim());
-							if (trimLine.Length > 0)
+							// Get the separator character
+							string separator = ";";
+							if (element.HasAttribute("Separator"))
+							{
+								separator = ReadAttribute(element, "Separator");
+							}
+
+							// Read the element content, and append each line to the value as a semicolon delimited list
+							foreach (string line in element.InnerText.Split('\n'))
 							{
-								if (builder.Length > 0)
+								string trimLine = ExpandProperties(element, line.Trim());
+								if (trimLine.Length > 0)
 								{
-									builder.Append(separator);
+									if (builder.Length > 0)
+									{
+										builder.Append(separator);
+									}
+									builder.Append(trimLine);
 								}
-								builder.Append(trimLine);
 							}
 						}
 						value = builder.ToString();
