Horde: Improve error message when user does not have permissions to view a secret.

#rnx

diff --git a/Engine/Source/Programs/AutomationTool/BuildGraph/Tasks/HordeGetSecretsTask.cs b/Engine/Source/Programs/AutomationTool/BuildGraph/Tasks/HordeGetSecretsTask.cs
--- a/Engine/Source/Programs/AutomationTool/BuildGraph/Tasks/HordeGetSecretsTask.cs
+++ b/Engine/Source/Programs/AutomationTool/BuildGraph/Tasks/HordeGetSecretsTask.cs
@@ -2,7 +2,9 @@
 
 using System;
 using System.Collections.Generic;
+using System.Diagnostics;
 using System.Linq;
+using System.Net.Http;
 using System.Threading.Tasks;
 using System.Xml;
 using EpicGames.Core;
@@ -114,7 +116,16 @@ namespace AutomationTool.Tasks
 
 					foreach ((SecretId secretId, List<ReplacementInfo> replacements) in secretToReplacementInfo)
 					{
-						GetSecretResponse secret = await hordeHttpClient.GetSecretAsync(secretId);
+						GetSecretResponse secret;
+						try
+						{
+							secret = await hordeHttpClient.GetSecretAsync(secretId);
+						}
+						catch (HttpRequestException ex) when (ex.StatusCode == System.Net.HttpStatusCode.Forbidden)
+						{
+							throw new AutomationException(ex, $"User does not have permissions to read {secretId}");
+						}
+
 						foreach (ReplacementInfo replacement in replacements)
 						{
 							string? value;
