[Gauntlet] Cleanup/Re-order of TargetDeviceiOS properties/functions

diff --git a/Engine/Source/Programs/AutomationTool/Gauntlet/Framework/Gauntlet.Utils.cs b/Engine/Source/Programs/AutomationTool/Gauntlet/Framework/Gauntlet.Utils.cs
--- a/Engine/Source/Programs/AutomationTool/Gauntlet/Framework/Gauntlet.Utils.cs
+++ b/Engine/Source/Programs/AutomationTool/Gauntlet/Framework/Gauntlet.Utils.cs
@@ -833,10 +833,9 @@ namespace Gauntlet
 				{
 					ConstructorInfo TypeConstructor = FoundType.GetConstructor(Type.EmptyTypes);
 
-					if (TypeConstructor != null)
+					if (TypeConstructor != null && !FoundType.IsAbstract)
 					{
 						InterfaceType NewInstance = TypeConstructor.Invoke(null) as InterfaceType;
-
 						ConstructedTypes.Add(NewInstance);
 					}
 				}
diff --git a/Engine/Source/Programs/AutomationTool/Gauntlet/Platform/IOS/Gauntlet.IOSBuildSource.cs b/Engine/Source/Programs/AutomationTool/Gauntlet/Platform/IOS/Gauntlet.IOSBuildSource.cs
--- a/Engine/Source/Programs/AutomationTool/Gauntlet/Platform/IOS/Gauntlet.IOSBuildSource.cs
+++ b/Engine/Source/Programs/AutomationTool/Gauntlet/Platform/IOS/Gauntlet.IOSBuildSource.cs
@@ -14,7 +14,7 @@ namespace Gauntlet
 {
 	public class AppleBuild : IBuild
 	{
-		public int PreferenceOrder { get { return 0; } }
+		public int PreferenceOrder => 0;
 
 		public UnrealTargetConfiguration Configuration { get; protected set; }
 
@@ -28,7 +28,7 @@ namespace Gauntlet
 
 		public BuildFlags Flags { get; protected set; }
 
-		public string Flavor { get { return ""; } }
+		public string Flavor => string.Empty;
 
 		public virtual UnrealTargetPlatform Platform { get; }
 
@@ -76,7 +76,7 @@ namespace Gauntlet
 		}
 
 		// There are issues with IPA Zip64 files being created with Ionic.Zip possibly limited to when running on mono (see IOSPlatform.PackageIPA)
-		// This manifests as header overflow errors, etc in 7zip, Ionic.Zip, System.IO.Compression, and OSX system unzip		
+		// This manifests as header overflow errors, etc in 7zip, Ionic.Zip, System.IO.Compression, and OSX system unzip
 		internal static bool ExecuteIPAZipCommand(String Arguments, out String Output, String ShouldExist = "")
 		{
 			using (new ScopedSuspendECErrorParsing())
@@ -95,7 +95,7 @@ namespace Gauntlet
 						}
 					}
 
-					Log.Info(String.Format("unzip encountered an issue procesing IPA, possibly due to Zip64. Future steps may fail."));
+					Log.Info("unzip encountered an issue procesing IPA, possibly due to Zip64. Future steps may fail.");
 				}
 			}
 
@@ -112,16 +112,16 @@ namespace Gauntlet
 
 				if (Result.ExitCode != 0)
 				{
-					if (!String.IsNullOrEmpty(ShouldExist))
+					if (!string.IsNullOrEmpty(ShouldExist))
 					{
 						if (!File.Exists(ShouldExist) && !Directory.Exists(ShouldExist))
 						{
-							Log.Error(String.Format("ditto encountered an error or warning procesing IPA, {0} missing", ShouldExist));
+							Log.Error("ditto encountered an error or warning procesing IPA, {ShouldExist} missing", ShouldExist);
 							return false;
 						}
 					}
 
-					Log.Error(String.Format("ditto encountered an issue procesing IPA"));
+					Log.Error("ditto encountered an issue procesing IPA");
 					return false;
 
 				}
@@ -132,55 +132,44 @@ namespace Gauntlet
 
 		private static PlistInfo GetPlistInfo(string Source)
 		{
-			PlistInfo Info = null;
 			bool IsIPAFile = Path.GetExtension(Source).Equals(".ipa", StringComparison.OrdinalIgnoreCase);
-
 			if (IsIPAFile)
 			{
-				string Output;
-
 				// Get a list of files in the IPA
-				if (!ExecuteIPAZipCommand(String.Format("-Z1 {0}", Source), out Output))
+				if (!ExecuteIPAZipCommand(string.Format("-Z1 {0}", Source), out string Output))
 				{
-					Log.Warning(String.Format("Unable to list files for IPA {0}", Source));
+					Log.Info("Unable to list files for IPA {IPAPath}", Source);
 					return null;
 				}
 
-
 				string[] Filenames = Regex.Split(Output, "\r\n|\r|\n");
-				string PList = Filenames.Where(F => Regex.IsMatch(F.ToLower().Trim(), @"(payload\/)([^\/]+)(\/info\.plist)")).FirstOrDefault();
-
-				if (String.IsNullOrEmpty(PList))
+				string PListFile = Filenames.Where(F => Regex.IsMatch(F.ToLower().Trim(), @"(payload\/)([^\/]+)(\/info\.plist)")).FirstOrDefault();
+				if (string.IsNullOrEmpty(PListFile))
 				{
-					Log.Warning(String.Format("Unable to find plist for IPA {0}", Source));
+					Log.Info("Unable to find plist for IPA {IPAPath}", Source);
 					return null;
 				}
 
 				// Get the plist info
-				if (!ExecuteIPAZipCommand(String.Format("-p '{0}' '{1}'", Source, PList), out Output))
+				if (!ExecuteIPAZipCommand(string.Format("-p '{0}' '{1}'", Source, PListFile), out Output))
 				{
-					Log.Warning(String.Format("Unable to extract plist data for IPA {0}", Source));
+					Log.Info("Unable to extract plist data for IPA {IPAPath}", Source);
 					return null;
 				}
 
-				Info = new PlistInfo(Output);
+				return new PlistInfo(Output);
 			}
 			else
 			{
-				// Find plist file
-				DirectoryInfo Di = new DirectoryInfo(Source);
-				FileInfo PlistFile = new FileInfo(Path.Combine(Di.FullName, "Info.plist"));
-				if (!PlistFile.Exists)
+				string PlistFile = Path.Combine(Source, "Info.plist");
+				if (!File.Exists(PlistFile))
 				{
-					Log.Warning(String.Format("Unable to find plist from {0}", Source));
+					Log.Info("Unable to find plist from {IPAPath}. Skipping.", Source);
 					return null;
 				}
 
-				StreamReader PListStream = new StreamReader(PlistFile.FullName);
-				Info = new PlistInfo(PListStream.ReadToEnd());
+				return new PlistInfo(File.ReadAllText(PlistFile));
 			}
-
-			return Info;
 		}
 
 		private class PlistInfo
@@ -263,142 +252,178 @@ namespace Gauntlet
 		public static T CreateFromPath<T>(string InProjectName, string InPath, AppleBuildSource<T> BuildSource)
 			where T : AppleBuild
 		{
-			T DiscoveredBuild = null;
-
-			DirectoryInfo Di = new DirectoryInfo(InPath);
+			FileSystemInfo BuildPath;
+			if (Directory.Exists(InPath))
+			{
+				BuildPath = new DirectoryInfo(InPath);
+			}
+			else if (File.Exists(InPath))
+			{
+				BuildPath = new FileInfo(InPath);
+			}
+			else
+			{
+				Log.Verbose("Build path does not exist! Skipping.");
+				return null;
+			}
 
-			var UnrealConfig = UnrealHelpers.GetConfigurationFromExecutableName(InProjectName, Di.Name);
+			UnrealTargetConfiguration Configuration = UnrealHelpers.GetConfigurationFromExecutableName(InProjectName, InPath);
+			if (Configuration == UnrealTargetConfiguration.Unknown)
+			{
+				Log.Verbose("Could not deduce iOS build configuration from build at path {BuildPath}. Skipping.", InPath);
+				return null;
+			}
 
-			if (UnrealConfig != UnrealTargetConfiguration.Unknown)
+			// Check there's an executable with the right name
+			string AppShortName = Regex.Replace(InProjectName, "Game", string.Empty, RegexOptions.IgnoreCase);
+			if (BuildPath is DirectoryInfo App)
 			{
-				// check there's an executable with the right name 
-				string ShortName = Regex.Replace(InProjectName, "Game", "", RegexOptions.IgnoreCase);
-				FileInfo Executable = new DirectoryInfo(InPath).GetFiles().Where(Fi => Fi.Name.StartsWith(ShortName, StringComparison.OrdinalIgnoreCase)).FirstOrDefault();
-				if (Executable != null)
+				if (App.GetFiles(AppShortName + '*', SearchOption.TopDirectoryOnly).FirstOrDefault() == null)
 				{
-					Log.Verbose("Pulling package data from {0}", Di.FullName);
-
-					string AbsPath = Di.FullName;
+					Log.Verbose("Could not find an executable within build path {BuildPath}. Skipping", InPath);
+					return null;
+				}
+			}
+			else
+			{
+				// Get a list of files in the IPA
+				if (!ExecuteIPAZipCommand(string.Format("-Z1 {0}", InPath), out string Output))
+				{
+					Log.Info("Unable to list files for IPA {IPAPath}", InPath);
+					return null;
+				}
 
-					// IOS builds are always packaged, and can always replace the command line and executable as we cache the unzip'd IPA
-					BuildFlags Flags = BuildFlags.Packaged | BuildFlags.CanReplaceCommandLine | BuildFlags.CanReplaceExecutable;
+				bool bFoundExecutable = false;
+				IEnumerable<string> FileNames = Regex.Split(Output, "\r\n|\r|\n").Select(File => Path.GetFileName(File));
 
-					if (AbsPath.Contains("Bulk"))
-					{
-						Flags |= BuildFlags.Bulk;
-					}
-					else
+				foreach(string File in FileNames)
+				{
+					if(File.Contains(AppShortName, StringComparison.OrdinalIgnoreCase))
 					{
-						Flags |= BuildFlags.NotBulk;
+						bFoundExecutable = true;
+						break;
 					}
+				}
 
-					PlistInfo Info = GetPlistInfo(Di.FullName);
-					if (Info != null)
-					{
-						IEnumerable<string> CFBundlePlatformNames = Info.GetAllValues("CFBundleSupportedPlatforms");
-						if (CFBundlePlatformNames != null && CFBundlePlatformNames.Contains(BuildSource.CFBundlePlatformName))
-						{
-							string PackageName = Info.GetFirstValue("CFBundleIdentifier");
+				if (!bFoundExecutable)
+				{
+					Log.Verbose("Could not find an executable within build path {BuildPath}. Skipping", InPath);
+					return null;
+				}
+			}
 
-							if (!String.IsNullOrEmpty(PackageName))
-							{
-								Dictionary<string, string> FilesToInstall = new Dictionary<string, string>();
+			Log.Verbose("Pulling package data from {BuildPath}", BuildPath);
 
-								DiscoveredBuild = Activator.CreateInstance(typeof(T), new object[] { UnrealConfig, PackageName, Di.FullName, FilesToInstall, Flags }) as T;
+			// IOS builds are always packaged, and can always replace the command line and executable (even as IPAs because we cache the unzipped app)
+			BuildFlags Flags = BuildFlags.Packaged | BuildFlags.CanReplaceCommandLine | BuildFlags.CanReplaceExecutable;
+			if (BuildPath.FullName.Contains("Bulk"))
+			{
+				Flags |= BuildFlags.Bulk;
+			}
+			else
+			{
+				Flags |= BuildFlags.NotBulk;
+			}
 
-								Log.Verbose("Found bundle id: {0}", PackageName);
-								Log.Verbose("Found {0} {1} build at {2}", UnrealConfig, ((Flags & BuildFlags.Bulk) == BuildFlags.Bulk) ? "(bulk)" : "(not bulk)", AbsPath);
-							}
-							else
-							{
-								Log.Warning(String.Format("Unable to find CFBundleIdentifier in plist info for App {0}", Di.FullName));
-							}
-						}
-						else
-						{
-							Log.Verbose("Unable to find matching platform '{0}' for CFBundleSupportedPlatforms in plist info for App {1}", BuildSource.CFBundlePlatformName, Di.FullName);
-						}
-					}
-				}
+			PlistInfo Info = GetPlistInfo(BuildPath.FullName);
+			if(Info == null)
+			{
+				Log.Info("Unable to parse PlistInfo for '{BuildPath}'. Skipping", BuildPath.FullName);
 			}
 
-			return DiscoveredBuild;
+			IEnumerable<string> CFBundlePlatformNames = Info.GetAllValues("CFBundleSupportedPlatforms");
+			if(CFBundlePlatformNames == null || !CFBundlePlatformNames.Contains(BuildSource.CFBundlePlatformName))
+			{
+				Log.Info("Unable to find matching platform '{BundlePlatform}' for CFBundleSupportedPlatforms in PlistInfo for App {BuildPath}. Skipping", BuildSource.CFBundlePlatformName, BuildPath.FullName);
+				return null;
+			}
+
+			string PackageName = Info.GetFirstValue("CFBundleIdentifier");
+			if (string.IsNullOrEmpty(PackageName))
+			{
+				Log.Info("Unable to find CFBundleIdentifier in PlistInfo for App {BuildPath}. Skipping.", BuildPath);
+				return null;
+			}
+
+			Log.Verbose("Found bundle id: {PackageName}", PackageName);
+			Log.Verbose("Found {Configuration} {Flags} build at {BuildPath}", Configuration, ((Flags & BuildFlags.Bulk) == BuildFlags.Bulk) ? "(bulk)" : "(not bulk)", BuildPath);
+
+			// Todo: Consider handling Bulk content copies here?
+			Dictionary<string, string> FilesToInstall = new Dictionary<string, string>();
+			return Activator.CreateInstance(typeof(T), new object[] { Configuration, PackageName, BuildPath.FullName, FilesToInstall, Flags }) as T;
 		}
 	}
 
 	public abstract class AppleBuildSource<T> : IFolderBuildSource
 		where T : AppleBuild
 	{
-		protected abstract UnrealTargetPlatform Platform { get; }
+		public string ProjectName { get; protected set; }
 
 		public abstract string CFBundlePlatformName { get; }
 
-		public string BuildName { get { return $"{Platform}BuildSource"; } }
+		public string BuildName => Platform.ToString() + "BuildSource";
+
+		protected abstract UnrealTargetPlatform Platform { get; }
+
+		protected string BuildFilter;
+
+		public AppleBuildSource()
+		{
+			BuildFilter = Globals.Params.ParseValue(Platform.ToString() + "BuildFilter", null);
+		}
 
 		public bool CanSupportPlatform(UnrealTargetPlatform InPlatform)
 		{
 			return InPlatform == Platform;
 		}
 
-		public string ProjectName { get; protected set; }
-
 		public virtual List<IBuild> GetBuildsAtPath(string InProjectName, string InPath, int MaxRecursion = 3)
 		{
 			// We only want iOS builds on Mac host
+			List<IBuild> Builds = new List<IBuild>();
 			if (BuildHostPlatform.Current.Platform != UnrealTargetPlatform.Mac)
 			{
 				return new List<IBuild>();
 			}
 
-			List<IBuild> Builds = new List<IBuild>();
-
-			// c:\path\to\build
-			DirectoryInfo PathDI = new DirectoryInfo(InPath);
+			// Interface default parameters don't let us modify the default if calling from an interface cast...
+			// IOS builds are often located deeper within a client directory, so increase the depth here
+			MaxRecursion = MaxRecursion > 5 ? MaxRecursion : 5;
 
-			if (PathDI.Exists)
+			DirectoryInfo BuildDirectory = new DirectoryInfo(InPath);
+			if (BuildDirectory.Exists)
 			{
 				List<DirectoryInfo> SearchDirs = new List<DirectoryInfo>();
 
-				if (PathDI.Name.IndexOf(Platform.ToString(), StringComparison.OrdinalIgnoreCase) >= 0)
-				{
-					SearchDirs.Add(PathDI);
-				}
-				else
-				{
-					// find all directories that begin with IOS
-					SearchDirs.AddRange(PathDI.GetDirectories($"{Platform}*", SearchOption.TopDirectoryOnly));
-				}
-
-				IEnumerable<DirectoryInfo> DirsToRecurse = new List<DirectoryInfo>(SearchDirs);
+				// Find the first folder in the build directory containing the platform name
+				DirectoryInfo[] SubDirectories = BuildDirectory.Name.Contains(Platform.ToString(), StringComparison.OrdinalIgnoreCase)
+					? new[] { BuildDirectory }
+					: BuildDirectory.GetDirectories(Platform.ToString() + '*', SearchOption.TopDirectoryOnly).ToArray();
 
-				List<DirectoryInfo> AllDirs = new List<DirectoryInfo>();
-
-				// get subdirs
+				// Now, recursively search the IOS directory for any .app folders
+				// We also discard any apps that don't include the IOSBuildFilter
+				List<FileSystemInfo> Apps = new();
 				while (MaxRecursion-- > 0)
 				{
-					IEnumerable<DirectoryInfo> DiscoveredDirs = DirsToRecurse.SelectMany(D => D.GetDirectories("*", SearchOption.TopDirectoryOnly));
-
-					// mac packages are folders so we only want things that end with .app
-					IEnumerable<DirectoryInfo> Packages = DiscoveredDirs.Where(D => Path.GetExtension(D.Name).Equals(".app", StringComparison.OrdinalIgnoreCase));
+					IEnumerable<DirectoryInfo> ValidApps = SubDirectories
+						.Where(Directory => Directory.Extension.Equals(".app", StringComparison.OrdinalIgnoreCase) || Directory.Extension.Equals(".ipa", StringComparison.OrdinalIgnoreCase))
+						.Where(Directory => string.IsNullOrEmpty(BuildFilter) || Directory.FullName.Contains(BuildFilter, StringComparison.OrdinalIgnoreCase));
 
-					AllDirs.AddRange(Packages);
+					IEnumerable<FileInfo> ValidIPAs = SubDirectories
+						.SelectMany(Directory => Directory.GetFiles("*.ipa", SearchOption.TopDirectoryOnly))
+						.Where(File => string.IsNullOrEmpty(BuildFilter) || File.FullName.Contains(BuildFilter, StringComparison.OrdinalIgnoreCase));
 
-					// don't recurse into the dirs that end with .app since those are builds we'll look at
-					DirsToRecurse = DiscoveredDirs.Except(Packages);
+					Apps.AddRange(ValidApps);
+					Apps.AddRange(ValidIPAs);
+					SubDirectories = SubDirectories.SelectMany(Directory => Directory.GetDirectories("*", SearchOption.TopDirectoryOnly)).ToArray();
 				}
 
-				string IOSBuildFilter = Globals.Params.ParseValue($"{Platform}BuildFilter", "");
-				foreach (DirectoryInfo Di in AllDirs)
+				foreach (FileSystemInfo App in Apps)
 				{
-					AppleBuild FoundBuild = AppleBuild.CreateFromPath<T>(InProjectName, Di.FullName, this);
-
-					if (FoundBuild != null)
+					AppleBuild Build = AppleBuild.CreateFromPath(InProjectName, App.FullName, this);
+					if (Build != null)
 					{
-						if (!string.IsNullOrEmpty(IOSBuildFilter) && FoundBuild.SourcePath.IndexOf(IOSBuildFilter, StringComparison.OrdinalIgnoreCase) >= 0)
-						{
-							continue;
-						}
-						Builds.Add(FoundBuild);
+						Builds.Add(Build);
 					}
 				}
 			}
diff --git a/Engine/Source/Programs/AutomationTool/Gauntlet/Platform/IOS/Gauntlet.TargetDeviceIOS.cs b/Engine/Source/Programs/AutomationTool/Gauntlet/Platform/IOS/Gauntlet.TargetDeviceIOS.cs
--- a/Engine/Source/Programs/AutomationTool/Gauntlet/Platform/IOS/Gauntlet.TargetDeviceIOS.cs
+++ b/Engine/Source/Programs/AutomationTool/Gauntlet/Platform/IOS/Gauntlet.TargetDeviceIOS.cs
@@ -23,182 +23,79 @@
 
 namespace Gauntlet
 {
-
-	class IOSAppInstance : IAppInstance
-	{
-		protected IOSAppInstall Install;
-		public IOSAppInstance(IOSAppInstall InInstall, IProcessResult InProcess, string InCommandLine)			
-		{
-			Install = InInstall;
-			this.CommandLine = InCommandLine;
-			this.ProcessResult = InProcess;		
-		}
-
-		public string ArtifactPath
-		{
-			get
-			{
-				if (bHaveSavedArtifacts == false)
-				{
-					if (HasExited)
-					{
-						SaveArtifacts();
-						bHaveSavedArtifacts = true;
-					}
-				}
-				
-				return Install.IOSDevice.LocalCachePath + "/" + Install.IOSDevice.DeviceArtifactPath;
-			}
-		}
-
-		public ITargetDevice Device
-		{
-			get
-			{
-				return Install.Device;
-			}
-		}
-
-		protected void SaveArtifacts()
-		{
-			TargetDeviceIOS Device = Install.IOSDevice;
-
-			// copy remote artifacts to local		
-			string CommandLine = String.Format("--bundle_id {0} --download={1} --to {2}", Install.PackageName, Device.DeviceArtifactPath, Device.LocalCachePath);
-
-			IProcessResult DownloadCmd = Device.ExecuteIOSDeployCommand(CommandLine, 120);
-
-			if (DownloadCmd.ExitCode != 0)
-			{
-				Log.Warning(KnownLogEvents.Gauntlet_DeviceEvent, "Failed to retrieve artifacts. {Output}", DownloadCmd.Output);
-			}
-			
-		}
-
-		public IProcessResult ProcessResult { get; private set; }
-
-		public bool HasExited { get { return ProcessResult.HasExited; } }
-
-		public bool WasKilled { get; protected set; }		
-
-		public int ExitCode { get { return ProcessResult.ExitCode; } }
-
-		public string CommandLine { get; private set; }
-
-		public string StdOut 
-		{  
-			get
-			{
-				if (HasExited)
-				{
-					// The ios application is being run under lldb by ios-deploy
-					// lldb catches crashes and we have it setup to dump thread callstacks
-					// parse any crash dumps into Unreal crash format and append to output
-					string CrashLog = LLDBCrashParser.GenerateCrashLog(ProcessResult.Output);
-
-					if (!string.IsNullOrEmpty(CrashLog))
-					{
-						return String.Format("{0}\n{1}", ProcessResult.Output, CrashLog);
-					}
-				}
-
-				return ProcessResult.Output;
-
-			}
-		}
-
-		public int WaitForExit()
-		{
-			if (!HasExited)
-			{
-				ProcessResult.WaitForExit();
-			}
-
-			return ExitCode;
-		}
-
-		public void Kill(bool GenerateDump)
-		{
-			if (!HasExited)
-			{
-				WasKilled = true;
-				ProcessResult.ProcessObject.Kill(true);
-			}
-		}
-
-
-		internal bool bHaveSavedArtifacts;		
-	}
-
-
-	class IOSAppInstall : IAppInstall
-	{
-		public string Name { get; protected set; }
-			
-		public string CommandLine { get; protected set; }
-
-        public string PackageName { get; protected set; }
-
-		public ITargetDevice Device { get { return IOSDevice; } }
-
-		public TargetDeviceIOS IOSDevice;
-
-		public IOSAppInstall(string InName, TargetDeviceIOS InDevice, string InPackageName, string InCommandLine)
-		{
-			Name = InName;			
-			CommandLine = InCommandLine;
-            PackageName = InPackageName;
-
-			IOSDevice = InDevice;
-		}
-
-		public IAppInstance Run()
-		{
-			return Device.Run(this);
-		}
-	}
-
-	public class IOSDeviceFactory : IDeviceFactory
-	{
-		public bool CanSupportPlatform(UnrealTargetPlatform? Platform)
-		{
-			return Platform == UnrealTargetPlatform.IOS;
-		}
-
-		public ITargetDevice CreateDevice(string InRef, string InCachePath, string InParam)
-		{
-			return new TargetDeviceIOS(InRef, InCachePath);
-		}
-	}
-
 	/// <summary>
 	/// iOS implementation of a device to run applications
 	/// </summary>
 	public class TargetDeviceIOS : ITargetDevice
 	{
+
+		#region ITargetDevice
 		public string Name { get; protected set; }
 
+		public UnrealTargetPlatform? Platform => UnrealTargetPlatform.IOS;
+
+		public CommandUtils.ERunOptions RunOptions { get; set; }
+
+		public bool IsConnected => Connected;
+
+		public bool IsOn => true;
+
+		public bool IsAvailable => true;
+		#endregion
+
 		/// <summary>
+		/// Temp path we use to push/pull things from the device
+		/// </summary>
+		public string LocalCachePath { get; protected set; }
+
+		/// <summary>
+		/// Artifact (e.g. Saved) path on the device
+		/// </summary>
+		public string DeviceArtifactPath { get; protected set; }
+
+		/// <summary>
 		/// Low-level device name (uuid)
 		/// </summary>
 		public string DeviceName { get; protected set; }
-		
+
 		protected Dictionary<EIntendedBaseCopyDirectory, string> LocalDirectoryMappings { get; set; }
 
+		private static object AppLock = new object();
+
+		private static bool ZombiesKilled = false;
+
+		private static Dictionary<string, bool> ConnectedDevices = new Dictionary<string, bool>();
+
+		private bool Connected = false;
+
+		private bool IsDefaultDevice = false;
+
+		// path to locally extracted (and possibly resigned) app bundle
+		// Note: ios-deploy works with app bundles, which requires the IPA be unzipped for deployment (this will allow us to resign in the future as well)
+		private string LocalAppBundle = null;
+
+		// Gauntlet cache folder for tracking device/ipa state
+		string GauntletAppCache => Path.Combine(Globals.TempDir, "IOSAppCache");
+
+		// the current App MD5 hash, which is tracked to avoid unneccessary deployments and unzip operations
+		private string AppHashFilename => Path.Combine(GauntletAppCache, "AppHash.txt");
+
+		// file whose presence signals that cache was resigned
+		private string CacheResignedFilename => Path.Combine(GauntletAppCache, "Resigned.txt");
+
 		public TargetDeviceIOS(string InName, string InCachePath = null)
 		{
-			KillZombies();
-
-			var DefaultDevices = GetConnectedDeviceUUID();			
-
-			IsDefaultDevice = (String.IsNullOrEmpty(InName) || InName.Equals("default", StringComparison.OrdinalIgnoreCase));
-
 			Name = InName;
 			LocalDirectoryMappings = new Dictionary<EIntendedBaseCopyDirectory, string>();
+			IsDefaultDevice = string.IsNullOrEmpty(InName) || InName.Equals("default", StringComparison.OrdinalIgnoreCase);
 
+			List<string> DefaultDevices = GetConnectedDeviceUUID();	
+
+			KillZombies();
+
 			// If no device name or its 'default' then use the first default device
 			if (IsDefaultDevice)
-			{				
+			{
 				if (DefaultDevices.Count() == 0)
 				{
 					throw new AutomationException("No default device available");
@@ -218,34 +115,22 @@
 			}
 
 			// setup local cache
-			LocalCachePath = InCachePath ?? Path.Combine(GauntletAppCache, "Device_" + Name);
+			LocalCachePath = InCachePath ?? Path.Combine(GauntletAppCache, "IOSDevice_" + DeviceName);
 		}
 
-		bool IsDefaultDevice = false;		
-
 		#region IDisposable Support
-		private bool disposedValue = false; // To detect redundant calls
+		private bool DisposedValue = false; // To detect redundant calls
+		~TargetDeviceIOS()
+		{
+			Dispose(false);
+		}
 
 		protected virtual void Dispose(bool disposing)
 		{
-			if (!disposedValue)
+			if (!DisposedValue)
 			{
-				try
-				{
-					if (Directory.Exists(LocalCachePath))
-					{
-						Directory.Delete(LocalCachePath, true);
-					}
-				}
-				catch (Exception Ex)
-				{
-					Log.Warning(KnownLogEvents.Gauntlet_DeviceEvent, "TargetDeviceIOS.Dispose() threw: {Exception}", Ex.Message);
-				}
-				finally
-				{
-					disposedValue = true;
-				}
-
+				// TODO: dispose managed state (managed objects).
+				DisposedValue = true;
 			}
 		}
 
@@ -258,406 +143,118 @@
 		}
 		#endregion
 
-		public CommandUtils.ERunOptions RunOptions { get; set; }
+		#region ITargetDevice
+		public bool Connect()
+		{
+			lock (Globals.MainLock)
+			{
+				if (Connected)
+				{
+					return true;
+				}
 
-		public IAppInstance Run(IAppInstall App)
-		{
-			IOSAppInstall IOSApp = App as IOSAppInstall;
+				bool ExistingConnection = false;
+				if (ConnectedDevices.TryGetValue(DeviceName, out ExistingConnection))
+				{
+					if (ExistingConnection)
+					{
+						throw new AutomationException("Connected to already connected device");
+					}
+				}
 
-			if (IOSApp == null)
-			{
-				throw new DeviceException("AppInstance is of incorrect type!");
+				ConnectedDevices[DeviceName] = true;
+
+				Connected = true;
 			}
 
-			string CommandLine = IOSApp.CommandLine.Replace("\"", "\\\\\"");
-			
-			Log.Info("Launching {0} on {1}", App.Name, ToString());
-			Log.Verbose("\t{0}", CommandLine);
+			return true;
+		}
 
-			// ios-deploy notes: -L launches detached, -I non-interactive  (exits when app exits), -r uninstalls before install (removes app Documents folder)
-			// -t <seconds> number of seconds to wait for device to be connected
+		public bool Disconnect(bool bForce = false)
+		{
+			lock (Globals.MainLock)
+			{
+				if (!Connected)
+				{
+					return true;
+				}
 
-			// check xcodebuild -version, to see if it larger than 15.0
-			bool bIsXcode15 = false;
-			IProcessResult XcodeVersionResult = CommandUtils.Run("xcodebuild", "-version");
-			if (XcodeVersionResult.ExitCode == 0)
-			{
-				// XcodeVersionResult.Output is: Xcode 15, try get version number
-				Match VersionMatch = Regex.Match(XcodeVersionResult.Output, @"Xcode\ (?<Version>\d+)");
-				if (VersionMatch.Success)
+				Connected = false;
+
+				if (ConnectedDevices.ContainsKey(DeviceName))
 				{
-					int Version = 0;
-					if (int.TryParse(VersionMatch.Groups["Version"].Value, out Version))
-					{
-						bIsXcode15 = Version >= 15;
-					}
-				}
-			}
-
-			IProcessResult Result;
-			if (bIsXcode15)
-			{
-				CommandLine = CommandLine.Replace("\\\\", "");
-				string CL = String.Format("device process launch --device {0} {1} '{2}'", DeviceName, IOSApp.PackageName, CommandLine.Trim());
-				Result = ExecuteIOSDeployCommand(CL, 0, true, false, true);
-				Thread.Sleep(5000);
-				// check if IOSApp.PackageName process running using devicectl device info processes
-				while(true)
-				{
-					CL = String.Format("device info processes --device {0}", DeviceName);
-					Result = ExecuteIOSDeployCommand(CL, 0, true, false, true);
-					// Get BundleName from IOSApp.PackageName
-					string BundleName = IOSApp.PackageName.Substring(IOSApp.PackageName.LastIndexOf('.') + 1);
-					Log.Info("Process running: " + BundleName);
-					// check if IOSApp.PackageName process not running
-					if (!Result.Output.Contains(BundleName))
-					{
-						break;
-					}
-					Thread.Sleep(5000);
+					ConnectedDevices.Remove(DeviceName);
 				}
 
 			}
-			else
+
+			return true;
+		}
+
+		public bool PowerOn()
+		{
+			return true;
+		}
+
+		public bool PowerOff()
+		{
+			return true;
+		}
+
+		public bool Reboot()
+		{
+			const string Cmd = "/usr/local/bin/idevicediagnostics";
+			if (!File.Exists(Cmd))
 			{
-
-				// setup symbols if available
-				string DSymBundle = "";
-				string DSymDir = Path.Combine(GauntletAppCache, "Symbols");
-
-				if (Directory.Exists(DSymDir))
-				{
-					DSymBundle = Directory.GetDirectories(DSymDir).Where(D => Path.GetExtension(D).ToLower() == ".dsym").FirstOrDefault();
-					DSymBundle = string.IsNullOrEmpty(DSymBundle) ? "" : DSymBundle = " -S \"" + DSymBundle + "\"";
-				}
-
-				string CL = "--noinstall -I" + DSymBundle + " -b \"" + LocalAppBundle + "\" --args '" + CommandLine.Trim() + "'";
-
-				Result = ExecuteIOSDeployCommand(CL, 0);
-
-				Thread.Sleep(5000);
-
-				// Give ios-deploy a chance to throw out any errors...
-				if (Result.HasExited)
-				{
-					Log.Warning(KnownLogEvents.Gauntlet_DeviceEvent, "ios-deploy exited early: " + Result.Output);
-					throw new DeviceException("Failed to launch on {0}. {1}", Name, Result.Output);
-				}
+				Log.Verbose("Rebooting iOS device requires idevicediagnostics binary");
+				return true;
 			}
 
-			return new IOSAppInstance(IOSApp, Result, IOSApp.CommandLine);
-		}
-
-		/// <summary>
-		/// Remove the application entirely from the iOS device, this includes any persistent app data in /Documents
-		/// </summary>
-		private void RemoveApplication(IOSBuild Build)
-		{
-			string CommandLine = String.Format("--bundle_id {0} --uninstall_only", Build.PackageName);
-			ExecuteIOSDeployCommand(CommandLine);
-		}
-
-		/// <summary>
-		/// Remove artifacts from device
-		/// </summary>
-		private bool CleanDeviceArtifacts(IOSBuild Build)
-		{			
-			try
-			{
-				Log.Verbose("Cleaning device artifacts");
-
-				string CleanCommand = String.Format("--bundle_id {0} --rmtree {1}", Build.PackageName, DeviceArtifactPath);
-				IProcessResult Result = ExecuteIOSDeployCommand(CleanCommand, 120);
-
-				if (Result.ExitCode != 0)
-				{
-					Log.Warning(KnownLogEvents.Gauntlet_DeviceEvent, "Failed to clean artifacts from device");
-					return false;
-				}
-
-			}
-			catch (Exception Ex)
-			{
-				Log.Verbose("Exception while cleaning artifacts from device: {0}", Ex.Message);
-			}
-
-			return true;
-
-		}
-
-		/// <summary>
-		/// Checks whether version of deployed bundle matches local Build
-		/// </summary>
-		bool CheckDeployedApp(IOSBuild Build)
-		{
-			try
-			{
-				Log.Verbose("Checking deployed App hash");
-
-				string CommandLine = String.Format("--bundle_id {0} --download={1} --to {2}", Build.PackageName, "/Documents/AppHash.txt", LocalCachePath);
-				IProcessResult Result = ExecuteIOSDeployCommand(CommandLine, 120);
-
-				if (Result.ExitCode != 0)
-				{
-					return false;
-				}
-
-				string Hash = File.ReadAllText(LocalCachePath + "/Documents/AppHash.txt").Trim();
-				string StoredHash = File.ReadAllText(AppHashFilename).Trim();
-
-				if (Hash == StoredHash)
-				{
-					Log.Verbose("Deployed app hash matched cached App hash");
-					return true;
-				}
-
-			}
-			catch (Exception Ex)
-			{
-				if (!Ex.Message.Contains("is denied"))
-				{
-					Log.Verbose("Unable to pull cached App cache from device, cached file may not exist: {0}", Ex.Message);
-				}
-			}
-
-			Log.Verbose("Deployed app hash doesn't match, App will be installed");
-			return false;
-
-		}
-
-		/// <summary>
-		/// Resign application using local executable and update debug symbols
-		/// </summary>
-		void ResignApplication(UnrealAppConfig AppConfig)
-		{
-			// check that we have the signing stuff we need
-			string SignProvision = Globals.Params.ParseValue("signprovision", String.Empty);
-			string SignEntitlements = Globals.Params.ParseValue("signentitlements", String.Empty);
-			string SigningIdentity =  Globals.Params.ParseValue("signidentity", String.Empty);
-			
-			// handle signing provision
-			if (string.IsNullOrEmpty(SignProvision) || !File.Exists(SignProvision))
-			{
-				throw new AutomationException("Absolute path to existing provision must be specified, example: -signprovision=/path/to/myapp.provision");
-			}
-			
-			// handle entitlements
-			// Note this extracts entitlements: which may be useful when using same provision/entitlements?: codesign -d --entitlements :entitlements.plist ~/.gauntletappcache/Payload/Example.app/
-
-			if (string.IsNullOrEmpty(SignEntitlements) || !File.Exists(SignEntitlements))
-			{
-				throw new AutomationException("Absolute path to existing entitlements must be specified, example: -signprovision=/path/to/entitlements.plist");
-			}
-
-			// signing identity
-			if (string.IsNullOrEmpty(SigningIdentity))
-			{
-				throw new AutomationException("Signing identity must be specified, example: -signidentity=\"iPhone Developer: John Smith\"");
-			}
-
-			string ProjectName = AppConfig.ProjectName;
-			string BundleName = Path.GetFileNameWithoutExtension(LocalAppBundle);
-			string ExecutableName = UnrealHelpers.GetExecutableName(ProjectName, UnrealTargetPlatform.IOS, AppConfig.Configuration, AppConfig.ProcessType, "");
-			string CachedAppPath = Path.Combine(GauntletAppCache, "Payload", string.Format("{0}.app", BundleName));			
-
-			string LocalExecutable = Path.Combine(Environment.CurrentDirectory, ProjectName, string.Format("Binaries/IOS/{0}", ExecutableName));
-			if (!File.Exists(LocalExecutable))
-			{
-				throw new AutomationException("Local executable not found for -dev argument: {0}", LocalExecutable);
-			}
-
-			File.WriteAllText(CacheResignedFilename, "The application has been resigned");
-
-			// copy local executable
-			FileInfo SrcInfo = new FileInfo(LocalExecutable);			
-			string DestPath = Path.Combine(CachedAppPath, BundleName);
-			SrcInfo.CopyTo(DestPath, true);
-			Log.Verbose("Copied local executable from {0} to {1}", LocalExecutable, DestPath);
-
-			// copy provision
-			SrcInfo = new FileInfo(SignProvision);
-			DestPath = Path.Combine(CachedAppPath, "embedded.mobileprovision");
-			SrcInfo.CopyTo(DestPath, true);
-			Log.Verbose("Copied provision from {0} to {1}", SignProvision, DestPath);
-
-			// handle symbols
-			string LocalSymbolsDir = Path.Combine(Environment.CurrentDirectory, ProjectName, string.Format("Binaries/IOS/{0}.dSYM", ExecutableName));
-			DestPath = Path.Combine(GauntletAppCache, string.Format("Symbols/{0}.dSYM", ExecutableName));
-
-			if (Directory.Exists(DestPath))
-			{
-				Directory.Delete(DestPath, true);
-			}
-
-			if (Directory.Exists(LocalSymbolsDir))
-			{				
-				CommandUtils.CopyDirectory_NoExceptions(LocalSymbolsDir, DestPath, true);
-			}
-			else
-			{
-				Log.Warning(KnownLogEvents.Gauntlet_DeviceEvent, "No symbols found for local build at {Directory}, removing cached app symbols", LocalSymbolsDir);
-			}
-
-			// resign application
-			// @todo: this asks for password unless "Always Allow" is selected, also for builders, document how to permanently grant codesign access to keychain
-			string SignArgs = string.Format("-f -s \"{0}\" --entitlements \"{1}\" \"{2}\"", SigningIdentity, SignEntitlements, CachedAppPath);
-			Log.Info("\nResigning app, please enter keychain password if prompted:\n\ncodesign {0}", SignArgs);
-			var Result = IOSBuild.ExecuteCommand("codesign", SignArgs);
+			var Result = IOSBuild.ExecuteCommand(Cmd, string.Format("restart -u {0}", DeviceName));
 			if (Result.ExitCode != 0)
 			{
-				throw new AutomationException("Failed to resign application");
+				Log.Warning(string.Format("Failed to reboot iOS device {0}, restart command failed", DeviceName));
+				return true;
 			}
+
+			// initial wait 20 seconds
+			Thread.Sleep(20 * 1000);
+
+			const int WaitPeriod = 10;
+			int WaitTime = 120;
+			bool rebooted = false;
+			do
+			{
+				Result = IOSBuild.ExecuteCommand(Cmd, string.Format("diagnostics WiFi -u {0}", DeviceName));
+				if (Result.ExitCode == 0)
+				{
+					rebooted = true;
+					break;
+				}
+
+				Thread.Sleep(WaitPeriod * 1000);
+				WaitTime -= WaitPeriod;
+
+			} while (WaitTime > 0);
+
+			if (!rebooted)
+			{
+				Log.Warning(KnownLogEvents.Gauntlet_DeviceEvent, "Failed to reboot iOS device {Name}, device didn't come back after restart", DeviceName);
+			}
+
+			return true;
 		}
 
-		// We need to lock around setting up the App
-		static object AppLock = new object();
-		public IAppInstall InstallApplication(UnrealAppConfig AppConfig)
+		public Dictionary<EIntendedBaseCopyDirectory, string> GetPlatformDirectoryMappings()
 		{
-            IOSBuild Build = AppConfig.Build as IOSBuild;
-
-			// Ensure Build exists
-			if (Build == null)
-			{
-				throw new AutomationException("Invalid build for IOS!");
-			}
-
-			if (!AppConfig.SkipInstall)
-			{
-				bool CacheResigned = false;
-				bool UseLocalExecutable = Globals.Params.ParseParam("dev");
-
-				lock (AppLock)
-				{
-					Log.Info("Installing using build source {0}", Build.SourcePath);
-
-					// device artifact path
-					DeviceArtifactPath = string.Format("/Documents/{0}/Saved", AppConfig.ProjectName);
-
-					CacheResigned = File.Exists(CacheResignedFilename);
-
-					if (CacheResigned && !UseLocalExecutable)
-					{
-						if (File.Exists(AppHashFilename))
-						{
-							Log.Verbose("App was resigned, invalidating app cache");
-							File.Delete(AppHashFilename);
-						}
-					}
-
-					if (Build.IsIPAFile)
-					{
-						PrepareIPA(Build);
-					}
-					else
-					{
-						LocalAppBundle = Build.SourcePath;
-					}
-
-					// local executable support
-					if (UseLocalExecutable)
-					{
-						ResignApplication(AppConfig);
-					}
-				}
-
-				if (CacheResigned || UseLocalExecutable || !CheckDeployedApp(Build))
-				{
-					// uninstall will clean all device artifacts
-					ExecuteIOSDeployCommand(String.Format("--uninstall -b \"{0}\"", LocalAppBundle), 20 * 60);
-				}
-				else
-				{
-					// remove device artifacts
-					CleanDeviceArtifacts(Build);
-				}
-
-				// parallel iOS tests use same app install folder, so lock it as setup is quick
-				lock (Globals.MainLock)
-				{
-					// local app install with additional files, this directory will be mirrored to device in a single operation
-					string AppInstallPath;
-
-					AppInstallPath = Path.Combine(Globals.TempDir, "iOSAppInstall");
-
-					if (Directory.Exists(AppInstallPath))
-					{
-						Directory.Delete(AppInstallPath, true);
-					}
-
-					Directory.CreateDirectory(AppInstallPath);
-
-					if (LocalDirectoryMappings.Count == 0)
-					{
-						PopulateDirectoryMappings(AppInstallPath);
-					}
-
-					//@todo: Combine Build and AppConfig files, this should be done in higher level code, not per device implementation
-
-					if (AppConfig.FilesToCopy != null)
-					{
-						foreach (UnrealFileToCopy FileToCopy in AppConfig.FilesToCopy)
-						{
-							string PathToCopyTo = Path.Combine(LocalDirectoryMappings[FileToCopy.TargetBaseDirectory], FileToCopy.TargetRelativeLocation);
-
-							if (File.Exists(FileToCopy.SourceFileLocation))
-							{
-								FileInfo SrcInfo = new FileInfo(FileToCopy.SourceFileLocation);
-								SrcInfo.IsReadOnly = false;
-								string DirectoryToCopyTo = Path.GetDirectoryName(PathToCopyTo);
-								if (!Directory.Exists(DirectoryToCopyTo))
-								{
-									Directory.CreateDirectory(DirectoryToCopyTo);
-								}
-								if (File.Exists(PathToCopyTo))
-								{
-									FileInfo ExistingFile = new FileInfo(PathToCopyTo);
-									ExistingFile.IsReadOnly = false;
-								}
-
-								SrcInfo.CopyTo(PathToCopyTo, true);
-								Log.Verbose("Copying app install: {0} to {1}", FileToCopy, DirectoryToCopyTo);
-							}
-							else
-							{
-								Log.Warning(KnownLogEvents.Gauntlet_DeviceEvent, "File to copy {File} not found", FileToCopy);
-							}
-						}
-					}
-
-					// copy mapped files in a single pass
-					string CopyCommand = String.Format("--bundle_id {0} --upload={1} --to {2}", Build.PackageName, AppInstallPath, DeviceArtifactPath);
-					ExecuteIOSDeployCommand(CopyCommand, 120);
-
-					// store the App hash to avoid redundant deployments
-					if (File.Exists(AppHashFilename))
-					{
-						CopyCommand = String.Format("--bundle_id {0} --upload={1} --to {2}", Build.PackageName, AppHashFilename, "/Documents/AppHash.txt");
-						ExecuteIOSDeployCommand(CopyCommand, 120);
-					}
-				}
-			}
-			else
-			{
-				// Skipping installation
-				Log.Info("Skipping installation as per config.");
-				if (Build.IsIPAFile)
-				{
-					string PayloadDir = Path.Combine(GauntletAppCache, "Payload");
-					LocalAppBundle = Directory.GetDirectories(PayloadDir).Where(D => Path.GetExtension(D) == ".app").FirstOrDefault();
-					if (string.IsNullOrEmpty(LocalAppBundle))
-					{
-						throw new AutomationException($"Could not find local IPA payload at {PayloadDir}");
-					}
-				}
-				else
-				{
-					LocalAppBundle = Build.SourcePath;
-				}
-			}
-
-			IOSAppInstall IOSApp = new IOSAppInstall(AppConfig.Name, this, Build.PackageName, AppConfig.CommandLine);
-			return IOSApp;
+			return LocalDirectoryMappings;
 		}
 
+		public override string ToString()
+		{
+			return Name;
+		}
+		#endregion
+
 		public void FullClean()
 		{
@@ -696,137 +293,298 @@
             LocalDirectoryMappings.Add(EIntendedBaseCopyDirectory.Saved, ProjectDir);
         }
 
+		public IAppInstance Run(IAppInstall App)
+		{
+			IOSAppInstall IOSApp = App as IOSAppInstall;
 
-		public UnrealTargetPlatform? Platform { get { return UnrealTargetPlatform.IOS; } }
+			if (IOSApp == null)
+			{
+				throw new DeviceException("AppInstance is of incorrect type!");
+			}
+
+			string CommandLine = IOSApp.CommandLine.Replace("\"", "\\\\\"");
+
+			Log.Info("Launching {0} on {1}", App.Name, ToString());
+			Log.Verbose("\t{0}", CommandLine);
+
+			// ios-deploy notes: -L launches detached, -I non-interactive  (exits when app exits), -r uninstalls before install (removes app Documents folder)
+			// -t <seconds> number of seconds to wait for device to be connected
+
+			// check xcodebuild -version, to see if it larger than 15.0
+			bool bIsXcode15 = false;
+			IProcessResult XcodeVersionResult = CommandUtils.Run("xcodebuild", "-version");
+			if (XcodeVersionResult.ExitCode == 0)
+			{
+				// XcodeVersionResult.Output is: Xcode 15, try get version number
+				Match VersionMatch = Regex.Match(XcodeVersionResult.Output, @"Xcode\ (?<Version>\d+)");
+				if (VersionMatch.Success)
+				{
+					int Version = 0;
+					if (int.TryParse(VersionMatch.Groups["Version"].Value, out Version))
+					{
+						bIsXcode15 = Version >= 15;
+					}
+				}
+			}
+
+			IProcessResult Result;
+			if (bIsXcode15)
+			{
+				CommandLine = CommandLine.Replace("\\\\", "");
+				string CL = String.Format("device process launch --device {0} {1} '{2}'", DeviceName, IOSApp.PackageName, CommandLine.Trim());
+				Result = ExecuteIOSDeployCommand(CL, 0, true, false, true);
+				Thread.Sleep(5000);
+				// check if IOSApp.PackageName process running using devicectl device info processes
+				while(true)
+				{
+					CL = String.Format("device info processes --device {0}", DeviceName);
+					Result = ExecuteIOSDeployCommand(CL, 0, true, false, true);
+					// Get BundleName from IOSApp.PackageName
+					string BundleName = IOSApp.PackageName.Substring(IOSApp.PackageName.LastIndexOf('.') + 1);
+					Log.Info("Process running: " + BundleName);
+					// check if IOSApp.PackageName process not running
+					if (!Result.Output.Contains(BundleName))
+					{
+						break;
+					}
+					Thread.Sleep(5000);
+				}
+
+			}
+			else
+			{
+
+				// setup symbols if available
+				string DSymBundle = "";
+				string DSymDir = Path.Combine(GauntletAppCache, "Symbols");
+
+				if (Directory.Exists(DSymDir))
+				{
+					DSymBundle = Directory.GetDirectories(DSymDir).Where(D => Path.GetExtension(D).ToLower() == ".dsym").FirstOrDefault();
+					DSymBundle = string.IsNullOrEmpty(DSymBundle) ? "" : DSymBundle = " -S \"" + DSymBundle + "\"";
+				}
+
+				string CL = "--noinstall -I" + DSymBundle + " -b \"" + LocalAppBundle + "\" --args '" + CommandLine.Trim() + "'";
+
+				Result = ExecuteIOSDeployCommand(CL, 0);
+
+				Thread.Sleep(5000);
+
+				// Give ios-deploy a chance to throw out any errors...
+				if (Result.HasExited)
+				{
+					Log.Warning(KnownLogEvents.Gauntlet_DeviceEvent, "ios-deploy exited early: " + Result.Output);
+					throw new DeviceException("Failed to launch on {0}. {1}", Name, Result.Output);
+				}
+			}
+
+			return new IOSAppInstance(IOSApp, Result, IOSApp.CommandLine);
+		}
+
+		public IProcessResult ExecuteIOSDeployCommand(String CommandLine, int WaitTime = 60, bool WarnOnTimeout = true, bool UseDeviceID = true, bool UseDeviceCtl = false)
+		{
+			String IOSDeployPath = Path.Combine(Globals.UnrealRootDir, "Engine/Extras/ThirdPartyNotUE/ios-deploy/bin/ios-deploy");
+			if (UseDeviceCtl)
+			{
+				CommandLine = String.Format("devicectl {0}", CommandLine);
+				IOSDeployPath = "/usr/bin/xcrun";
+			}
+			else
+			{
+				if (UseDeviceID && !IsDefaultDevice)
+				{
+					CommandLine = String.Format("--id {0} {1}", DeviceName, CommandLine);
+				}
+			}
+
+			if (!File.Exists(IOSDeployPath))
+			{
+				throw new AutomationException("Unable to run ios-deploy binary at {0}", IOSDeployPath);
+			}
+
+			CommandUtils.ERunOptions RunOptions = CommandUtils.ERunOptions.NoWaitForExit;
+
+			if (Log.IsVeryVerbose)
+			{
+				RunOptions |= CommandUtils.ERunOptions.AllowSpew;
+			}
+			else
+			{
+				RunOptions |= CommandUtils.ERunOptions.NoLoggingOfRunCommand;
+			}
+
+			Log.Verbose("ios-deploy executing '{0}'", CommandLine);
+
+			IProcessResult Result = CommandUtils.Run(IOSDeployPath, CommandLine, Options: RunOptions);
+
+			if (WaitTime > 0)
+			{
+				DateTime StartTime = DateTime.Now;
+
+				Result.ProcessObject.WaitForExit(WaitTime * 1000);
+
+				if (Result.HasExited == false)
+				{
+					if ((DateTime.Now - StartTime).TotalSeconds >= WaitTime)
+					{
+						string Message = String.Format("IOSDeployPath timeout after {0} secs: {1}, killing process", WaitTime, CommandLine);
+
+						if (WarnOnTimeout)
+						{ 
+							Log.Warning(Message);
+						}
+						else
+						{
+							Log.Info(Message);
+						}
+
+						Result.ProcessObject.Kill();
+						// wait up to 15 seconds for process exit
+						Result.ProcessObject.WaitForExit(15000);
+					}
+				}
+			}
+
+			return Result;
+		}
 
 		/// <summary>
-		/// Temp path we use to push/pull things from the device
+		/// Remove the application entirely from the iOS device, this includes any persistent app data in /Documents
 		/// </summary>
-		public string LocalCachePath { get; protected set; }
-
+		private void RemoveApplication(IOSBuild Build)
+		{
+			string CommandLine = String.Format("--bundle_id {0} --uninstall_only", Build.PackageName);
+			ExecuteIOSDeployCommand(CommandLine);
+		}
 
 		/// <summary>
-		/// Artifact (e.g. Saved) path on the device
+		/// Checks whether version of deployed bundle matches local Build
 		/// </summary>
-		public string DeviceArtifactPath { get; protected set;  }
+		private bool CheckDeployedApp(IOSBuild Build)
+		{
+			try
+			{
+				Log.Verbose("Checking deployed App hash");
 
+				string CommandLine = String.Format("--bundle_id {0} --download={1} --to {2}", Build.PackageName, "/Documents/AppHash.txt", LocalCachePath);
+				IProcessResult Result = ExecuteIOSDeployCommand(CommandLine, 120);
 
-		#region Device State Management
+				if (Result.ExitCode != 0)
+				{
+					return false;
+				}
 
-		// NOTE: We check that a default device UUID or the one specifed is connected with 'ios-deploy --detect' at device creation time
-		// otherwise, ios-deploy doesn't currently support this style of queries
-		// It might be possible to add additional lldb queries/commands through the python interface (there are various solutions for reboot, though the ones I have found require jailbreaking)
-		public bool IsAvailable { get { return true; } }
-		public bool IsConnected { get { return Connected; } }
-		public bool IsOn { get { return true; } }
-		public bool PowerOn() { return true; }
-		public bool PowerOff() { return true; }
-		
-		public bool Reboot() 
-		{ 
-			const string Cmd = "/usr/local/bin/idevicediagnostics";
-			if (!File.Exists(Cmd))
+				string Hash = File.ReadAllText(LocalCachePath + "/Documents/AppHash.txt").Trim();
+				string StoredHash = File.ReadAllText(AppHashFilename).Trim();
+
+				if (Hash == StoredHash)
+				{
+					Log.Verbose("Deployed app hash matched cached App hash");
+					return true;
+				}
+
+			}
+			catch (Exception Ex)
 			{
-				Log.Verbose("Rebooting iOS device requires idevicediagnostics binary");
-				return true;
+				if (!Ex.Message.Contains("is denied"))
+				{
+					Log.Verbose("Unable to pull cached App cache from device, cached file may not exist: {0}", Ex.Message);
+				}
 			}
 
-			var Result = IOSBuild.ExecuteCommand(Cmd, string.Format("restart -u {0}", DeviceName));
+			Log.Verbose("Deployed app hash doesn't match, App will be installed");
+			return false;
+
+		}
+
+		/// <summary>
+		/// Resign application using local executable and update debug symbols
+		/// </summary>
+		private void ResignApplication(UnrealAppConfig AppConfig)
+		{
+			// check that we have the signing stuff we need
+			string SignProvision = Globals.Params.ParseValue("signprovision", String.Empty);
+			string SignEntitlements = Globals.Params.ParseValue("signentitlements", String.Empty);
+			string SigningIdentity =  Globals.Params.ParseValue("signidentity", String.Empty);
+
+			// handle signing provision
+			if (string.IsNullOrEmpty(SignProvision) || !File.Exists(SignProvision))
+			{
+				throw new AutomationException("Absolute path to existing provision must be specified, example: -signprovision=/path/to/myapp.provision");
+			}
+
+			// handle entitlements
+			// Note this extracts entitlements: which may be useful when using same provision/entitlements?: codesign -d --entitlements :entitlements.plist ~/.gauntletappcache/Payload/Example.app/
+
+			if (string.IsNullOrEmpty(SignEntitlements) || !File.Exists(SignEntitlements))
+			{
+				throw new AutomationException("Absolute path to existing entitlements must be specified, example: -signprovision=/path/to/entitlements.plist");
+			}
+
+			// signing identity
+			if (string.IsNullOrEmpty(SigningIdentity))
+			{
+				throw new AutomationException("Signing identity must be specified, example: -signidentity=\"iPhone Developer: John Smith\"");
+			}
+
+			string ProjectName = AppConfig.ProjectName;
+			string BundleName = Path.GetFileNameWithoutExtension(LocalAppBundle);
+			string ExecutableName = UnrealHelpers.GetExecutableName(ProjectName, UnrealTargetPlatform.IOS, AppConfig.Configuration, AppConfig.ProcessType, "");
+			string CachedAppPath = Path.Combine(GauntletAppCache, "Payload", string.Format("{0}.app", BundleName));
+
+			string LocalExecutable = Path.Combine(Environment.CurrentDirectory, ProjectName, string.Format("Binaries/IOS/{0}", ExecutableName));
+			if (!File.Exists(LocalExecutable))
+			{
+				throw new AutomationException("Local executable not found for -dev argument: {0}", LocalExecutable);
+			}
+
+			File.WriteAllText(CacheResignedFilename, "The application has been resigned");
+
+			// copy local executable
+			FileInfo SrcInfo = new FileInfo(LocalExecutable);
+			string DestPath = Path.Combine(CachedAppPath, BundleName);
+			SrcInfo.CopyTo(DestPath, true);
+			Log.Verbose("Copied local executable from {0} to {1}", LocalExecutable, DestPath);
+
+			// copy provision
+			SrcInfo = new FileInfo(SignProvision);
+			DestPath = Path.Combine(CachedAppPath, "embedded.mobileprovision");
+			SrcInfo.CopyTo(DestPath, true);
+			Log.Verbose("Copied provision from {0} to {1}", SignProvision, DestPath);
+
+			// handle symbols
+			string LocalSymbolsDir = Path.Combine(Environment.CurrentDirectory, ProjectName, string.Format("Binaries/IOS/{0}.dSYM", ExecutableName));
+			DestPath = Path.Combine(GauntletAppCache, string.Format("Symbols/{0}.dSYM", ExecutableName));
+
+			if (Directory.Exists(DestPath))
+			{
+				Directory.Delete(DestPath, true);
+			}
+
+			if (Directory.Exists(LocalSymbolsDir))
+			{
+				CommandUtils.CopyDirectory_NoExceptions(LocalSymbolsDir, DestPath, true);
+			}
+			else
+			{
+				Log.Warning(KnownLogEvents.Gauntlet_DeviceEvent, "No symbols found for local build at {Directory}, removing cached app symbols", LocalSymbolsDir);
+			}
+
+			// resign application
+			// @todo: this asks for password unless "Always Allow" is selected, also for builders, document how to permanently grant codesign access to keychain
+			string SignArgs = string.Format("-f -s \"{0}\" --entitlements \"{1}\" \"{2}\"", SigningIdentity, SignEntitlements, CachedAppPath);
+			Log.Info("\nResigning app, please enter keychain password if prompted:\n\ncodesign {0}", SignArgs);
+			var Result = IOSBuild.ExecuteCommand("codesign", SignArgs);
 			if (Result.ExitCode != 0)
 			{
-				Log.Warning(string.Format("Failed to reboot iOS device {0}, restart command failed", DeviceName));
-				return true;
+				throw new AutomationException("Failed to resign application");
 			}
-
-			// initial wait 20 seconds
-			Thread.Sleep(20 * 1000);
-
-			const int WaitPeriod = 10;
-			int WaitTime = 120;
-			bool rebooted = false;
-			do
-			{
-				Result = IOSBuild.ExecuteCommand(Cmd, string.Format("diagnostics WiFi -u {0}", DeviceName));
-				if (Result.ExitCode == 0)
-				{
-					rebooted = true;
-					break;
-				}
-				
-				Thread.Sleep(WaitPeriod * 1000);
-				WaitTime -= WaitPeriod;
-
-			} while (WaitTime > 0);
-
-			if (!rebooted) 
-			{
-				Log.Warning(KnownLogEvents.Gauntlet_DeviceEvent, "Failed to reboot iOS device {Name}, device didn't come back after restart", DeviceName);
-			}
-
-			return true; 
 		}
 
-		static Dictionary<string, bool> ConnectedDevices = new Dictionary<string, bool>();
-		bool Connected = false;
-
-		public bool Connect() 
-		{ 
-			lock (Globals.MainLock)
-			{
-				if (Connected)
-				{
-					return true;								
-				}
-
-				bool ExistingConnection = false;
-				if (ConnectedDevices.TryGetValue(DeviceName, out ExistingConnection))
-				{
-					if (ExistingConnection)
-					{
-						throw new AutomationException("Connected to already connected device");
-					}
-				}
-
-				ConnectedDevices[DeviceName] = true;
-
-				Connected = true;
-			}
-
-			return true; 
-		}
-		public bool Disconnect(bool bForce=false)
-		{ 
-			lock (Globals.MainLock)
-			{
-				if (!Connected)
-				{
-					return true;								
-				}
-
-				Connected = false;
-
-				if (ConnectedDevices.ContainsKey(DeviceName))
-				{
-					ConnectedDevices.Remove(DeviceName);
-				}
-
-			}
-
-			return true; 
-		}
-
-		#endregion
-
-		public override string ToString()
-		{
-			return Name;
-		}
-
-
 		/// <summary>
 		/// Get UUID of all connected iOS devices
 		/// </summary>
-		List<string> GetConnectedDeviceUUID()
+		private List<string> GetConnectedDeviceUUID()
 		{
 			var Result = ExecuteIOSDeployCommand("--detect", 60, true, false);
 
@@ -842,10 +600,8 @@
 			).ToList();
 		}
 
-		static bool ZombiesKilled = false;
-
 		// Get rid of any zombie lldb/iosdeploy processes, this needs to be reworked to use tracked process id's when running parallel tests across multiple AutomationTool.exe processes on test workers
-		void KillZombies()
+		private void KillZombies()
 		{
 
 			if (ZombiesKilled)
@@ -857,40 +613,20 @@
 
 			IOSBuild.ExecuteCommand("killall", "ios-deploy");
 			Thread.Sleep(2500);
-			IOSBuild.ExecuteCommand("killall", "lldb");			
+			IOSBuild.ExecuteCommand("killall", "lldb");
 			Thread.Sleep(2500);
 		}
-		
-		// Gauntlet cache folder for tracking device/ipa state
-		string GauntletAppCache
-		{
-			get
-			{	
-				return Path.Combine(Globals.TempDir, "IOSAppCache");
-			}
-		}
 
-		// path to locally extracted (and possibly resigned) app bundle
-		// Note: ios-deploy works with app bundles, which requires the IPA be unzipped for deployment (this will allow us to resign in the future as well)
-		string LocalAppBundle = null;
-
-		// the current App MD5 hash, which is tracked to avoid unneccessary deployments and unzip operations
-		string AppHashFilename { get { return Path.Combine(GauntletAppCache, "AppHash.txt"); } }
-
-		// file whose presence signals that cache was resigned
-		string CacheResignedFilename { get { return Path.Combine(GauntletAppCache, "Resigned.txt"); } }
-
-
 		/// <summary>
 		/// Generate MD5 and cache IPA bundle files
 		/// </summary>
 		private bool PrepareIPA(IOSBuild Build)
-		{	
+		{
 			Log.Info("Preparing IPA {0}", Build.SourcePath);
 
 			try
-			{	
-				// cache the unzipped app using a MD5 checksum, avoiding needing to unzip		
+			{
+				// cache the unzipped app using a MD5 checksum, avoiding needing to unzip
 				string Hash = null;
 				string StoredHash = null;
 				using (var MD5Hash = MD5.Create())
@@ -902,7 +638,7 @@
 				}
 				string PayloadDir = Path.Combine(GauntletAppCache, "Payload");
 				string SymbolsDir = Path.Combine(GauntletAppCache, "Symbols");
-				
+
 				if (File.Exists(AppHashFilename) && Directory.Exists(PayloadDir))
 				{
 					StoredHash = File.ReadAllText(AppHashFilename).Trim();
@@ -944,14 +680,14 @@
 					Log.Verbose("Checking Symbols at {0}", SymbolsZipFile);
 
 					if (File.Exists(SymbolsZipFile))
-					{						
+					{
 						Log.Verbose("Unzipping Symbols {0} to cache at: {1}", SymbolsZipFile, SymbolsDir);
 
 						if (!IOSBuild.ExecuteIPAZipCommand(String.Format("{0} -d {1}", SymbolsZipFile, SymbolsDir), out Output, SymbolsDir))
 						{
 							throw new Exception(String.Format("Unable to extract build symbols {0} -> {1}", SymbolsZipFile, SymbolsDir));
 						}
-					}					
+					}
 
 					// store hash
 					File.WriteAllText(AppHashFilename, Hash);
@@ -965,95 +701,319 @@
 
 				LocalAppBundle = Directory.GetDirectories(PayloadDir).Where(D => Path.GetExtension(D) == ".app").FirstOrDefault();
 
-				if (String.IsNullOrEmpty(LocalAppBundle))
+				if (string.IsNullOrEmpty(LocalAppBundle))
 				{
-					throw new Exception(String.Format("Unable to find app in local app bundle {0}", PayloadDir));
+					throw new Exception(string.Format("Unable to find app in local app bundle {0}", PayloadDir));
 				}
 
 			}
 			catch (Exception Ex)
 			{
 				throw new AutomationException("Unable to prepare {0} : {1}", Build.SourcePath, Ex.Message);
-			}		
+			}
 
 			return true;
 		}
 
-		public Dictionary<EIntendedBaseCopyDirectory, string> GetPlatformDirectoryMappings()
-		{
-			if (LocalDirectoryMappings.Count == 0)
+		#region Legacy Implementations
+
+		/// <summary>
+		/// Remove artifacts from device
+		/// </summary>
+		private bool CleanDeviceArtifacts(IOSBuild Build)
+		{			
+			try
 			{
-				Log.Warning("Platform directory mappings have not been populated for this platform! This should be done within InstallApplication()");
+				Log.Verbose("Cleaning device artifacts");
+
+				string CleanCommand = String.Format("--bundle_id {0} --rmtree {1}", Build.PackageName, DeviceArtifactPath);
+				IProcessResult Result = ExecuteIOSDeployCommand(CleanCommand, 120);
+
+				if (Result.ExitCode != 0)
+				{
+					Log.Warning(KnownLogEvents.Gauntlet_DeviceEvent, "Failed to clean artifacts from device");
+					return false;
+				}
+
 			}
-			return LocalDirectoryMappings;
+			catch (Exception Ex)
+			{
+				Log.Verbose("Exception while cleaning artifacts from device: {0}", Ex.Message);
+			}
+
+			return true;
+
 		}
 
-		public IProcessResult ExecuteIOSDeployCommand(String CommandLine, int WaitTime = 60, bool WarnOnTimeout = true, bool UseDeviceID = true, bool UseDeviceCtl = false)
+		public IAppInstall InstallApplication(UnrealAppConfig AppConfig)
 		{
-			String IOSDeployPath = Path.Combine(Globals.UnrealRootDir, "Engine/Extras/ThirdPartyNotUE/ios-deploy/bin/ios-deploy");
-			if (UseDeviceCtl)
+			IOSBuild Build = AppConfig.Build as IOSBuild;
+
+			// Ensure Build exists
+			if (Build == null)
 			{
-				CommandLine = String.Format("devicectl {0}", CommandLine);
-				IOSDeployPath = "/usr/bin/xcrun";		
+				throw new AutomationException("Invalid build for IOS!");
+			}
+
+			if (!AppConfig.SkipInstall)
+			{
+				bool CacheResigned = false;
+				bool UseLocalExecutable = Globals.Params.ParseParam("dev");
+
+				lock (AppLock)
+				{
+					Log.Info("Installing using build source {0}", Build.SourcePath);
+
+					// device artifact path
+					DeviceArtifactPath = string.Format("/Documents/{0}/Saved", AppConfig.ProjectName);
+
+					CacheResigned = File.Exists(CacheResignedFilename);
+
+					if (CacheResigned && !UseLocalExecutable)
+					{
+						if (File.Exists(AppHashFilename))
+						{
+							Log.Verbose("App was resigned, invalidating app cache");
+							File.Delete(AppHashFilename);
+						}
+					}
+
+					if (Build.IsIPAFile)
+					{
+						PrepareIPA(Build);
+					}
+					else
+					{
+						LocalAppBundle = Build.SourcePath;
+					}
+
+					// local executable support
+					if (UseLocalExecutable)
+					{
+						ResignApplication(AppConfig);
+					}
+				}
+
+				if (CacheResigned || UseLocalExecutable || !CheckDeployedApp(Build))
+				{
+					// uninstall will clean all device artifacts
+					IProcessResult Result = ExecuteIOSDeployCommand(String.Format("--uninstall -b \"{0}\"", LocalAppBundle), 20 * 60);
+					Log.Info(Result.Output);
+				}
+				else
+				{
+					// remove device artifacts
+					CleanDeviceArtifacts(Build);
+				}
+
+				// parallel iOS tests use same app install folder, so lock it as setup is quick
+				lock (Globals.MainLock)
+				{
+					// local app install with additional files, this directory will be mirrored to device in a single operation
+					string AppInstallPath;
+
+					AppInstallPath = Path.Combine(Globals.TempDir, "iOSAppInstall");
+
+					if (Directory.Exists(AppInstallPath))
+					{
+						Directory.Delete(AppInstallPath, true);
+					}
+
+					Directory.CreateDirectory(AppInstallPath);
+
+					if (LocalDirectoryMappings.Count == 0)
+					{
+						PopulateDirectoryMappings(AppInstallPath);
+					}
+
+					//@todo: Combine Build and AppConfig files, this should be done in higher level code, not per device implementation
+
+					if (AppConfig.FilesToCopy != null)
+					{
+						foreach (UnrealFileToCopy FileToCopy in AppConfig.FilesToCopy)
+						{
+							string PathToCopyTo = Path.Combine(LocalDirectoryMappings[FileToCopy.TargetBaseDirectory], FileToCopy.TargetRelativeLocation);
+
+							if (File.Exists(FileToCopy.SourceFileLocation))
+							{
+								FileInfo SrcInfo = new FileInfo(FileToCopy.SourceFileLocation);
+								SrcInfo.IsReadOnly = false;
+								string DirectoryToCopyTo = Path.GetDirectoryName(PathToCopyTo);
+								if (!Directory.Exists(DirectoryToCopyTo))
+								{
+									Directory.CreateDirectory(DirectoryToCopyTo);
+								}
+								if (File.Exists(PathToCopyTo))
+								{
+									FileInfo ExistingFile = new FileInfo(PathToCopyTo);
+									ExistingFile.IsReadOnly = false;
+								}
+
+								SrcInfo.CopyTo(PathToCopyTo, true);
+								Log.Verbose("Copying app install: {0} to {1}", FileToCopy, DirectoryToCopyTo);
+							}
+							else
+							{
+								Log.Warning(KnownLogEvents.Gauntlet_DeviceEvent, "File to copy {File} not found", FileToCopy);
+							}
+						}
+					}
+
+					// copy mapped files in a single pass
+					string CopyCommand = String.Format("--bundle_id {0} --upload={1} --to {2}", Build.PackageName, AppInstallPath, DeviceArtifactPath);
+					ExecuteIOSDeployCommand(CopyCommand, 120);
+
+					// store the App hash to avoid redundant deployments
+					if (File.Exists(AppHashFilename))
+					{
+						CopyCommand = String.Format("--bundle_id {0} --upload={1} --to {2}", Build.PackageName, AppHashFilename, "/Documents/AppHash.txt");
+						ExecuteIOSDeployCommand(CopyCommand, 120);
+					}
+				}
 			}
 			else
 			{
-				if (UseDeviceID && !IsDefaultDevice)
+				// Skipping installation
+				Log.Info("Skipping installation as per config.");
+				if (Build.IsIPAFile)
 				{
-					CommandLine = String.Format("--id {0} {1}", DeviceName, CommandLine);
+					string PayloadDir = Path.Combine(GauntletAppCache, "Payload");
+					LocalAppBundle = Directory.GetDirectories(PayloadDir).Where(D => Path.GetExtension(D) == ".app").FirstOrDefault();
+					if (string.IsNullOrEmpty(LocalAppBundle))
+					{
+						throw new AutomationException($"Could not find local IPA payload at {PayloadDir}");
+					}
 				}
-			}
-	
-
-			if (!File.Exists(IOSDeployPath))
-			{
-				throw new AutomationException("Unable to run ios-deploy binary at {0}", IOSDeployPath);
-			}
-
-			CommandUtils.ERunOptions RunOptions = CommandUtils.ERunOptions.NoWaitForExit;
-
-			if (Log.IsVeryVerbose)
-			{
-				RunOptions |= CommandUtils.ERunOptions.AllowSpew;
-			}
-			else
-			{
-				RunOptions |= CommandUtils.ERunOptions.NoLoggingOfRunCommand;
-			}
-
-			Log.Verbose("ios-deploy executing '{0}'", CommandLine);
-
-			IProcessResult Result = CommandUtils.Run(IOSDeployPath, CommandLine, Options: RunOptions);
-
-			if (WaitTime > 0)
-			{
-				DateTime StartTime = DateTime.Now;
-
-				Result.ProcessObject.WaitForExit(WaitTime * 1000);
-
-				if (Result.HasExited == false)
+				else
 				{
-					if ((DateTime.Now - StartTime).TotalSeconds >= WaitTime)
-					{
-						string Message = String.Format("IOSDeployPath timeout after {0} secs: {1}, killing process", WaitTime, CommandLine);
-
-						if (WarnOnTimeout)
-						{ 
-							Log.Warning(Message);
-						}
-						else
-						{
-							Log.Info(Message);
-						}
-						
-						Result.ProcessObject.Kill();
-						// wait up to 15 seconds for process exit
-						Result.ProcessObject.WaitForExit(15000);
-					}
+					LocalAppBundle = Build.SourcePath;
 				}
 			}
 
-			return Result;
+			IOSAppInstall IOSApp = new IOSAppInstall(AppConfig.Name, this, Build.PackageName, AppConfig.CommandLine);
+			return IOSApp;
+		}
+		#endregion
+	}
+
+	class IOSAppInstall : IAppInstall
+	{
+		public string Name { get; protected set; }
+
+		public string CommandLine { get; protected set; }
+
+		public string PackageName { get; protected set; }
+
+		public ITargetDevice Device => IOSDevice;
+
+		public TargetDeviceIOS IOSDevice;
+
+		public IOSAppInstall(string InName, TargetDeviceIOS InDevice, string InPackageName, string InCommandLine)
+		{
+			Name = InName;
+			CommandLine = InCommandLine;
+			PackageName = InPackageName;
+
+			IOSDevice = InDevice;
+		}
+
+		public IAppInstance Run()
+		{
+			return Device.Run(this);
 		}
 	}
 
+	class IOSAppInstance : IAppInstance
+	{
+		public ITargetDevice Device => Install.Device;
+		public IProcessResult ProcessResult { get; private set; }
+		public string CommandLine { get; private set; }
+		public bool WasKilled { get; protected set; }
+		public int ExitCode => ProcessResult.ExitCode;
+		public bool HasExited => ProcessResult.HasExited;
+
+		public string ArtifactPath
+		{
+			get
+			{
+				if (bHaveSavedArtifacts == false)
+				{
+					if (HasExited)
+					{
+						SaveArtifacts();
+						bHaveSavedArtifacts = true;
+					}
+				}
+
+				return Install.IOSDevice.LocalCachePath + "/" + Install.IOSDevice.DeviceArtifactPath;
+			}
+		}
+
+		public string StdOut
+		{
+			get
+			{
+				if (HasExited)
+				{
+					// The ios application is being run under lldb by ios-deploy
+					// lldb catches crashes and we have it setup to dump thread callstacks
+					// parse any crash dumps into Unreal crash format and append to output
+					string CrashLog = LLDBCrashParser.GenerateCrashLog(ProcessResult.Output);
+
+					if (!string.IsNullOrEmpty(CrashLog))
+					{
+						return String.Format("{0}\n{1}", ProcessResult.Output, CrashLog);
+					}
+				}
+
+				return ProcessResult.Output;
+
+			}
+		}
+
+		protected IOSAppInstall Install;
+
+		protected bool bHaveSavedArtifacts;
+
+		public IOSAppInstance(IOSAppInstall InInstall, IProcessResult InProcess, string InCommandLine)
+		{
+			Install = InInstall;
+			this.CommandLine = InCommandLine;
+			this.ProcessResult = InProcess;
+		}
+
+		public int WaitForExit()
+		{
+			if (!HasExited)
+			{
+				ProcessResult.WaitForExit();
+			}
+
+			return ExitCode;
+		}
+
+		public void Kill(bool bGenerateDump = false)
+		{
+			if (!HasExited)
+			{
+				WasKilled = true;
+				ProcessResult.ProcessObject.Kill(true);
+			}
+		}
+
+		protected void SaveArtifacts()
+		{
+			TargetDeviceIOS Device = Install.IOSDevice;
+
+			// copy remote artifacts to local
+			string CommandLine = String.Format("--bundle_id {0} --download={1} --to {2}", Install.PackageName, Device.DeviceArtifactPath, Device.LocalCachePath);
+
+			IProcessResult DownloadCmd = Device.ExecuteIOSDeployCommand(CommandLine, 120);
+
+			if (DownloadCmd.ExitCode != 0)
+			{
+				Log.Warning(KnownLogEvents.Gauntlet_DeviceEvent, "Failed to retrieve artifacts. {Output}", DownloadCmd.Output);
+			}
+		}
+	}
+
@@ -1102,8 +1062,8 @@
 		}
 
 		/// <summary>
-		/// Parse lldb thread crash dump to Unreal log format 
-		/// </summary>		
+		/// Parse lldb thread crash dump to Unreal log format
+		/// </summary>
 		public static string GenerateCrashLog(string LogOutput)
 		{
 			try
@@ -1122,7 +1082,7 @@
 				CrashLog.Append(string.Join("\n", Thread.Frames));
 
 				return CrashLog.ToString();
-			} 
+			}
 			catch (Exception Ex)
 			{
 				Log.Warning(KnownLogEvents.Gauntlet_DeviceEvent, "Exception parsing LLDB callstack {Exception}", Ex.Message);
@@ -1132,7 +1092,7 @@
 
 		}
 
-		static ThreadInfo ParseCallstack(string LogOutput, out DateTime Timestamp, out int FrameNum)
+		private static ThreadInfo ParseCallstack(string LogOutput, out DateTime Timestamp, out int FrameNum)
 		{
 			Timestamp = DateTime.UtcNow;
 			FrameNum = 0;
@@ -1283,7 +1243,19 @@
 			return Thread;
 
 		}
-
 	}
 
+	public class IOSDeviceFactory : IDeviceFactory
+	{
+		public bool CanSupportPlatform(UnrealTargetPlatform? Platform)
+		{
+			return Platform == UnrealTargetPlatform.IOS;
+		}
+
+		public ITargetDevice CreateDevice(string InRef, string InCachePath, string InParam)
+		{
+			return new TargetDeviceIOS(InRef, InCachePath);
+		}
+	}
+
 	public class IOSBuildSupport : BaseBuildSupport
