BuildGraph: Fix static analysis warnings. Also add support for writing a file containing secrets using inline variable declarations.

#rnx

diff --git a/Engine/Source/Programs/AutomationTool/BuildGraph/Tasks/AgeStoreTask.cs b/Engine/Source/Programs/AutomationTool/BuildGraph/Tasks/AgeStoreTask.cs
--- a/Engine/Source/Programs/AutomationTool/BuildGraph/Tasks/AgeStoreTask.cs
+++ b/Engine/Source/Programs/AutomationTool/BuildGraph/Tasks/AgeStoreTask.cs
@@ -168,7 +168,9 @@ namespace AutomationTool.Tasks
 
 				// Delete this directory if it is empty, and it is not the root directory.
 				if (level > 0 && !currentDirectory.EnumerateFileSystemInfos().Any())
+				{
 					TryDelete(currentDirectory);
+				}
 			}
 		}
 
diff --git a/Engine/Source/Programs/AutomationTool/BuildGraph/Tasks/HordeGetSecretsTask.cs b/Engine/Source/Programs/AutomationTool/BuildGraph/Tasks/HordeGetSecretsTask.cs
--- a/Engine/Source/Programs/AutomationTool/BuildGraph/Tasks/HordeGetSecretsTask.cs
+++ b/Engine/Source/Programs/AutomationTool/BuildGraph/Tasks/HordeGetSecretsTask.cs
@@ -2,7 +2,6 @@
 
 using System;
 using System.Collections.Generic;
-using System.Diagnostics;
 using System.Linq;
 using System.Net.Http;
 using System.Threading.Tasks;
@@ -10,6 +9,7 @@ using System.Xml;
 using EpicGames.Core;
 using EpicGames.Horde;
 using EpicGames.Horde.Secrets;
+using IdentityModel.Client;
 using Microsoft.Extensions.DependencyInjection;
 using Microsoft.Extensions.Logging;
 
@@ -29,10 +29,17 @@ namespace AutomationTool.Tasks
 		public string File { get; set; } = String.Empty;
 
 		/// <summary>
-		/// Pairs of strings and secret names to expand in the text file, in the form SOURCE_TEXT=secret-name;SOURCE_TEXT_2=secret-name-2
+		/// Text to update with secrets
 		/// </summary>
-		[TaskParameter]
-		public string Replace { get; set; } = String.Empty;
+		[TaskParameter(Optional = true)]
+		public string? Text { get; set; }
+
+		/// <summary>
+		/// Pairs of strings and secret names to expand in the text file, in the form SOURCE_TEXT=secret-name;SOURCE_TEXT_2=secret-name-2. 
+		/// If not specified, secrets embedded inline in the text will be expanded from {{secret-name.value}} strings.
+		/// </summary>
+		[TaskParameter(Optional = true)]
+		public string? Replace { get; set; }
 	}
 
 	/// <summary>
@@ -62,44 +69,68 @@ namespace AutomationTool.Tasks
 		/// <param name="tagNameToFileSet">Mapping from tag names to the set of files they include.</param>
 		public override async Task ExecuteAsync(JobContext job, HashSet<FileReference> buildProducts, Dictionary<string, HashSet<FileReference>> tagNameToFileSet)
 		{
-			// Read the input text
 			FileReference file = ResolveFile(_parameters.File);
-			string text = await FileReference.ReadAllTextAsync(file);
+
+			// Read the input text
+			string? text = _parameters.Text;
+			if (String.IsNullOrEmpty(text))
+			{
+				text = await FileReference.ReadAllTextAsync(file);
+			}
 
 			// Parse the secrets to replace
 			Dictionary<SecretId, List<ReplacementInfo>> secretToReplacementInfo = new Dictionary<SecretId, List<ReplacementInfo>>();
-			foreach (string clause in _parameters.Replace.Split(';', StringSplitOptions.RemoveEmptyEntries | StringSplitOptions.TrimEntries))
+			if (String.IsNullOrEmpty(_parameters.Replace))
 			{
-				int idx = clause.LastIndexOf('=');
-				if (idx == -1)
+				int pos = 0;
+				for (; ; )
 				{
-					throw new AutomationException($"Invalid replacement clause in Horde-GetSecrets task: {clause} (expected VARIABLE=Secret.Property)");
-				}
+					pos = text.IndexOf("{{", pos, StringComparison.Ordinal);
+					if (pos == -1)
+					{
+						break;
+					}
 
-				int propertyIdx = clause.IndexOf('.', idx + 1);
-				if (propertyIdx == -1)
-				{
-					throw new AutomationException($"Invalid replacement clause for secret in Horde-GetSecrets task: {clause} (expected VARIABLE=Secret.Property)");
-				}
+					pos += 2;
 
-				string variable = clause.Substring(0, idx);
-				if (!text.Contains(variable, StringComparison.Ordinal))
-				{
-					Logger.LogWarning("Variable '{Variable}' not found in {File}", variable, file);
-					continue;
-				}
+					int endPos = text.IndexOf("}}", pos, StringComparison.Ordinal);
+					if (endPos == -1)
+					{
+						continue;
+					}
 
-				SecretId secretId = new SecretId(clause.Substring(idx + 1, propertyIdx - (idx + 1)));
-				string propertyName = clause.Substring(propertyIdx + 1);
+					string variable = text.Substring(pos - 2, (endPos + 2) - (pos - 2));
+					string replacement = text.Substring(pos, endPos - pos);
+					if (!ParseReplacementInfo(variable, replacement, secretToReplacementInfo))
+					{
+						throw new AutomationException($"Invalid replacement clause for secret in Horde-GetSecrets task: {replacement} (expected VARIABLE=Secret.Property)");
+					}
 
-				List<ReplacementInfo>? replacements;
-				if (!secretToReplacementInfo.TryGetValue(secretId, out replacements))
-				{
-					replacements = new List<ReplacementInfo>();
-					secretToReplacementInfo.Add(secretId, replacements);
+					pos = endPos + 2;
 				}
+			}
+			else
+			{
+				foreach (string replacement in _parameters.Replace.Split(';', StringSplitOptions.RemoveEmptyEntries | StringSplitOptions.TrimEntries))
+				{
+					int idx = replacement.LastIndexOf('=');
+					if (idx == -1)
+					{
+						throw new AutomationException($"Invalid replacement clause in Horde-GetSecrets task: {replacement} (expected VARIABLE=Secret.Property)");
+					}
 
-				replacements.Add(new ReplacementInfo(propertyName, variable));
+					string variable = replacement.Substring(0, idx);
+					if (!text.Contains(variable, StringComparison.Ordinal))
+					{
+						Logger.LogWarning("Variable '{Variable}' not found in {File}", variable, file);
+						continue;
+					}
+
+					if (!ParseReplacementInfo(variable, replacement.Substring(idx + 1), secretToReplacementInfo))
+					{
+						throw new AutomationException($"Invalid replacement clause for secret in Horde-GetSecrets task: {replacement} (expected VARIABLE=Secret.Property)");
+					}
+				}
 			}
 
 			// Read the secrets from Horde, and substitute them in the output file
@@ -143,10 +174,33 @@ namespace AutomationTool.Tasks
 			}
 
 			// Write the output file
+			DirectoryReference.CreateDirectory(file.Directory);
 			await FileReference.WriteAllTextAsync(file, text);
 			Logger.LogInformation("Updated {File} with secrets from Horde.", file);
 		}
 
+		static bool ParseReplacementInfo(string variable, string replacement, Dictionary<SecretId, List<ReplacementInfo>> secretToReplacementInfo)
+		{
+			int propertyIdx = replacement.IndexOf('.', StringComparison.Ordinal);
+			if (propertyIdx == -1)
+			{
+				return false;
+			}
+
+			SecretId secretId = new SecretId(replacement.Substring(0, propertyIdx));
+			string propertyName = replacement.Substring(propertyIdx + 1);
+
+			List<ReplacementInfo>? replacements;
+			if (!secretToReplacementInfo.TryGetValue(secretId, out replacements))
+			{
+				replacements = new List<ReplacementInfo>();
+				secretToReplacementInfo.Add(secretId, replacements);
+			}
+
+			replacements.Add(new ReplacementInfo(propertyName, variable));
+			return true;
+		}
+
 		/// <summary>
 		/// Output this task out to an XML writer.
 		/// </summary>
diff --git a/Engine/Source/Programs/AutomationTool/BuildGraph/Tasks/ModifyJsonValue.cs b/Engine/Source/Programs/AutomationTool/BuildGraph/Tasks/ModifyJsonValue.cs
--- a/Engine/Source/Programs/AutomationTool/BuildGraph/Tasks/ModifyJsonValue.cs
+++ b/Engine/Source/Programs/AutomationTool/BuildGraph/Tasks/ModifyJsonValue.cs
@@ -72,7 +72,9 @@ namespace AutomationTool.Tasks
 				for (int i = 0; i < keys.Length - 1; i++)
 				{
 					if (!currObj.TryGetValue(keys[i], out object nextNode))
+					{
 						currObj[keys[i]] = nextNode = new Dictionary<string, object>();
+					}
 					currObj = (IDictionary<string, object>)nextNode;
 				}
 
diff --git a/Engine/Source/Programs/AutomationTool/BuildGraph/Tasks/UploadTask.cs b/Engine/Source/Programs/AutomationTool/BuildGraph/Tasks/UploadTask.cs
--- a/Engine/Source/Programs/AutomationTool/BuildGraph/Tasks/UploadTask.cs
+++ b/Engine/Source/Programs/AutomationTool/BuildGraph/Tasks/UploadTask.cs
@@ -143,7 +143,9 @@ namespace AutomationTool.Tasks
 				{
 					int separatorIndex = option.IndexOf('=', StringComparison.Ordinal);
 					if (separatorIndex == -1)
+					{
 						continue;
+					}
 
 					string key = option.Substring(0, separatorIndex);
 					string value = option.Substring(separatorIndex + 1);
diff --git a/Engine/Source/Programs/AutomationTool/BuildGraph/TempStorage.cs b/Engine/Source/Programs/AutomationTool/BuildGraph/TempStorage.cs
--- a/Engine/Source/Programs/AutomationTool/BuildGraph/TempStorage.cs
+++ b/Engine/Source/Programs/AutomationTool/BuildGraph/TempStorage.cs
@@ -392,7 +392,7 @@ namespace AutomationTool
 		/// Load a manifest from disk
 		/// </summary>
 		/// <param name="file">File to load</param>
-		static public TempStorageManifest Load(FileReference file)
+		public static TempStorageManifest Load(FileReference file)
 		{
 			using (StreamReader reader = new(file.FullName))
 			{
@@ -500,7 +500,7 @@ namespace AutomationTool
 		/// Load this list of files from disk
 		/// </summary>
 		/// <param name="file">File to load</param>
-		static public TempStorageTagManifest Load(FileReference file)
+		public static TempStorageTagManifest Load(FileReference file)
 		{
 			using (StreamReader reader = new StreamReader(file.FullName))
 			{
