BuildGraph: Improve error handling for parsing environment variable strings.

#rnx

diff --git a/Engine/Source/Programs/AutomationTool/BuildGraph/BgScriptReader.cs b/Engine/Source/Programs/AutomationTool/BuildGraph/BgScriptReader.cs
--- a/Engine/Source/Programs/AutomationTool/BuildGraph/BgScriptReader.cs
+++ b/Engine/Source/Programs/AutomationTool/BuildGraph/BgScriptReader.cs
@@ -821,6 +821,15 @@ namespace AutomationTool
 							int prefixLen = -1;
 
 							List<string> lines = new List<string>(element.InnerText.Split('\n'));
+							while (lines.Count > 0 && String.IsNullOrWhiteSpace(lines[0]))
+							{
+								lines.RemoveAt(0);
+							}
+							while (lines.Count > 0 && String.IsNullOrWhiteSpace(lines[^1]))
+							{
+								lines.RemoveAt(lines.Count - 1);
+							}
+
 							foreach(string line in lines)
 							{
 								int whitespaceLen = 0;
@@ -831,14 +840,7 @@ namespace AutomationTool
 
 								if (prefixLen == -1)
 								{
-									if (whitespaceLen == line.Length)
-									{
-										continue;
-									}
-									else
-									{
-										prefixLen = whitespaceLen;
-									}
+									prefixLen = whitespaceLen;
 								}
 
 								whitespaceLen = Math.Min(whitespaceLen, prefixLen);
diff --git a/Engine/Source/Programs/AutomationTool/BuildGraph/Tasks/SpawnTask.cs b/Engine/Source/Programs/AutomationTool/BuildGraph/Tasks/SpawnTask.cs
--- a/Engine/Source/Programs/AutomationTool/BuildGraph/Tasks/SpawnTask.cs
+++ b/Engine/Source/Programs/AutomationTool/BuildGraph/Tasks/SpawnTask.cs
@@ -117,21 +117,25 @@ namespace AutomationTool.Tasks
 		{
 			for (int baseIdx = 0; baseIdx < environment.Length;)
 			{
-				int equalsIdx = environment.IndexOf('=', baseIdx);
-				if (equalsIdx == -1)
-				{
-					throw new AutomationException("Missing value in environment variable string '{0}'", environment);
-				}
-
-				int endIdx = environment.IndexOf(separator, equalsIdx + 1);
+				int endIdx = environment.IndexOf(separator, baseIdx);
 				if (endIdx == -1)
 				{
 					endIdx = environment.Length;
 				}
 
-				string name = environment.Substring(baseIdx, equalsIdx - baseIdx).Trim();
-				string value = environment.Substring(equalsIdx + 1, endIdx - (equalsIdx + 1)).Trim();
-				envVars[name] = value;
+				string line = environment.Substring(baseIdx, endIdx - baseIdx);
+				if (!String.IsNullOrWhiteSpace(line))
+				{
+					int equalsIdx = line.IndexOf('=');
+					if (equalsIdx == -1)
+					{
+						throw new AutomationException("Missing value in environment variable string '{0}'", environment.Substring(baseIdx, endIdx - baseIdx));
+					}
+
+					string name = line.Substring(0, equalsIdx).Trim();
+					string value = line.Substring(equalsIdx + 1).Trim();
+					envVars[name] = value;
+				}
 
 				baseIdx = endIdx + 1;
 			}
