BuildGraph: Static analysis fix.

#rnx

diff --git a/Engine/Source/Programs/AutomationTool/BuildGraph/Tasks/SpawnTask.cs b/Engine/Source/Programs/AutomationTool/BuildGraph/Tasks/SpawnTask.cs
--- a/Engine/Source/Programs/AutomationTool/BuildGraph/Tasks/SpawnTask.cs
+++ b/Engine/Source/Programs/AutomationTool/BuildGraph/Tasks/SpawnTask.cs
@@ -126,7 +126,7 @@ namespace AutomationTool.Tasks
 				string line = environment.Substring(baseIdx, endIdx - baseIdx);
 				if (!String.IsNullOrWhiteSpace(line))
 				{
-					int equalsIdx = line.IndexOf('=');
+					int equalsIdx = line.IndexOf('=', StringComparison.Ordinal);
 					if (equalsIdx == -1)
 					{
 						throw new AutomationException("Missing value in environment variable string '{0}'", environment.Substring(baseIdx, endIdx - baseIdx));
