Windows bootstrapper can handle stage-remapped executables.

#rnx

diff --git a/Engine/Source/Programs/AutomationTool/Win/WinPlatform.Automation.cs b/Engine/Source/Programs/AutomationTool/Win/WinPlatform.Automation.cs
--- a/Engine/Source/Programs/AutomationTool/Win/WinPlatform.Automation.cs
+++ b/Engine/Source/Programs/AutomationTool/Win/WinPlatform.Automation.cs
@@ -185,7 +185,8 @@ public class Win64Platform : Platform
 
 						foreach (StagedFileReference StagePath in StagedFiles)
 						{
-							StageBootstrapExecutable(SC, BootstrapExeName, Executable.Path, StagePath, BootstrapArguments);
+							StagedFileReference RemappedStagePath = DeploymentContext.ApplyDirectoryRemap(SC, StagePath);
+							StageBootstrapExecutable(SC, BootstrapExeName, Executable.Path, RemappedStagePath, BootstrapArguments);
 						}
 					}
 				}
