[Gauntlet] Log Sanitizer report using KnownLogEvents.Sanitizer_Address so that it gets caught by [REDACTED_PERSON] matcher

#rnx

diff --git a/Engine/Source/Programs/AutomationTool/AutomationUtils/Matchers/SanitizerEventMatcher.cs b/Engine/Source/Programs/AutomationTool/AutomationUtils/Matchers/SanitizerEventMatcher.cs
--- a/Engine/Source/Programs/AutomationTool/AutomationUtils/Matchers/SanitizerEventMatcher.cs
+++ b/Engine/Source/Programs/AutomationTool/AutomationUtils/Matchers/SanitizerEventMatcher.cs
@@ -3,6 +3,7 @@
 using EpicGames.Core;
 using EpicGames.Horde.Issues.Handlers;
 using Microsoft.Extensions.Logging;
+using System.Collections.Generic;
 using System.Text.RegularExpressions;
 
 #nullable enable
@@ -13,47 +14,67 @@ namespace AutomationUtils.Matchers
 	/// Matcher for LLVM Sanitizer reports
 	/// 
 	/// </summary>
-	class SanitizerEventMatcher : ILogEventMatcher
+	public class SanitizerEventMatcher : ILogEventMatcher
 	{
+		/// <summary>
+		/// Regex pattern that matches Report Start
+		/// </summary>
 		// e.g.
 		// ASAN: =================================================================
 		// TSAN: ==================
-		static readonly Regex s_ReportStartPattern = new Regex(@"(^=================================================================$)|(^==================$)");
+		public static readonly Regex ReportStartPattern = new Regex(@"(^=================================================================$)|(^==================$)", RegexOptions.Multiline | RegexOptions.ExplicitCapture);
 
+		/// <summary>
+		/// Regex pattern that matches Report End
+		/// </summary>
 		// e.g.
 		// ASAN: ==30562==ABORTING
 		// TSAN: ==================
-		static readonly Regex s_ReportEndPattern = new Regex(@"(^==\d+==(\s*\w+)?$)|(^==================$)");
+		public static readonly Regex ReportEndPattern = new Regex(@"(^==\d+==(\s*\w+)?$)|(^==================$)", RegexOptions.Multiline | RegexOptions.ExplicitCapture);
 
+		/// <summary>
+		/// Regex pattern that matches Report Level
+		/// </summary>
 		// e.g.
 		// ASAN: ==30562==ERROR: AddressSanitizer: heap-use-after-free on address 0x617002aa8418 at pc 0x7f98a08bd090 bp 0x7ffc30203af0 sp 0x7ffc30203ae8
 		// TSAN: WARNING: ThreadSanitizer: data race (pid=6049)
-		static readonly Regex s_ReportLevelPattern = new Regex(@"^(==\d+==\s*)?(?<ReportLevel>WARNING|ERROR):\s+(?<SanitizerName>\w+)Sanitizer:\s*(.+)$");
+		public static readonly Regex ReportLevelPattern = new Regex(@"^[\s\t]*(==\d+==\s*)?(?<ReportLevel>WARNING|ERROR):\s+(?<SanitizerName>\w+)Sanitizer:\s*(?<Summary>.+)$", RegexOptions.Multiline | RegexOptions.ExplicitCapture);
 
+		/// <summary>
+		/// Regex pattern that matches Stack Trace
+		/// </summary>
 		// e.g. Note line number, column number and BuildId are optional
 		// ASAN: #0 0x7f98a08bd08f in UObjectBase::GetFName() const /mnt/somepath/file.h:90:10
 		// TSAN: #0 MyType<Type, OtherType>::Function(unsigned long*&) const /mnt/somepath/file.h:90:10 (BinaryName+0x2bc102e0) (BuildId: 8b17597a9f5444a0)
-		static readonly Regex s_StackTracePattern = new Regex(
-			@"^\s+\#[\d]+\s+(0x[0-9a-fA-F]+\s+in\s+)?" + // The optional capture group is to handle ASAN's slightly different stack trace syntax from TSAN
+		public static readonly Regex StackTracePattern = new Regex(
+			@"^[\s\t]+\#[\d]+\s+(0x[0-9a-fA-F]+\s+in\s+)?" + // The optional capture group is to handle ASAN's slightly different stack trace syntax from TSAN
 			@"(?<Symbol>.+?)\s" +
 			@"(?<SourceFile>(\/|\w:).+?)" +
 			@"(:(?<Line>\d+?))?" +
 			@"(:(?<Column>\d+?))?" +
 			@"(\s+\(.+?\+(?<Address>0x[0-9a-fA-F]+?)\))?" +
-			@"(\s+?\(BuildId:\s+?(?<BuildId>.+?)\))?$");
+			@"(\s+?\(BuildId:\s+?(?<BuildId>.+?)\))?$", RegexOptions.Multiline | RegexOptions.ExplicitCapture);
 
+		/// <summary>
+		/// Regex pattern that matches Summary
+		/// </summary>
 		// e.g.
 		// ASAN: SUMMARY: AddressSanitizer: heap-use-after-free /some/path.cpp:169:33 in UObjectBase::GetFName() const
 		// TSAN: SUMMARY: ThreadSanitizer: data race /some/path.cpp:169:33 in MyType::Foo()
-		static readonly Regex s_SummaryPattern = new Regex(
-			@"^SUMMARY:\s+(?<SanitizerName>\w+)Sanitizer:\s*" +
+		public static readonly Regex SummaryPattern = new Regex(
+			@"^[\s\t]*SUMMARY:\s+(?<SanitizerName>\w+)Sanitizer:\s*" +
 			@"(?<SummaryReason>.+)\s" +
 			@"(?<SourceFile>(\/|\w:).+?)" +
 			@"(:(?<Line>\d+?))?" +
 			@"(:(?<Column>\d+?))?\s+?in\s+?" +
-			@"(?<Symbol>.+)$");
-
-		EventId ConvertSanitizerNameToEventId(string sanitizerName)
+			@"(?<Symbol>.+)$", RegexOptions.Multiline | RegexOptions.ExplicitCapture);
+
+		/// <summary>
+		/// Convert SanitizerName to EventId
+		/// </summary>
+		/// <param name="sanitizerName"></param>
+		/// <returns></returns>
+		public static EventId ConvertSanitizerNameToEventId(string sanitizerName)
 		{
 			switch (sanitizerName)
 			{
@@ -64,7 +85,12 @@ namespace AutomationUtils.Matchers
 			return KnownLogEvents.Sanitizer;
 		}
 
-		LogLevel ConvertReportLevel(string reportLevel)
+		/// <summary>
+		/// Convert ReportLevel string to LogLevel enum
+		/// </summary>
+		/// <param name="reportLevel"></param>
+		/// <returns></returns>
+		public static LogLevel ConvertReportLevel(string reportLevel)
 		{
 			switch (reportLevel)
 			{
@@ -74,10 +100,38 @@ namespace AutomationUtils.Matchers
 			return LogLevel.Error;
 		}
 
+		private static void AddProperty(LogValue Value, Group Group, Dictionary<string, object> Properties)
+		{
+			Value.Text = Group.Value;
+			Properties.Add(Value.Type.ToString(), Value);
+		}
+
+		/// <summary>
+		/// Add Sanitizer Summary information to input Key/Value Pair Properties
+		/// </summary>
+		/// <param name="Message"></param>
+		/// <param name="Properties"></param>
+		/// <returns></returns>
+		public static bool AddSanitizerSummaryProperties(string Message, Dictionary<string, object> Properties)
+		{
+			Match Match = SummaryPattern.Match(Message);
+			if (Match.Success)
+			{
+				AddProperty(SanitizerIssueHandler.SanitizerName, Match.Groups["SanitizerName"], Properties);
+				AddProperty(SanitizerIssueHandler.SummaryReason, Match.Groups["SummaryReason"], Properties);
+				AddProperty(SanitizerIssueHandler.SummarySourceFile, Match.Groups["SourceFile"], Properties);
+
+				return true;
+			}
+
+			return false;
+		}
+
+		/// <inheritdoc/> 
 		public LogEventMatch? Match(ILogCursor cursor)
 		{
 			Match? match;
-			if (cursor.TryMatch(s_ReportStartPattern, out match))
+			if (cursor.TryMatch(ReportStartPattern, out match))
 			{
 				EventId sanitizerId = KnownLogEvents.Sanitizer;
 				LogLevel reportLevel = LogLevel.Information;
@@ -88,7 +142,7 @@ namespace AutomationUtils.Matchers
 				while (builder.Current.CurrentLine != null)
 				{
 
-					if (builder.Current.TryMatch(s_StackTracePattern, out match))
+					if (builder.Current.TryMatch(StackTracePattern, out match))
 					{
 						do
 						{
@@ -100,11 +154,11 @@ namespace AutomationUtils.Matchers
 							builder.TryAnnotate(match.Groups["BuildId"], "");
 
 							builder.MoveNext();
-						} while (builder.Current.TryMatch(s_StackTracePattern, out match));
+						} while (builder.Current.TryMatch(StackTracePattern, out match));
 					}
 					else
 					{
-						if (builder.Current.TryMatch(s_ReportLevelPattern, out match))
+						if (builder.Current.TryMatch(ReportLevelPattern, out match))
 						{
 							sanitizerId = ConvertSanitizerNameToEventId(match.Groups["SanitizerName"].Value);
 							reportLevel = ConvertReportLevel(match.Groups["ReportLevel"].Value);
@@ -113,7 +167,7 @@ namespace AutomationUtils.Matchers
 							builder.Annotate(match.Groups["SanitizerName"], SanitizerIssueHandler.SanitizerName);
 						}
 
-						if (builder.Current.TryMatch(s_SummaryPattern, out match))
+						if (builder.Current.TryMatch(SummaryPattern, out match))
 						{
 							// Used by IssueHandler
 							builder.Annotate(match.Groups["SummaryReason"], SanitizerIssueHandler.SummaryReason);
@@ -125,7 +179,7 @@ namespace AutomationUtils.Matchers
 							builder.TryAnnotate(match.Groups["Column"], LogEventMarkup.ColumnNumber);
 							builder.AnnotateSymbol(match.Groups["Symbol"]);
 						}
-						else if (builder.Current.IsMatch(s_ReportEndPattern))
+						else if (builder.Current.IsMatch(ReportEndPattern))
 						{
 							return builder.ToMatch(LogEventPriority.AboveNormal, reportLevel, sanitizerId);
 						}
diff --git a/Engine/Source/Programs/AutomationTool/Gauntlet/Unreal/Automation/UE.Automation.cs b/Engine/Source/Programs/AutomationTool/Gauntlet/Unreal/Automation/UE.Automation.cs
--- a/Engine/Source/Programs/AutomationTool/Gauntlet/Unreal/Automation/UE.Automation.cs
+++ b/Engine/Source/Programs/AutomationTool/Gauntlet/Unreal/Automation/UE.Automation.cs
@@ -9,7 +9,6 @@ using System.Linq;
 using UnrealBuildTool;
 using EpicGames.Core;
 using Log = Gauntlet.Log;
-using Microsoft.Extensions.Logging;
 
 namespace UE
 {
@@ -1015,15 +1014,7 @@ namespace UE
 				IEnumerable<UnrealAutomatedTestResult> FailedTests = AllTests.Where(T => T.IsComplete && T.HasFailed);
 				IEnumerable<UnrealAutomatedTestResult> TestsWithWarnings = AllTests.Where(T => T.HasSucceeded && T.HasWarnings);
 
-				Func<IEnumerable<UnrealAutomationEvent>, IEnumerable<UnrealAutomationEvent>> CapErrorOrWarningList = (E) =>
-				{
-					if (E.Count() > MaxEventsDisplayPerTest)
-					{
-						E = E.Take(MaxEventsDisplayPerTest);
-					}
-					return E;
-				};
-				Action<IEnumerable<UnrealAutomationEvent>> NotifyMoreIfNeeded = E =>
+				Action<IEnumerable<object>> NotifyMoreIfNeeded = E =>
 				{
 					if (E.Count() > MaxEventsDisplayPerTest)
 					{
@@ -1048,20 +1039,13 @@ namespace UE
 							{
 								string Message = !Result.IsComplete ? " * Test '{Name}' did not complete." : " * Test '{Name}' failed.";
 								Log.Error(KnownLogEvents.Gauntlet_UnrealEngineTestEvent, Message, Result.FullTestPath);
-								var Errors = CapErrorOrWarningList(Result.ErrorEvents.Distinct());
-								foreach (var Error in Errors)
+								var Errors = Result.GetErrorEntries().Distinct().Select(E => E.AsLogEvent());
+								foreach (LogEvent Error in Errors.Take(MaxEventsDisplayPerTest))
 								{
-									if (Error.IsCriticalFailure)
-									{
-										Log.Error(KnownLogEvents.Gauntlet_FatalEvent, "    {Callstack}", Error.FormatToString());
-									}
-									else
-									{
-										Log.Error(KnownLogEvents.Gauntlet_UnrealEngineTestEvent, "    " + Error.FormatToString());
-									}
+									Log.Error(Error.Id, "    " + (Error.Format ?? Error.Message), Args: Error.Properties?.Select(P => P.Value).ToArray());
 								}
-								NotifyMoreIfNeeded(Result.ErrorEvents);
-								var Warnings = CapErrorOrWarningList(Result.WarningEvents.Distinct());
+								NotifyMoreIfNeeded(Errors);
+								var Warnings = Result.WarningEvents.Distinct();
 								foreach (var Warning in Warnings)
 								{
 									Log.Warning(KnownLogEvents.Gauntlet_UnrealEngineTestEvent, "    " + Warning.FormatToString());
@@ -1095,12 +1079,12 @@ namespace UE
 						if (Errors.Any() || Warnings.Any())
 						{
 							Log.Info("   See log above for details.");
-							foreach (var Error in CapErrorOrWarningList(Errors))
+							foreach (var Error in Errors.Take(MaxEventsDisplayPerTest))
 							{
 								Log.Error(KnownLogEvents.Gauntlet_UnrealEngineTestEvent, "    " + Error.FormatToString());
 							}
 							NotifyMoreIfNeeded(Errors);
-							foreach (var Warning in CapErrorOrWarningList(Warnings))
+							foreach (var Warning in Warnings.Take(MaxEventsDisplayPerTest))
 							{
 								Log.Warning(KnownLogEvents.Gauntlet_UnrealEngineTestEvent, "    " + Warning.FormatToString());
 							}
@@ -1122,17 +1106,10 @@ namespace UE
 							foreach (UnrealAutomatedTestResult Result in FailedTests)
 							{
 								Log.Error(KnownLogEvents.Gauntlet_UnrealEngineTestEvent, " * Test '{Name}' failed.", Result.FullTestPath);
-								IEnumerable<UnrealAutomationEvent> Events = Result.ErrorEvents.Distinct();
-								foreach (var Event in CapErrorOrWarningList(Events))
+								var Events = Result.GetErrorEntries().Distinct().Select(E => E.AsLogEvent());
+								foreach (LogEvent Error in Events.Take(MaxEventsDisplayPerTest))
 								{
-									if (Event.IsCriticalFailure)
-									{
-										Log.Error(KnownLogEvents.Gauntlet_FatalEvent, "    {Callstack}", Event.FormatToString());
-									}
-									else
-									{
-										Log.Error(KnownLogEvents.Gauntlet_UnrealEngineTestEvent, "    " + Event.FormatToString());
-									}
+									Log.Error(Error.Id, "    " + (Error.Format ?? Error.Message), Args: Error.Properties?.Select(P => P.Value).ToArray());
 								}
 								NotifyMoreIfNeeded(Events);
 								Log.Info("");
@@ -1148,7 +1125,7 @@ namespace UE
 								Log.Warning(KnownLogEvents.Gauntlet_UnrealEngineTestEvent, " * Test '{Name}' completed with warnings.", Result.FullTestPath);
 								// only show the first N items
 								IEnumerable<UnrealAutomationEvent> WarningEvents = Result.WarningEvents.Distinct();
-								foreach (var Event in CapErrorOrWarningList(WarningEvents))
+								foreach (var Event in WarningEvents.Take(MaxEventsDisplayPerTest))
 								{
 									Log.Warning(KnownLogEvents.Gauntlet_UnrealEngineTestEvent, "    " + Event.FormatToString());
 								}
@@ -1166,7 +1143,7 @@ namespace UE
 								Log.Error(KnownLogEvents.Gauntlet_UnrealEngineTestEvent, " * Test '{Name}' did not run or complete.", Result.FullTestPath);
 								// only show the first N items
 								IEnumerable<UnrealAutomationEvent> ErrorAndWarningEvents = Result.WarningAndErrorEvents.Distinct();
-								foreach (var Event in CapErrorOrWarningList(ErrorAndWarningEvents))
+								foreach (var Event in ErrorAndWarningEvents.Take(MaxEventsDisplayPerTest))
 								{
 									Log.Error(KnownLogEvents.Gauntlet_UnrealEngineTestEvent, "    " + Event.FormatToString());
 								}
diff --git a/Engine/Source/Programs/AutomationTool/Gauntlet/Unreal/Base/Gauntlet.UnrealTestEvent.cs b/Engine/Source/Programs/AutomationTool/Gauntlet/Unreal/Base/Gauntlet.UnrealTestEvent.cs
--- a/Engine/Source/Programs/AutomationTool/Gauntlet/Unreal/Base/Gauntlet.UnrealTestEvent.cs
+++ b/Engine/Source/Programs/AutomationTool/Gauntlet/Unreal/Base/Gauntlet.UnrealTestEvent.cs
@@ -31,6 +31,9 @@ namespace Gauntlet
 		// True if this is an ensure (Gauntlet does not define a level for this, but we log differently).
 		public bool IsEnsure { get; protected set; }
 
+		// True if this is a Sanitizer error report.
+		public bool IsSanReport { get; protected set; }
+
 		// True if the event severity is Error or Fatal
 		public bool IsError => Severity == EventSeverity.Error || Severity == EventSeverity.Fatal;
 
@@ -38,9 +41,9 @@ namespace Gauntlet
 		public bool IsWarning => Severity == EventSeverity.Warning;
 
 		// Constructor that requires all properties
-		public UnrealTestEvent(EventSeverity InSeverity, string InSummary, IEnumerable<string> InDetails, UnrealLog.CallstackMessage InCallstack = null)
+		public UnrealTestEvent(DateTime InTime, EventSeverity InSeverity, string InSummary, IEnumerable<string> InDetails, UnrealLog.CallstackMessage InCallstack = null)
 		{
-			Time = DateTime.UtcNow;
+			Time = InTime;
 			Severity = InSeverity;
 			Summary = InSummary;
 			Details = InDetails.ToArray();
@@ -48,13 +51,19 @@ namespace Gauntlet
 			if (InCallstack != null)
 			{
 				IsEnsure = InCallstack.IsEnsure;
+				IsSanReport = InCallstack.IsSanReport;
 				Callstack = InCallstack.Callstack;
 			}
 			else
 			{
 				IsEnsure = false;
+				IsSanReport = false;
 				Callstack = Enumerable.Empty<string>();
 			}
 		}
+
+		// Constructor that requires all properties but DateTime
+		public UnrealTestEvent(EventSeverity InSeverity, string InSummary, IEnumerable<string> InDetails, UnrealLog.CallstackMessage InCallstack = null)
+			: this(DateTime.UtcNow, InSeverity, InSummary, InDetails, InCallstack) { }
 	}
 }
diff --git a/Engine/Source/Programs/AutomationTool/Gauntlet/Unreal/Base/Gauntlet.UnrealTestNode.cs b/Engine/Source/Programs/AutomationTool/Gauntlet/Unreal/Base/Gauntlet.UnrealTestNode.cs
--- a/Engine/Source/Programs/AutomationTool/Gauntlet/Unreal/Base/Gauntlet.UnrealTestNode.cs
+++ b/Engine/Source/Programs/AutomationTool/Gauntlet/Unreal/Base/Gauntlet.UnrealTestNode.cs
@@ -8,6 +8,8 @@ using UnrealBuildTool;
 using System.Text.RegularExpressions;
 using System.Linq;
 using EpicGames.Core;
+using Microsoft.Extensions.Logging;
+using AutomationUtils.Matchers;
 
 namespace Gauntlet
 {
@@ -2005,17 +2007,35 @@ namespace Gauntlet
 				List<string> Callstack = new List<string>();
 				if (Event.Callstack.Any())
 				{
-					Callstack = Event.Callstack.Take(MaxCallstackLines).ToList();
-					if (Event.Callstack.Count() > MaxCallstackLines)
+					if (CommandUtils.IsBuildMachine)
 					{
-						Callstack.Add("See log for full callstack");
+						Callstack = Event.Callstack.ToList();
+					}
+					else
+					{
+						Callstack = Event.Callstack.Take(MaxCallstackLines).ToList();
+						if (Event.Callstack.Count() > MaxCallstackLines)
+						{
+							Callstack.Add("See log for full callstack");
+						}
 					}
 				}
 				else
 				{
 					Callstack.Add("Could not parse callstack. See log for full callstack");
 				}
-				Log.Error(KnownLogEvents.Gauntlet_FatalEvent, " * Fatal Error: {Summary}\n{Callstack}", Event.Summary, string.Join("\n", Callstack.Select(C=>"    "+C)));
+				string CallstackString = string.Join("\n", Callstack.Select(C => "    " + C));
+				Dictionary<string, object> Properties = new Dictionary<string, object>() {
+					{ "Summary", Event.Summary },
+					{ "Callstack", CallstackString }
+				};
+				EventId EventType = KnownLogEvents.Gauntlet_FatalEvent;
+				if (Event.IsSanReport && CommandUtils.IsBuildMachine
+					&& SanitizerEventMatcher.AddSanitizerSummaryProperties(CallstackString, Properties))
+				{
+					EventType = SanitizerEventMatcher.ConvertSanitizerNameToEventId(Properties.GetValueOrDefault("SanitizerName")?.ToString());
+				}
+				Log.Error(EventType, " * Fatal Error: {Summary}\n{Callstack}", Args: Properties.Values.ToArray());
 				Log.Info("");
 			}
 
@@ -2024,11 +2044,17 @@ namespace Gauntlet
 				List<string> Callstack = new List<string>();
 				if (Event.Callstack.Any())
 				{
-					Callstack = Event.Callstack.Take(MaxCallstackLines).ToList();
-
-					if (Event.Callstack.Count() > MaxCallstackLines)
+					if (CommandUtils.IsBuildMachine)
+					{
+						Callstack = Event.Callstack.ToList();
+					}
+					else
 					{
-						Callstack.Add("See log for full callstack");
+						Callstack = Event.Callstack.Take(MaxCallstackLines).ToList();
+						if (Event.Callstack.Count() > MaxCallstackLines)
+						{
+							Callstack.Add("See log for full callstack");
+						}
 					}
 				}
 				else
diff --git a/Engine/Source/Programs/AutomationTool/Gauntlet/Unreal/Utils/Gauntlet.UnrealLogParser.cs b/Engine/Source/Programs/AutomationTool/Gauntlet/Unreal/Utils/Gauntlet.UnrealLogParser.cs
--- a/Engine/Source/Programs/AutomationTool/Gauntlet/Unreal/Utils/Gauntlet.UnrealLogParser.cs
+++ b/Engine/Source/Programs/AutomationTool/Gauntlet/Unreal/Utils/Gauntlet.UnrealLogParser.cs
@@ -8,6 +8,7 @@ using System.Text;
 using EpicGames.Core;
 using Microsoft.Extensions.Logging;
 using Logging = Microsoft.Extensions.Logging;
+using AutomationUtils.Matchers;
 
 namespace Gauntlet
 {
@@ -100,6 +101,7 @@ namespace Gauntlet
 			public string Message;
 			public string[] Callstack;
 			public bool IsEnsure;
+			public bool IsSanReport;
 			public bool IsPostMortem;
 
 			/// <summary>
@@ -975,20 +977,25 @@ namespace Gauntlet
 			/// remove anything inside the callstack starting with [2022.12.02-15.22.40:688][618]
 
 			List<UnrealLog.CallstackMessage> ASanReports = new List<UnrealLog.CallstackMessage>();
-			string InitPattern = @"==([0-9]+)==ERROR: ([^\n]+)";
-			string LineStartsWithTimeStamp = @"^\[[0-9.:-]+\]\[[0-9]+\]";
+			Regex InitPattern = SanitizerEventMatcher.ReportLevelPattern;
+			Regex LineStartsWithTimeStamp = new Regex(@"^[\s\t]*\[[0-9.:-]+\]\[[\s0-9]+\]");
 
-			MatchCollection Matches = Regex.Matches(Content, InitPattern, RegexOptions.IgnoreCase);
+			MatchCollection Matches = InitPattern.Matches(Content);
 
 			foreach (Match TraceInitMatch in Matches)
 			{
-				string TraceID = TraceInitMatch.Groups[1].Value;
-				string EndPattern = @$"(=={TraceID}==ABORTING)|(End of Address Sanitizer report)";
+				if (SanitizerEventMatcher.ConvertReportLevel(TraceInitMatch.Groups["ReportLevel"].Value) != Logging.LogLevel.Error)
+				{
+					continue;
+				}
+
+				Regex EndPattern = SanitizerEventMatcher.ReportEndPattern;
 				int TraceInitIndex = TraceInitMatch.Index + TraceInitMatch.Length;
 
 				UnrealLog.CallstackMessage NewTrace = new UnrealLog.CallstackMessage();
+				NewTrace.IsSanReport = true;
 				NewTrace.Position = TraceInitMatch.Index;
-				NewTrace.Message = TraceInitMatch.Groups[2].Value;
+				NewTrace.Message = $"{TraceInitMatch.Groups["SanitizerName"].Value}Sanitizer: {TraceInitMatch.Groups["Summary"].Value}";
 
 				// If the regex matches the very end of the string, the substring will get an invalid range.
 				string ErrorContent = Content.Length <= TraceInitIndex
@@ -996,7 +1003,7 @@ namespace Gauntlet
 
 				if (!string.IsNullOrEmpty(ErrorContent))
 				{
-					Match MsgMatch = Regex.Match(ErrorContent, EndPattern);
+					Match MsgMatch = EndPattern.Match(ErrorContent);
 					if (MsgMatch.Success)
 					{
 						string MsgContent = ErrorContent.Substring(0, MsgMatch.Index);
@@ -1004,7 +1011,7 @@ namespace Gauntlet
 						List<string> Backtrace = new List<string>();
 						foreach (string Line in MsgContent.Split("\n"))
 						{
-							Match IsTimeStampLine = Regex.Match(Line, LineStartsWithTimeStamp);
+							Match IsTimeStampLine = LineStartsWithTimeStamp.Match(Line);
 							if (!IsTimeStampLine.Success)
 							{
 								Backtrace.Add(Line);
diff --git a/Engine/Source/Programs/AutomationTool/Gauntlet/Unreal/Utils/Gauntlet.UnrealTestResults.cs b/Engine/Source/Programs/AutomationTool/Gauntlet/Unreal/Utils/Gauntlet.UnrealTestResults.cs
--- a/Engine/Source/Programs/AutomationTool/Gauntlet/Unreal/Utils/Gauntlet.UnrealTestResults.cs
+++ b/Engine/Source/Programs/AutomationTool/Gauntlet/Unreal/Utils/Gauntlet.UnrealTestResults.cs
@@ -7,6 +7,12 @@ using System.Text.Json;
 using System.Text.Json.Serialization;
 using System.Linq;
 using AutomationTool;
+using EpicGames.Core;
+using Microsoft.Extensions.Logging;
+using Logging = Microsoft.Extensions.Logging;
+using System.Globalization;
+using AutomationUtils.Matchers;
+using Polly;
 
 namespace Gauntlet
 {
@@ -69,6 +75,7 @@ namespace Gauntlet
 		public string Artifact { get; set; }
 
 		const string CriticalFailureString = "critical failure";
+		const string SanitizerReportString = "Sanitizer: ";
 
 		public UnrealAutomationEvent()
 		{ }
@@ -119,6 +126,18 @@ namespace Gauntlet
 			}
 		}
 
+		/// <summary>
+		/// True if the event contains a Sanitizer report
+		/// </summary>
+		[JsonIgnore]
+		public bool IsSanReport
+		{
+			get
+			{
+				return IsError && Message.Contains(SanitizerReportString);
+			}
+		}
+
 		public string FormatToString()
 		{
 			return Message;
@@ -153,6 +172,75 @@ namespace Gauntlet
 		public int LineNumber { get; set; }
 		[JsonPropertyName("timeStamp")]
 		public string Timestamp { get; set; }
+
+		private DateTime GetTimestampAsDateTime()
+		{
+			if (string.IsNullOrEmpty(Timestamp) || Timestamp == "0001.01.01-00.00.00")
+			{
+				// Special case: when UE Test DateTime is set this way, it means the value was 0 or null before getting converted to json.
+				Timestamp = DateTime.UtcNow.ToString("yyyy.MM.dd-HH.mm.ss");
+			}
+			DateTime Time = DateTime.UtcNow;
+			DateTime.TryParseExact(Timestamp, "yyyy.MM.dd-HH.mm.ss", CultureInfo.InvariantCulture, DateTimeStyles.None, out Time);
+			return Time;
+		}
+
+		public LogEvent AsLogEvent()
+		{
+			DateTime Time = GetTimestampAsDateTime();
+			Logging.LogLevel Level = Event.IsError ?
+										(Event.IsCriticalFailure ?
+											Logging.LogLevel.Critical : Logging.LogLevel.Error)
+										: (Event.IsWarning ?
+											Logging.LogLevel.Warning : Logging.LogLevel.Information);
+
+			EventId EventIdType = Event.IsCriticalFailure ? KnownLogEvents.Gauntlet_FatalEvent : KnownLogEvents.Gauntlet_UnrealEngineTestEvent;
+
+			string Message = Event.FormatToString();
+			string Format = null;
+			Dictionary<string, object> Properties = null;
+			if (Event.IsCriticalFailure)
+			{
+				Properties = new Dictionary<string, object>() { { "Callstack", Message } };
+				if (Event.IsSanReport
+					&& SanitizerEventMatcher.AddSanitizerSummaryProperties(Message, Properties))
+				{
+					EventIdType = SanitizerEventMatcher.ConvertSanitizerNameToEventId(Properties.GetValueOrDefault("SanitizerName")?.ToString());
+				}
+				Format = "{Callstack}";
+			}
+			if (!string.IsNullOrEmpty(Event.Context))
+			{
+				Properties.Add("Context", Event.Context);
+				Format = "[{Context}] " + Format?? Message;
+			}
+			if (!string.IsNullOrEmpty(Filename))
+			{
+				string FileReference = $"{Filename}:{LineNumber}";
+				Properties.Add("FileReference", FileReference);
+				Format = (Format?? Message) + " [{FileReference}]";
+			}
+
+			return new LogEvent(Time, Level, EventIdType, Message, Format, Properties, null);
+		}
+
+		protected bool Equals(UnrealAutomationEntry other)
+		{
+			return string.Equals(Event, other.Event);
+		}
+
+		public override bool Equals(object obj)
+		{
+			if (ReferenceEquals(null, obj)) return false;
+			if (ReferenceEquals(this, obj)) return true;
+			if (obj.GetType() != this.GetType()) return false;
+			return Equals((UnrealAutomationEntry)obj);
+		}
+
+		public override int GetHashCode()
+		{
+			return Event.GetHashCode();
+		}
 	}
 	public class UnrealAutomatedTestResult
 	{
@@ -190,10 +278,10 @@ namespace Gauntlet
 			var Event = new UnrealAutomationEvent(EventType, Message, bIsCriticalFailure);
 			var Entry = new UnrealAutomationEntry();
 			Entry.Event = Event;
-			Entry.Timestamp = System.DateTime.Now.ToString("yyyy.MM.dd-HH.mm.ss");
+			Entry.Timestamp = System.DateTime.UtcNow.ToString("yyyy.MM.dd-HH.mm.ss");
 			Entries.Add(Entry);
 
-			switch(Event.Type)
+			switch (Event.Type)
 			{
 				case EventType.Error:
 					Errors++;
@@ -280,6 +368,21 @@ namespace Gauntlet
 				return Entries.Where(E => E.Event.IsError || E.Event.IsWarning).Select(E => E.Event);
 			}
 		}
+
+		public IEnumerable<UnrealAutomationEntry> GetErrorEntries()
+		{
+			return Entries.Where(E => E.Event.IsError);
+		}
+
+		public IEnumerable<UnrealAutomationEntry> GetWarningEntries()
+		{
+			return Entries.Where(E => E.Event.IsWarning);
+		}
+
+		public IEnumerable<UnrealAutomationEntry> GetWarningAndErrorEntries()
+		{
+			return Entries.Where(E => E.Event.IsError || E.Event.IsWarning);
+		}
 	}
 	public class UnrealAutomatedTestPassResults
 	{
