BuildCookRun: Always try to copy the default UnrealPak.log file to the UAT logs folder after every invocation of UnrealPak.
This log file may contain more verbose logging than the one generated by RunAndLog() based on stdout/stderr.
Route all calls to UnrealPak through the same RunUnrealPak function.

#rnx

diff --git a/Engine/Source/Programs/AutomationTool/Scripts/CopyBuildToStagingDirectory.Automation.cs b/Engine/Source/Programs/AutomationTool/Scripts/CopyBuildToStagingDirectory.Automation.cs
--- a/Engine/Source/Programs/AutomationTool/Scripts/CopyBuildToStagingDirectory.Automation.cs
+++ b/Engine/Source/Programs/AutomationTool/Scripts/CopyBuildToStagingDirectory.Automation.cs
@@ -418,15 +418,20 @@ namespace AutomationScripts
 			}
 		}
 
-		static public void RunUnrealPak(ProjectParams Params, Dictionary<string, string> UnrealPakResponseFile, FileReference OutputLocation, List<OrderFile> PakOrderFileLocations, string PlatformOptions, bool Compressed, EncryptionAndSigning.CryptoSettings CryptoSettings, FileReference CryptoKeysCacheFilename, String PatchSourceContentPath, string EncryptionKeyGuid, List<OrderFile> SecondaryPakOrderFileLocations = null)
+		static public void RunUnrealPak(string Arguments, string Scenario)
 		{
-			if (UnrealPakResponseFile.Count < 1)
+			Logger.LogInformation("Running UnrealPak with arguments: {Arguments}", Arguments);
+			RunAndLog(CmdEnv, GetUnrealPakLocation().FullName, Arguments, Options: ERunOptions.Default | ERunOptions.UTF8Output);
+
+			// try to copy the default UnrealPak.log file to the UAT log folder,
+			// this log file may contain more verbose logging than the one generated by RunAndLog() based on stdout/stderr.
+			string UnrealPakLogFileName = "UnrealPak_" + Scenario;
+			string SourceLogFile = CombinePaths(DirectoryReference.Combine(Unreal.RootDirectory, "Engine", "Programs", "UnrealPak", "Saved", "Logs").FullName, "UnrealPak.log");
+			string DestLogFile = LogUtils.GetUniqueLogName(CombinePaths(CommandUtils.CmdEnv.LogFolder, UnrealPakLogFileName));
+			if (InternalUtils.SafeFileExists(SourceLogFile, true))
 			{
-				return;
+				CopyFile_NoExceptions(SourceLogFile, DestLogFile);
 			}
-
-			string Arguments = GetUnrealPakArguments(Params.RawProjectPath, UnrealPakResponseFile, OutputLocation, PakOrderFileLocations, PlatformOptions, Compressed, CryptoSettings, CryptoKeysCacheFilename, PatchSourceContentPath, EncryptionKeyGuid, SecondaryPakOrderFileLocations, Params.Unattended);
-			RunAndLog(CmdEnv, GetUnrealPakLocation().FullName, Arguments, Options: ERunOptions.Default | ERunOptions.UTF8Output);
 		}
 
 		static public void LogDeploymentContext(DeploymentContext SC)
@@ -892,8 +897,7 @@ namespace AutomationScripts
 									Arguments += " -unattended";
 								}
 
-								Logger.LogInformation("Running UnrealPak with arguments: {UnrealPakParams}", Arguments);
-								RunAndLog(CmdEnv, GetUnrealPakLocation().FullName, Arguments, Options: ERunOptions.Default | ERunOptions.UTF8Output);
+								RunUnrealPak(Arguments, "ZenAutoLaunch");
 								bAttemptAutoLaunchOnFailure = false;
 							}
 							else
@@ -1626,8 +1630,7 @@ namespace AutomationScripts
 						FileReference OutputFile = FileReference.Combine(SC.ProjectRoot, "Intermediate", "Config", TargetPlatformName, "BinaryConfig.ini");
 						String UnrealPakParams = String.Format("MakeBinaryConfig -Project=\"{0}\" -Platform={1} -OutputFile=\"{2}\" -StagedPluginsFile=\"{3}\"", SC.RawProjectPath, ConfigHierarchy.GetIniPlatformName(ThisPlatform.IniPlatformType), OutputFile.FullName, PluginListFile.FullName);
 
-						Logger.LogInformation("Running UnrealPak with arguments: {UnrealPakParams}", UnrealPakParams);
-						RunAndLog(CmdEnv, GetUnrealPakLocation().FullName, UnrealPakParams, Options: ERunOptions.Default | ERunOptions.UTF8Output);
+						RunUnrealPak(UnrealPakParams, "MakeBinaryConfig");
 
 						SC.StageFile(StagedFileType.UFS, OutputFile, StagedFileReference.Combine(SC.RelativeProjectRootForStage, "Config", OutputFile.GetFileName()));
 					}
@@ -2689,7 +2692,12 @@ namespace AutomationScripts
 			Dictionary<string, string> PakResponseFile = CreatePakResponseFileFromStagingManifest(SC, SC.CrashReporterUFSFiles);
 			FileReference OutputLocation = FileReference.Combine(SC.RuntimeRootDir, "Engine", "Programs", "CrashReportClient", "Content", "Paks", "CrashReportClient.pak");
 
-			RunUnrealPak(Params, PakResponseFile, OutputLocation, null, null, Params.Compressed, null, null, null, null, null);
+			if (PakResponseFile.Count < 1)
+			{
+				return;
+			}
+			string Arguments = GetUnrealPakArguments(Params.RawProjectPath, PakResponseFile, OutputLocation, null, null, Params.Compressed, null, null, null, null, null, Params.Unattended);
+			RunUnrealPak(Arguments, "CrashReportClient");
 		}
 
 		/// <summary>
@@ -3670,7 +3678,7 @@ namespace AutomationScripts
 				);
 				Arguments.AppendFormat(" {0}", CommonArguments);
 				Arguments.AppendFormat(" -CreateMultiple={0}", CommandUtils.MakePathSafeToUseWithCommandLine(PakCommandsFileName));
-				RunAndLog(CmdEnv, GetUnrealPakLocation().FullName, Arguments.ToString(), Options: ERunOptions.Default | ERunOptions.UTF8Output);
+				RunUnrealPak(Arguments.ToString(), "CreateMultiplePaks");
 			}
 
 			if (IoStoreCommands.Count > 0)
@@ -4161,8 +4169,7 @@ namespace AutomationScripts
 				CommandletParams += " -unattended";
 			}
 
-			Logger.LogInformation("Running UnrealPak with arguments: {CommandletParams}", CommandletParams);
-			RunAndLog(CmdEnv, GetUnrealPakLocation().FullName, CommandletParams, Options: ERunOptions.Default | ERunOptions.UTF8Output);
+			RunUnrealPak(CommandletParams, "CreateIoStoreContainers");
 		}
 
 		/// <summary>
