Horde: Remove unused using directive.

#rnx

diff --git a/Engine/Source/Programs/AutomationTool/BuildGraph/Tasks/HordeGetSecretsTask.cs b/Engine/Source/Programs/AutomationTool/BuildGraph/Tasks/HordeGetSecretsTask.cs
--- a/Engine/Source/Programs/AutomationTool/BuildGraph/Tasks/HordeGetSecretsTask.cs
+++ b/Engine/Source/Programs/AutomationTool/BuildGraph/Tasks/HordeGetSecretsTask.cs
@@ -9,7 +9,6 @@ using System.Xml;
 using EpicGames.Core;
 using EpicGames.Horde;
 using EpicGames.Horde.Secrets;
-using IdentityModel.Client;
 using Microsoft.Extensions.DependencyInjection;
 using Microsoft.Extensions.Logging;
 
