Horde: Fix some static analyzer warnings.

#rnx

diff --git a/Engine/Source/Programs/AutomationTool/AutomationUtils/Matchers/SanitizerEventMatcher.cs b/Engine/Source/Programs/AutomationTool/AutomationUtils/Matchers/SanitizerEventMatcher.cs
--- a/Engine/Source/Programs/AutomationTool/AutomationUtils/Matchers/SanitizerEventMatcher.cs
+++ b/Engine/Source/Programs/AutomationTool/AutomationUtils/Matchers/SanitizerEventMatcher.cs
@@ -7,6 +7,7 @@ using System.Collections.Generic;
 using System.Text.RegularExpressions;
 
 #nullable enable
+#pragma warning disable MA0016
 
 namespace AutomationUtils.Matchers
 {
