Substitute properties from Sanitizer report in Log format + fix Context and SourceFile properties in UE.Automation test report summary.

#rnx

diff --git a/Engine/Source/Programs/AutomationTool/AutomationUtils/Matchers/SanitizerEventMatcher.cs b/Engine/Source/Programs/AutomationTool/AutomationUtils/Matchers/SanitizerEventMatcher.cs
--- a/Engine/Source/Programs/AutomationTool/AutomationUtils/Matchers/SanitizerEventMatcher.cs
+++ b/Engine/Source/Programs/AutomationTool/AutomationUtils/Matchers/SanitizerEventMatcher.cs
@@ -101,26 +101,38 @@ namespace AutomationUtils.Matchers
 			return LogLevel.Error;
 		}
 
-		private static void AddProperty(LogValue Value, Group Group, Dictionary<string, object> Properties)
+		/// <summary>
+		/// Add Value to Properties using (Regex) Group match and update Message with the corresponding key instead of value from Group.
+		/// </summary>
+		/// <param name="Value"></param>
+		/// <param name="Group"></param>
+		/// <param name="Properties"></param>
+		/// <param name="Message"></param>
+		/// <param name="Offset">Optional argument to offset Group index from Message if Message changed since Group was initialy matched</param>
+		private static void AddPropertyAndReplace(LogValue Value, Group Group, Dictionary<string, object> Properties, ref string Message, int Offset = 0)
 		{
+			string Key = Value.Type.ToString();
 			Value.Text = Group.Value;
-			Properties.Add(Value.Type.ToString(), Value);
+			Properties.Add(Key, Value);
+			int Index = Group.Index + Offset;
+			Message = $"{Message.Substring(0, Index)}{{{Key}}}{Message.Substring(Index + Group.Length)}";
 		}
 
 		/// <summary>
-		/// Add Sanitizer Summary information to input Key/Value Pair Properties
+		/// Add Sanitizer Summary information to input Key/Value Pair Properties and update Message with Key markers
 		/// </summary>
 		/// <param name="Message"></param>
 		/// <param name="Properties"></param>
 		/// <returns></returns>
-		public static bool AddSanitizerSummaryProperties(string Message, Dictionary<string, object> Properties)
+		public static bool AddSanitizerSummaryProperties(ref string Message, Dictionary<string, object> Properties)
 		{
 			Match Match = SummaryPattern.Match(Message);
 			if (Match.Success)
 			{
-				AddProperty(SanitizerIssueHandler.SanitizerName, Match.Groups["SanitizerName"], Properties);
-				AddProperty(SanitizerIssueHandler.SummaryReason, Match.Groups["SummaryReason"], Properties);
-				AddProperty(SanitizerIssueHandler.SummarySourceFile, Match.Groups["SourceFile"], Properties);
+				int InitialLength = Message.Length;
+				AddPropertyAndReplace(SanitizerIssueHandler.SanitizerName, Match.Groups["SanitizerName"], Properties, ref Message);
+				AddPropertyAndReplace(SanitizerIssueHandler.SummaryReason, Match.Groups["SummaryReason"], Properties, ref Message, Message.Length - InitialLength);
+				AddPropertyAndReplace(SanitizerIssueHandler.SummarySourceFile, Match.Groups["SourceFile"], Properties, ref Message, Message.Length - InitialLength);
 
 				return true;
 			}
diff --git a/Engine/Source/Programs/AutomationTool/Gauntlet/Unreal/Base/Gauntlet.UnrealTestNode.cs b/Engine/Source/Programs/AutomationTool/Gauntlet/Unreal/Base/Gauntlet.UnrealTestNode.cs
--- a/Engine/Source/Programs/AutomationTool/Gauntlet/Unreal/Base/Gauntlet.UnrealTestNode.cs
+++ b/Engine/Source/Programs/AutomationTool/Gauntlet/Unreal/Base/Gauntlet.UnrealTestNode.cs
@@ -2025,17 +2025,20 @@ namespace Gauntlet
 					Callstack.Add("Could not parse callstack. See log for full callstack");
 				}
 				string CallstackString = string.Join("\n", Callstack.Select(C => "    " + C));
-				Dictionary<string, object> Properties = new Dictionary<string, object>() {
-					{ "Summary", Event.Summary },
-					{ "Callstack", CallstackString }
-				};
+				Dictionary<string, object> Properties = new Dictionary<string, object>() { { "Summary", Event.Summary } };
+				string Message = " * Fatal Error: {Summary}\n{Callstack}";
 				EventId EventType = KnownLogEvents.Gauntlet_FatalEvent;
 				if (Event.IsSanReport && CommandUtils.IsBuildMachine
-					&& SanitizerEventMatcher.AddSanitizerSummaryProperties(CallstackString, Properties))
+					&& SanitizerEventMatcher.AddSanitizerSummaryProperties(ref CallstackString, Properties))
 				{
+					Message = $" * Fatal Error: {{Summary}}\n{CallstackString}";
 					EventType = SanitizerEventMatcher.ConvertSanitizerNameToEventId(Properties.GetValueOrDefault("SanitizerName")?.ToString());
 				}
-				Log.Error(EventType, " * Fatal Error: {Summary}\n{Callstack}", Args: Properties.Values.ToArray());
+				else
+				{
+					Properties.Add("Callstack", CallstackString);
+				}
+				Log.Error(EventType, Message, Args: Properties.Values.ToArray());
 				Log.Info("");
 			}
 
diff --git a/Engine/Source/Programs/AutomationTool/Gauntlet/Unreal/Utils/Gauntlet.UnrealTestResults.cs b/Engine/Source/Programs/AutomationTool/Gauntlet/Unreal/Utils/Gauntlet.UnrealTestResults.cs
--- a/Engine/Source/Programs/AutomationTool/Gauntlet/Unreal/Utils/Gauntlet.UnrealTestResults.cs
+++ b/Engine/Source/Programs/AutomationTool/Gauntlet/Unreal/Utils/Gauntlet.UnrealTestResults.cs
@@ -201,24 +201,36 @@ namespace Gauntlet
 			Dictionary<string, object> Properties = null;
 			if (Event.IsCriticalFailure)
 			{
-				Properties = new Dictionary<string, object>() { { "Callstack", Message } };
+				Properties = new Dictionary<string, object>();
 				if (Event.IsSanReport
-					&& SanitizerEventMatcher.AddSanitizerSummaryProperties(Message, Properties))
+					&& SanitizerEventMatcher.AddSanitizerSummaryProperties(ref Message, Properties))
 				{
+					Format = Message;
 					EventIdType = SanitizerEventMatcher.ConvertSanitizerNameToEventId(Properties.GetValueOrDefault("SanitizerName")?.ToString());
 				}
-				Format = "{Callstack}";
-			}
-			if (!string.IsNullOrEmpty(Event.Context))
-			{
-				Properties.Add("Context", Event.Context);
-				Format = "[{Context}] " + Format?? Message;
+				else
+				{
+					Properties.Add("Callstack", Message);
+					Format = "{Callstack}";
+				}
 			}
-			if (!string.IsNullOrEmpty(Filename))
+			else
 			{
-				string FileReference = $"{Filename}:{LineNumber}";
-				Properties.Add("FileReference", FileReference);
-				Format = (Format?? Message) + " [{FileReference}]";
+				if (!string.IsNullOrEmpty(Event.Context))
+				{
+					Properties = new Dictionary<string, object>() { { "Context", Event.Context } };
+					Format = "[{Context}] " + (Format ?? Message);
+				}
+				if (!string.IsNullOrEmpty(Filename))
+				{
+					if (Properties == null)
+					{
+						Properties = new Dictionary<string, object>();
+					}
+					Properties.Add("SourceFile", Filename);
+					Properties.Add("Line", LineNumber.ToString());
+					Format = (Format ?? Message) + " [{SourceFile}:{Line}]";
+				}
 			}
 
 			return new LogEvent(Time, Level, EventIdType, Message, Format, Properties, null);
