Add "SkipExport" and "SnapshotDescriptorPlatform" parameters to ZenExportSnapshotTask so that it can be used to write snapshot descriptors for data that was already exported and should be aliased to another platform.

diff --git a/Engine/Source/Programs/AutomationTool/BuildGraph/Tasks/ZenExportSnapshotTask.cs b/Engine/Source/Programs/AutomationTool/BuildGraph/Tasks/ZenExportSnapshotTask.cs
--- a/Engine/Source/Programs/AutomationTool/BuildGraph/Tasks/ZenExportSnapshotTask.cs
+++ b/Engine/Source/Programs/AutomationTool/BuildGraph/Tasks/ZenExportSnapshotTask.cs
@@ -94,6 +94,12 @@ namespace AutomationTool.Tasks
 		[TaskParameter(Optional = true)]
 		public string SnapshotDescriptorCloudHost { get; set; }
 
+		/// <summary>
+		/// The target platform to use when writing a snapshot descriptor
+		/// </summary>
+		[TaskParameter(Optional = true)]
+		public string SnapshotDescriptorPlatform { get; set; }
+
 		/// <summary>
 		/// The http version to use when exporting to a cloud destination
 		/// </summary>
@@ -147,7 +153,13 @@ namespace AutomationTool.Tasks
 		/// Optional. Whether to force export of data even if the destination claims to have them.
 		/// </summary>
 		[TaskParameter(Optional = true)]
-		public bool Force { get; set; } = false;
+		public bool ForceExport { get; set; } = false;
+
+		/// <summary>
+		/// Optional. Whether to entirely bypass the exporting of data and write a snapshot descriptor as if the data had been exported.
+		/// </summary>
+		[TaskParameter(Optional = true)]
+		public bool SkipExport { get; set; } = false;
 	}
 
 	/// <summary>
@@ -312,6 +324,11 @@ namespace AutomationTool.Tasks
 
 		private void WriteExportSource(JsonWriter writer, SnapshotStorageType destinationStorageType, ExportSourceData exportSource, string name)
 		{
+			string targetPlatform = _parameters.SnapshotDescriptorPlatform;
+			if (string.IsNullOrEmpty(targetPlatform))
+			{
+				targetPlatform = exportSource._targetPlatform;
+			}
 			writer.WriteObjectStart();
 			switch (destinationStorageType)
 			{
@@ -339,7 +356,7 @@ namespace AutomationTool.Tasks
 					IoHash destinationKeyHash = IoHash.Compute(Encoding.UTF8.GetBytes(name));
 					writer.WriteValue("name", name);
 					writer.WriteValue("type", "cloud");
-					writer.WriteValue("targetplatform", exportSource._targetPlatform);
+					writer.WriteValue("targetplatform", targetPlatform);
 					writer.WriteValue("host", hostName);
 					if (!string.IsNullOrEmpty(httpVersion) && !httpVersion.Equals("None", StringComparison.OrdinalIgnoreCase))
 					{
@@ -354,7 +371,7 @@ namespace AutomationTool.Tasks
 
 					writer.WriteValue("name", name);
 					writer.WriteValue("type", "zen");
-					writer.WriteValue("targetplatform", exportSource._targetPlatform);
+					writer.WriteValue("targetplatform", targetPlatform);
 					writer.WriteValue("host", _parameters.DestinationZenHost);
 					writer.WriteValue("projectid", projectName);
 					writer.WriteValue("oplogid", SanitizeOplogName(name));
@@ -362,7 +379,7 @@ namespace AutomationTool.Tasks
 				case SnapshotStorageType.File:
 					writer.WriteValue("name", name);
 					writer.WriteValue("type", "file");
-					writer.WriteValue("targetplatform", exportSource._targetPlatform);
+					writer.WriteValue("targetplatform", targetPlatform);
 					writer.WriteValue("directory", _parameters.DestinationFileDir.FullName);
 					writer.WriteValue("filename", _parameters.DestinationFileName);
 					break;
@@ -496,7 +513,7 @@ namespace AutomationTool.Tasks
 			// Format the command line
 			StringBuilder oplogExportCommandline = new StringBuilder();
 			oplogExportCommandline.Append("oplog-export --embedloosefiles");
-			if (_parameters.Force)
+			if (_parameters.ForceExport)
 			{
 				oplogExportCommandline.Append(" --force");
 			}
@@ -566,7 +583,7 @@ namespace AutomationTool.Tasks
 						IoHash destinationKeyHash = IoHash.Compute(Encoding.UTF8.GetBytes(exportNames[exportIndex]));
 
 						exportSingleSourceCommandline.AppendFormat(" {0} --key {1} {2} {3} {4}", hostUrlArg, destinationKeyHash.ToString().ToLowerInvariant(), baseKeyArg, exportSource._projectId, exportSource._oplogId);
-						if (TryExportOplogCommand(zenExe.FullName, exportSingleSourceCommandline.ToString()))
+						if (_parameters.SkipExport || TryExportOplogCommand(zenExe.FullName, exportSingleSourceCommandline.ToString()))
 						{
 							successfullyExportedSources.Add(exportSource);
 						}
@@ -607,7 +624,7 @@ namespace AutomationTool.Tasks
 						string destinationOplog = SanitizeOplogName(exportNames[exportIndex]);
 
 						exportSingleSourceCommandline.AppendFormat(" {0} --target-project {1} --target-oplog {2} {3} {4}", hostUrlArg, projectName, destinationOplog, exportSource._projectId, exportSource._oplogId);
-						if (TryExportOplogCommand(zenExe.FullName, exportSingleSourceCommandline.ToString()))
+						if (_parameters.SkipExport || TryExportOplogCommand(zenExe.FullName, exportSingleSourceCommandline.ToString()))
 						{
 							successfullyExportedSources.Add(exportSource);
 						}
@@ -658,7 +675,7 @@ namespace AutomationTool.Tasks
 						}
 						exportSingleSourceCommandline.AppendFormat(" --file {0} --name {1} {2} {3} {4}", CommandUtils.MakePathSafeToUseWithCommandLine(platformDestinationFileDir.FullName), destinationFileName, baseNameArg, projectId, exportSource._oplogId);
 
-						if (TryExportOplogCommand(zenExe.FullName, exportSingleSourceCommandline.ToString()))
+						if (_parameters.SkipExport || TryExportOplogCommand(zenExe.FullName, exportSingleSourceCommandline.ToString()))
 						{
 							successfullyExportedSources.Add(exportSource);
 						}
