PR #11889: BuildGraph: Create LogPath directory before writing to directory. (prevents System.IO.DirectoryNotFoundException) (2 more instances that went unnoticed previously)

#rnx

diff --git a/Engine/Source/Programs/AutomationTool/BuildGraph/TempStorage.cs b/Engine/Source/Programs/AutomationTool/BuildGraph/TempStorage.cs
--- a/Engine/Source/Programs/AutomationTool/BuildGraph/TempStorage.cs
+++ b/Engine/Source/Programs/AutomationTool/BuildGraph/TempStorage.cs
@@ -620,7 +620,7 @@ namespace AutomationTool
 			{
 				return true;
 			}
-
+			
 			// Check if it exists in shared storage
 			if (_sharedDir != null)
 			{
@@ -953,8 +953,9 @@ namespace AutomationTool
 					string result = ParallelUnzipFiles(zipFiles, _rootDir);
 					if (!string.IsNullOrWhiteSpace(result))
 					{
-						string logPath = CommandUtils.CombinePaths(_rootDir.FullName, "Engine/Programs/AutomationTool/Saved/Logs", $"Copy Manifest - {nodeName}.log");
+						string logPath = CommandUtils.CombinePaths(CmdEnv.LogFolder, $"Copy Manifest - {nodeName}.log");
 						Logger.LogInformation("Saving copy log to {LogPath}", logPath);
+						Directory.CreateDirectory(Path.GetDirectoryName(logPath)!);
 						File.WriteAllText(logPath, result);
 					}
 
@@ -1111,7 +1112,7 @@ namespace AutomationTool
 				try
 				{
 					string copyResult = CopyDirectory(zipDir, outputDir);
-					string logPath = CommandUtils.CombinePaths(rootDir.FullName, "Engine/Programs/AutomationTool/Saved/Logs", $"Copy Files to Temp Storage.log");
+					string logPath = CommandUtils.CombinePaths(CmdEnv.LogFolder, $"Copy Files to Temp Storage.log");
 					Logger.LogInformation("Saving copy log to {LogPath}", logPath);
 					Directory.CreateDirectory(Path.GetDirectoryName(logPath)!);
 					File.WriteAllText(logPath, string.Join(Environment.NewLine, copyResult));
@@ -1215,8 +1216,9 @@ namespace AutomationTool
 							if (unzipFileAttempts == 0)
 							{
 								Log.Logger.LogError(ex, "All retries exhausted attempting to unzip entries from '{LocalZipFile}'. Terminating.", localZipFile);
-								string logPath = CommandUtils.CombinePaths(rootDir.FullName, "Engine/Programs/AutomationTool/Saved/Logs", $"Copy Manifest - {zipFile.Name}.log");
+								string logPath = CommandUtils.CombinePaths(CmdEnv.LogFolder, $"Copy Manifest - {zipFile.Name}.log");
 								Logger.LogInformation("Saving copy log to {LogPath}", logPath);
+								Directory.CreateDirectory(Path.GetDirectoryName(logPath)!);
 								File.WriteAllText(logPath, string.Join(Environment.NewLine, copyResults));
 								throw;
 							}
@@ -1345,7 +1347,7 @@ namespace AutomationTool
 			}
 			return false;
 		}
-
+		
 		/// <summary>
 		/// Copy a temp storage .zip file to directory.
 		/// Uses Robocopy on Windows, rsync on Linux/macOS and native .NET API when under Wine.
@@ -1410,7 +1412,7 @@ namespace AutomationTool
 				return CommandUtils.RunAndLog("rsync", $"-vam --include=\"**\" \"{sourceDir}/\" \"{destinationDir}/\"", Options: CommandUtils.ERunOptions.AppMustExist | CommandUtils.ERunOptions.NoLoggingOfRunCommand);
 			}
 		}
-
+		
 		/// <summary>
 		/// Copy a directory using .NET SDK directory and file copy API
 		/// </summary>
@@ -1503,7 +1505,7 @@ namespace AutomationTool
 
 			Dictionary<FileReference, DateTime> namedOutput = SelectFiles(workingDir, 'g', 'i');
 			tempStore.Archive("TestNode", "NamedOutput", namedOutput.Keys.ToArray(), true);
-
+			
 			// Check both outputs are still ok
 			TempStorageManifest defaultManifest = tempStore.Retrieve("TestNode", null, new FileFilter());
 			CheckManifest(workingDir, defaultManifest, defaultOutput);
