PR #11848: Changes to improve BuildGraph in C#

diff --git a/Engine/Source/Programs/AutomationTool/BuildGraph/BuildGraph.cs b/Engine/Source/Programs/AutomationTool/BuildGraph/BuildGraph.cs
--- a/Engine/Source/Programs/AutomationTool/BuildGraph/BuildGraph.cs
+++ b/Engine/Source/Programs/AutomationTool/BuildGraph/BuildGraph.cs
@@ -455,7 +455,7 @@ namespace AutomationTool
 				(byte[] data, BgThunkDef[] methods) = BgCompiler.Compile(graphSpec);
 
 				BgInterpreter interpreter = new BgInterpreter(data, methods, arguments);
-				interpreter.Disassemble(Logger);
+				// interpreter.Disassemble(Logger);
 				graph = ((BgObjectDef)interpreter.Evaluate()).Deserialize<BgGraphExpressionDef>().ToGraphDef();
 			}
 			else
diff --git a/Engine/Source/Programs/AutomationTool/BuildGraph/Tasks/CommandTask.cs b/Engine/Source/Programs/AutomationTool/BuildGraph/Tasks/CommandTask.cs
--- a/Engine/Source/Programs/AutomationTool/BuildGraph/Tasks/CommandTask.cs
+++ b/Engine/Source/Programs/AutomationTool/BuildGraph/Tasks/CommandTask.cs
@@ -191,5 +191,21 @@ namespace AutomationTool.Tasks
 
 			await ExecuteAsync(new CommandTask(parameters));
 		}
+
+		/// <summary>
+		/// Runs another UAT command
+		/// </summary>
+		/// <param name="name">Name of the command to run</param>
+		/// <param name="arguments">Arguments for the command</param>
+		/// <param name="mergeTelemetryWithPrefix">If non-null, instructs telemetry from the command to be merged into the telemetry for this UAT instance with the given prefix. May be an empty (non-null) string.</param>
+		public static async Task CommandAsync(string name, string arguments = null, string mergeTelemetryWithPrefix = null)
+		{
+			CommandTaskParameters parameters = new CommandTaskParameters();
+			parameters.Name = name;
+			parameters.Arguments = arguments ?? parameters.Arguments;
+			parameters.MergeTelemetryWithPrefix = mergeTelemetryWithPrefix ?? parameters.MergeTelemetryWithPrefix;
+
+			await ExecuteAsync(new CommandTask(parameters));
+		}
 	}
 }
diff --git a/Engine/Source/Programs/Shared/EpicGames.BuildGraph/Expressions/BgNode.cs b/Engine/Source/Programs/Shared/EpicGames.BuildGraph/Expressions/BgNode.cs
--- a/Engine/Source/Programs/Shared/EpicGames.BuildGraph/Expressions/BgNode.cs
+++ b/Engine/Source/Programs/Shared/EpicGames.BuildGraph/Expressions/BgNode.cs
@@ -490,5 +490,16 @@ namespace EpicGames.BuildGraph.Expressions
 		{
 			return (T)node.Modify(fences: node.Fences.Add(inputs));
 		}
+		
+		/// <summary>
+		/// Add a label to the node.
+		/// </summary>
+		/// <param name="node">The node to modify</param>
+		/// <param name="label">The label to add to the node.</param>
+		/// <returns>The current node spec, to allow chaining calls</returns>
+		public static T AddLabel<T>(this T node, BgLabel label) where T : BgNode
+		{
+			return (T)node.Modify(labels: node.Labels.Add(label));
+		}
 	}
 }
