[Gauntlet] When running a single UnrealTestNode, the Horde Summary will now include the full artifact path to the specific test

diff --git a/Engine/Source/Programs/AutomationTool/Gauntlet/Framework/Horde/Gauntlet.Horde.cs b/Engine/Source/Programs/AutomationTool/Gauntlet/Framework/Horde/Gauntlet.Horde.cs
--- a/Engine/Source/Programs/AutomationTool/Gauntlet/Framework/Horde/Gauntlet.Horde.cs
+++ b/Engine/Source/Programs/AutomationTool/Gauntlet/Framework/Horde/Gauntlet.Horde.cs
@@ -34,7 +34,7 @@ namespace Gauntlet
 			}
 		}
 
-		public static void GenerateSummary()
+		public static void GenerateSummary(string ArtifactPath)
 		{
 			try
 			{
@@ -49,7 +49,7 @@ namespace Gauntlet
 
 				// file:// is attractive here, and file://\\computer\share\folder\file.txt does work ok in Chrome, but it does not Firefox
 				// so, lacking a uniform way to link to a network filesystem, just put the path here as static text and escaped for markdown
-				string Markdown = System.Text.RegularExpressions.Regex.Replace($"Gauntlet Artifacts: {Globals.LogDir}", @"([|\\*\+])", @"\$1");
+				string Markdown = System.Text.RegularExpressions.Regex.Replace($"Gauntlet Artifacts: {ArtifactPath}", @"([|\\*\+])", @"\$1");
 
 				File.WriteAllText(Path.Combine(LogFolder, MarkdownFilename), Markdown);
 
diff --git a/Engine/Source/Programs/AutomationTool/Gauntlet/Unreal/Base/Gauntlet.UnrealTestNode.cs b/Engine/Source/Programs/AutomationTool/Gauntlet/Unreal/Base/Gauntlet.UnrealTestNode.cs
--- a/Engine/Source/Programs/AutomationTool/Gauntlet/Unreal/Base/Gauntlet.UnrealTestNode.cs
+++ b/Engine/Source/Programs/AutomationTool/Gauntlet/Unreal/Base/Gauntlet.UnrealTestNode.cs
@@ -166,6 +166,11 @@ namespace Gauntlet
 		/// </summary>
 		public UnrealSessionInstance TestInstance { get; private set; }
 
+		/// <summary>
+		/// Path to the directory that logs and other artifacts are copied to after the test run.
+		/// </summary>
+		public string ArtifactPath { get; private set; }
+
 		/// <summary>
 		/// Describes the post-test results for a role.
 		/// </summary>
@@ -321,11 +326,6 @@ namespace Gauntlet
 
 		protected Version TestVersion;
 
-		/// <summary>
-		/// Path to the directory that logs and other artifacts are copied to after the test run.
-		/// </summary>
-		protected string ArtifactPath { get; private set; }
-
 		/// <summary>
 		/// Our test result. May be set directly, or by overriding GetUnrealTestResult()
 		/// </summary>
diff --git a/Engine/Source/Programs/AutomationTool/Gauntlet/Unreal/RunUnreal.cs b/Engine/Source/Programs/AutomationTool/Gauntlet/Unreal/RunUnreal.cs
--- a/Engine/Source/Programs/AutomationTool/Gauntlet/Unreal/RunUnreal.cs
+++ b/Engine/Source/Programs/AutomationTool/Gauntlet/Unreal/RunUnreal.cs
@@ -16,15 +16,13 @@ namespace Gauntlet
 {
 	/*
 	 	public class RunUnreal<SourceType, OptionType> : BuildCommand
-		where SourceType : UnrealBuildSource, 
+		where SourceType : UnrealBuildSource,
 		where OptionType : UnrealTestOptions
 	 */
 
 	/// <summary>
 	/// Base class for executing Unreal tests.
-	/// 
 	/// For a full list of options see UnrealTestContextOption
-	/// 
 	/// </summary>
 	[Help("Run Unreal tests using Gauntlet")]
 	[ParamHelp("Tests", "List of gauntlet tests to run", Required = true, MultiSelectSeparator = ",")]
@@ -311,8 +309,28 @@ namespace Gauntlet
 			// dispose now, not during shutdown gc, because this runs commands...
 			DevicePool.Instance.Dispose();
 
+			// If only running a single unreal test node, we'll direct the summary artifact link to the specific test
+			string ArtifactPath = Globals.LogDir;
+			if(AllTestNodes.Count == 1)
+			{
+				try
+				{
+					// Can't read UnrealTestNode artifact path without some ugly reflection, attempt dynamic property read instead
+					dynamic UnrealTestNode = AllTestNodes.First();
+					string UnrealArtifactPath = UnrealTestNode.ArtifactPath;
+					if (!string.IsNullOrEmpty(UnrealArtifactPath))
+					{
+						ArtifactPath = UnrealArtifactPath;
+					}
+				}
+				catch
+				{
+					// Type isn't UnrealTestNode, use default behavior
+				}
+			}
+
 			// Generate Horde summary for CIS test (maybe want to use a delegate here)
-			Horde.GenerateSummary();
+			Horde.GenerateSummary(ArtifactPath);
 
 			return AllTestsPassed ? ExitCode.Success : ExitCode.Error_TestFailure;
 		}
