Add  single quotes around SSHPath in order to support spaces
#11549

diff --git a/Engine/Source/Programs/AutomationTool/SteamDeck/SteamDeckSupport.cs b/Engine/Source/Programs/AutomationTool/SteamDeck/SteamDeckSupport.cs
--- a/Engine/Source/Programs/AutomationTool/SteamDeck/SteamDeckSupport.cs
+++ b/Engine/Source/Programs/AutomationTool/SteamDeck/SteamDeckSupport.cs
@@ -201,7 +201,7 @@ public static class SteamDeckSupport
 		// make a standard set of options to pass to the rsync for auth's -e option
 		string[] AuthOpts =
 		{
-			SSHPath,
+			$"'{SSHPath}'",
 			$"-o UserKnownHostsFile='{KnownHostsPath}'",
 			$"-o StrictHostKeyChecking=no",
 			$"-i '{DevKitRSAPath}'",
