Adding single bracket for SSHPAth only when a path contains a white space, since it was causing issues for regular paths

diff --git a/Engine/Source/Programs/AutomationTool/SteamDeck/SteamDeckSupport.cs b/Engine/Source/Programs/AutomationTool/SteamDeck/SteamDeckSupport.cs
--- a/Engine/Source/Programs/AutomationTool/SteamDeck/SteamDeckSupport.cs
+++ b/Engine/Source/Programs/AutomationTool/SteamDeck/SteamDeckSupport.cs
@@ -201,11 +201,15 @@ public static class SteamDeckSupport
 		// make a standard set of options to pass to the rsync for auth's -e option
 		string[] AuthOpts =
 		{
-			$"'{SSHPath}'",
+			SSHPath,
 			$"-o UserKnownHostsFile='{KnownHostsPath}'",
 			$"-o StrictHostKeyChecking=no",
 			$"-i '{DevKitRSAPath}'",
 		};
+
+		if (SSHPath.Any(Char.IsWhiteSpace))
+			AuthOpts[0] = $"'{SSHPath}'";
+
 		string AuthOptions = $"-e \"{string.Join(" ", AuthOpts)}\"";
 
 		// make a set of --exclude options for anything we want to exclude
