BuildGraph: Create the output directory for generating random data chunks if it doesn't exist.

#rnx

diff --git a/Engine/Source/Programs/AutomationTool/BuildGraph/Tasks/RandomDataTask.cs b/Engine/Source/Programs/AutomationTool/BuildGraph/Tasks/RandomDataTask.cs
--- a/Engine/Source/Programs/AutomationTool/BuildGraph/Tasks/RandomDataTask.cs
+++ b/Engine/Source/Programs/AutomationTool/BuildGraph/Tasks/RandomDataTask.cs
@@ -72,6 +72,7 @@ namespace AutomationTool.Tasks
 		public override async Task ExecuteAsync(JobContext job, HashSet<FileReference> buildProducts, Dictionary<string, HashSet<FileReference>> tagNameToFileSet)
 		{
 			DirectoryReference outputDir = ResolveDirectory(_parameters.OutputDir);
+			DirectoryReference.CreateDirectory(outputDir);
 
 			byte[] buffer = Array.Empty<byte>();
 			for (int idx = 0; idx < _parameters.Count; idx++)
