UAT scripts: Assembly.SafeGetLoadedTypes()

#rnx

diff --git a/Engine/Source/Programs/AutomationTool/Gauntlet/Framework/Gauntlet.Utils.cs b/Engine/Source/Programs/AutomationTool/Gauntlet/Framework/Gauntlet.Utils.cs
--- a/Engine/Source/Programs/AutomationTool/Gauntlet/Framework/Gauntlet.Utils.cs
+++ b/Engine/Source/Programs/AutomationTool/Gauntlet/Framework/Gauntlet.Utils.cs
@@ -18,6 +18,7 @@ using Microsoft.Extensions.Logging;
 using Logging = Microsoft.Extensions.Logging;
 
 using static AutomationTool.CommandUtils;
+using UnrealBuildBase;
 
 namespace Gauntlet
 {
@@ -662,7 +663,7 @@ namespace Gauntlet
 			public static IEnumerable<Type> GetTypesInNamespaces<BaseType>(IEnumerable<string> Namespaces)
 				where BaseType : class
 			{
-				var AllTypes = AppDomain.CurrentDomain.GetAssemblies().SelectMany(S => S.GetTypes()).Where(T => typeof(BaseType).IsAssignableFrom(T));
+				var AllTypes = AppDomain.CurrentDomain.GetAssemblies().SelectMany(S => S.SafeGetLoadedTypes()).Where(T => typeof(BaseType).IsAssignableFrom(T));
 
 				if (Namespaces.Count() > 0)
 				{
diff --git a/Engine/Source/Programs/AutomationTool/OneSkyLocalization/OneSkyConfigHelper.cs b/Engine/Source/Programs/AutomationTool/OneSkyLocalization/OneSkyConfigHelper.cs
--- a/Engine/Source/Programs/AutomationTool/OneSkyLocalization/OneSkyConfigHelper.cs
+++ b/Engine/Source/Programs/AutomationTool/OneSkyLocalization/OneSkyConfigHelper.cs
@@ -8,6 +8,7 @@ using UnrealBuildTool;
 using Microsoft.Extensions.Logging;
 
 using static AutomationTool.CommandUtils;
+using UnrealBuildBase;
 
 namespace EpicGames.OneSkyLocalization.Config
 {
@@ -29,7 +30,7 @@ namespace EpicGames.OneSkyLocalization.Config
 				Assembly[] LoadedAssemblies = AppDomain.CurrentDomain.GetAssemblies();
 				foreach (var Dll in LoadedAssemblies)
 				{
-					Type[] AllTypes = Dll.GetTypes();
+					Type[] AllTypes = Dll.SafeGetLoadedTypes();
 					foreach (var PotentialConfigType in AllTypes)
 					{
 						if (PotentialConfigType != typeof(OneSkyConfigData) && typeof(OneSkyConfigData).IsAssignableFrom(PotentialConfigType))
