UnrealPak: When the existing optional .uproject argument is provided first on the commandline, then let UnrealPak use that to set its project directory and override the monolithic engine program directory (Engine/Programs/UnrealPak).
This will make UnrealPak read from and use the project-specific ini files.
This will make UnrealPak write to and use the project-specific Saved directory, e.g. to save the UnrealPak.log file.
In the BuildCookRun/CopyBuildToStagingDirectory scripts always invoke UnrealPak with the .uproject file argument.

Motivation: When UnrealPak is run with -compressionEnableDDC it needs to use the project's ddc configuration instead of the engine's ddc configuration.

Note: When UnrealPak is invoked without the project argument it will still use the default program directory, and there is no change in behavior when e.g. manually executing inspection commands like list, extract and describe.
Note: All important packaging options are already passed on command line from BuildCookRun/CopyBuildToStagingDirectory, so from that perspective the change should be low risk, but in the future there is a possibility to reduce the number of command line options and rely on the project configuration instead.

#tests identical binary output

diff --git a/Engine/Source/Developer/IoStoreUtilities/Private/IoStoreUtilities.cpp b/Engine/Source/Developer/IoStoreUtilities/Private/IoStoreUtilities.cpp
--- a/Engine/Source/Developer/IoStoreUtilities/Private/IoStoreUtilities.cpp
+++ b/Engine/Source/Developer/IoStoreUtilities/Private/IoStoreUtilities.cpp
@@ -9503,6 +9503,12 @@ bool ParseContainerGenerationArguments(FIoStoreArguments& Arguments, FIoStoreWri
 	}
 
 	WriterSettings.bCompressionEnableDDC = FParse::Param(FCommandLine::Get(), TEXT("compressionEnableDDC"));
+	if (WriterSettings.bCompressionEnableDDC && !FPaths::IsProjectFilePathSet())
+	{
+		UE_LOG(LogIoStore, Warning,
+			TEXT("Ignoring -compressionEnableDDC due to missing .uproject file as the first unnamed argument."));
+		WriterSettings.bCompressionEnableDDC = false;
+	}
 
 	int32 CompressionMinSizeToConsiderDDC = 0;
 	if (FParse::Value(FCommandLine::Get(), TEXT("-compressionMinSizeToConsiderDDC="), CompressionMinSizeToConsiderDDC))
diff --git a/Engine/Source/Developer/PakFileUtilities/Private/PakFileUtilities.cpp b/Engine/Source/Developer/PakFileUtilities/Private/PakFileUtilities.cpp
--- a/Engine/Source/Developer/PakFileUtilities/Private/PakFileUtilities.cpp
+++ b/Engine/Source/Developer/PakFileUtilities/Private/PakFileUtilities.cpp
@@ -5447,6 +5447,7 @@ bool MakeBinaryConfig(const TCHAR* CmdLine)
  */
 bool ExecuteUnrealPak(const TCHAR* CmdLine)
 {
+	UE_LOG(LogPakFile, Display, TEXT("ProjectDir: %s"), *FPaths::ProjectDir());
 	{
 		FString IoStoreArg;
 		if (FParse::Value(CmdLine, TEXT("-CreateGlobalContainer="), IoStoreArg) ||
diff --git a/Engine/Source/Programs/AutomationTool/Scripts/CopyBuildToStagingDirectory.Automation.cs b/Engine/Source/Programs/AutomationTool/Scripts/CopyBuildToStagingDirectory.Automation.cs
--- a/Engine/Source/Programs/AutomationTool/Scripts/CopyBuildToStagingDirectory.Automation.cs
+++ b/Engine/Source/Programs/AutomationTool/Scripts/CopyBuildToStagingDirectory.Automation.cs
@@ -324,12 +324,21 @@ namespace AutomationScripts
 			return CmdLine.ToString();
 		}
 
-		static public string GetUnrealPakArguments(FileReference ProjectPath, Dictionary<string, string> UnrealPakResponseFile, FileReference OutputLocation, List<OrderFile> PakOrderFileLocations, string AdditionalOptions, bool Compressed, EncryptionAndSigning.CryptoSettings CryptoSettings, FileReference CryptoKeysCacheFilename, String PatchSourceContentPath, string EncryptionKeyGuid, List<OrderFile> SecondaryPakOrderFileLocations, bool bUnattended)
+		static public string GetCreateSinglePakArguments(
+				Dictionary<string, string> UnrealPakResponseFile,
+				FileReference OutputLocation,
+				List<OrderFile> PakOrderFileLocations,
+				string AdditionalOptions,
+				bool Compressed,
+				EncryptionAndSigning.CryptoSettings CryptoSettings,
+				FileReference CryptoKeysCacheFilename,
+				String PatchSourceContentPath,
+				string EncryptionKeyGuid, List<OrderFile> SecondaryPakOrderFileLocations,
+				bool bUnattended)
 		{
 			bool RehydrateAssets = false;
 
-			StringBuilder CmdLine = new StringBuilder(MakePathSafeToUseWithCommandLine(ProjectPath.FullName));
-			CmdLine.Append(" ");
+			StringBuilder CmdLine = new StringBuilder();
 			CmdLine.Append(GetPakFileSpecificUnrealPakArguments(UnrealPakResponseFile, OutputLocation, "", Compressed, RehydrateAssets, CryptoSettings, PatchSourceContentPath, EncryptionKeyGuid));
 			CmdLine.Append(" ");
 			CmdLine.Append(GetCommonUnrealPakArguments(PakOrderFileLocations, AdditionalOptions, CryptoSettings, CryptoKeysCacheFilename, SecondaryPakOrderFileLocations, bUnattended));
@@ -418,15 +427,18 @@ namespace AutomationScripts
 			}
 		}
 
-		static public void RunUnrealPak(string Arguments, string Scenario)
+		static public void RunUnrealPak(FileReference ProjectPath, string Arguments, string Scenario)
 		{
+			// Invoke UnrealPak with the project argument first on the commandline to make it use project-specific configs.
+			// This will also use use the project specific Saved directory instead of the program Saved directory.
+			Arguments = MakePathSafeToUseWithCommandLine(ProjectPath.FullName) + " " + Arguments;
 			Logger.LogInformation("Running UnrealPak with arguments: {Arguments}", Arguments);
 			RunAndLog(CmdEnv, GetUnrealPakLocation().FullName, Arguments, Options: ERunOptions.Default | ERunOptions.UTF8Output);
 
 			// try to copy the default UnrealPak.log file to the UAT log folder,
 			// this log file may contain more verbose logging than the one generated by RunAndLog() based on stdout/stderr.
 			string UnrealPakLogFileName = "UnrealPak_" + Scenario;
-			string SourceLogFile = CombinePaths(DirectoryReference.Combine(Unreal.RootDirectory, "Engine", "Programs", "UnrealPak", "Saved", "Logs").FullName, "UnrealPak.log");
+			string SourceLogFile = CombinePaths(DirectoryReference.Combine(ProjectPath.Directory, "Saved", "Logs").FullName, "UnrealPak.log");
 			string DestLogFile = LogUtils.GetUniqueLogName(CombinePaths(CommandUtils.CmdEnv.LogFolder, UnrealPakLogFileName));
 			if (InternalUtils.SafeFileExists(SourceLogFile, true))
 			{
@@ -897,7 +909,7 @@ namespace AutomationScripts
 									Arguments += " -unattended";
 								}
 
-								RunUnrealPak(Arguments, "ZenAutoLaunch");
+								RunUnrealPak(Params.RawProjectPath, Arguments, "ZenAutoLaunch");
 								bAttemptAutoLaunchOnFailure = false;
 							}
 							else
@@ -1630,7 +1642,7 @@ namespace AutomationScripts
 						FileReference OutputFile = FileReference.Combine(SC.ProjectRoot, "Intermediate", "Config", TargetPlatformName, "BinaryConfig.ini");
 						String UnrealPakParams = String.Format("MakeBinaryConfig -Project=\"{0}\" -Platform={1} -OutputFile=\"{2}\" -StagedPluginsFile=\"{3}\"", SC.RawProjectPath, ConfigHierarchy.GetIniPlatformName(ThisPlatform.IniPlatformType), OutputFile.FullName, PluginListFile.FullName);
 
-						RunUnrealPak(UnrealPakParams, "MakeBinaryConfig");
+						RunUnrealPak(Params.RawProjectPath, UnrealPakParams, "MakeBinaryConfig");
 
 						SC.StageFile(StagedFileType.UFS, OutputFile, StagedFileReference.Combine(SC.RelativeProjectRootForStage, "Config", OutputFile.GetFileName()));
 					}
@@ -2696,8 +2708,8 @@ namespace AutomationScripts
 			{
 				return;
 			}
-			string Arguments = GetUnrealPakArguments(Params.RawProjectPath, PakResponseFile, OutputLocation, null, null, Params.Compressed, null, null, null, null, null, Params.Unattended);
-			RunUnrealPak(Arguments, "CrashReportClient");
+			string Arguments = GetCreateSinglePakArguments(PakResponseFile, OutputLocation, null, null, Params.Compressed, null, null, null, null, null, Params.Unattended);
+			RunUnrealPak(Params.RawProjectPath, Arguments, "CrashReportClient");
 		}
 
 		/// <summary>
@@ -3667,7 +3679,7 @@ namespace AutomationScripts
 					}
 				}
 
-				StringBuilder Arguments = new StringBuilder(MakePathSafeToUseWithCommandLine(Params.RawProjectPath.FullName));
+				StringBuilder Arguments = new StringBuilder();
 				string CommonArguments = GetCommonUnrealPakArguments(
 					PrimaryOrderFiles,
 					CommonAdditionalArgs,
@@ -3676,9 +3688,8 @@ namespace AutomationScripts
 					SecondaryOrderFiles,
 					Params.Unattended
 				);
-				Arguments.AppendFormat(" {0}", CommonArguments);
-				Arguments.AppendFormat(" -CreateMultiple={0}", CommandUtils.MakePathSafeToUseWithCommandLine(PakCommandsFileName));
-				RunUnrealPak(Arguments.ToString(), "CreateMultiplePaks");
+				Arguments.AppendFormat("{0} -CreateMultiple={1}", CommonArguments, CommandUtils.MakePathSafeToUseWithCommandLine(PakCommandsFileName));
+				RunUnrealPak(Params.RawProjectPath, Arguments.ToString(), "CreateMultiplePaks");
 			}
 
 			if (IoStoreCommands.Count > 0)
@@ -4169,7 +4180,7 @@ namespace AutomationScripts
 				CommandletParams += " -unattended";
 			}
 
-			RunUnrealPak(CommandletParams, "CreateIoStoreContainers");
+			RunUnrealPak(Params.RawProjectPath, CommandletParams, "CreateIoStoreContainers");
 		}
 
 		/// <summary>
diff --git a/Engine/Source/Programs/UnrealPak/Private/UnrealPak.cpp b/Engine/Source/Programs/UnrealPak/Private/UnrealPak.cpp
--- a/Engine/Source/Programs/UnrealPak/Private/UnrealPak.cpp
+++ b/Engine/Source/Programs/UnrealPak/Private/UnrealPak.cpp
@@ -4,6 +4,7 @@
 #include "RequiredProgramMainCPPInclude.h"
 #include "PakFileUtilities.h"
 #include "IPlatformFilePak.h"
+#include "ProjectUtilities.h"
 
 IMPLEMENT_APPLICATION(UnrealPak, "UnrealPak");
 
@@ -11,6 +12,9 @@ INT32_MAIN_INT32_ARGC_TCHAR_ARGV()
 {
 	FTaskTagScope Scope(ETaskTag::EGameThread);
 
+	// Allows this program to accept a project argument on the commandline and use project-specific config
+	UE::ProjectUtilities::ParseProjectDirFromCommandline(ArgC, ArgV);
+
 	// start up the main loop
 	GEngineLoop.PreInit(ArgC, ArgV);
 
