Calling msvsmoninstall.py before run the game on Steam Deck

diff --git a/Engine/Source/Programs/AutomationTool/SteamDeck/SteamDeckSupport.cs b/Engine/Source/Programs/AutomationTool/SteamDeck/SteamDeckSupport.cs
--- a/Engine/Source/Programs/AutomationTool/SteamDeck/SteamDeckSupport.cs
+++ b/Engine/Source/Programs/AutomationTool/SteamDeck/SteamDeckSupport.cs
@@ -10,6 +10,7 @@ using UnrealBuildBase;
 using Microsoft.Extensions.Logging;
 
 using static AutomationTool.CommandUtils;
+using IdentityModel.OidcClient;
 
 public static class SteamDeckSupport
 {
@@ -356,6 +357,18 @@ public static class SteamDeckSupport
 			return null;
 		}
 
+		if (RuntimePlatform == UnrealTargetPlatform.Win64)
+		{
+			// Run msvsmoninstall.py in order to copy remote debuggers files in the right folders
+			IProcessResult Result = SSHCommand("python3 ~/devkit-msvsmon/msvsmoninstall.py", UserName, IpAddr);
+
+			if (Result.ExitCode > 0)
+			{
+				Logger.LogWarning("Failed to run msvsmoninstall.py");
+				return Result;
+			}
+		}
+
 		string GameFolderPath = $"[REDACTED_USER_HOME]/devkit-game/{Params.ShortProjectName}_{RuntimePlatform}";
 		string RunGameScriptFileName = "RunGameHelper.sh";
 
