BuildGraph: Improve logging for copying files to/from temp storage.

#rnx

diff --git a/Engine/Source/Programs/AutomationTool/BuildGraph/BuildGraph.cs b/Engine/Source/Programs/AutomationTool/BuildGraph/BuildGraph.cs
--- a/Engine/Source/Programs/AutomationTool/BuildGraph/BuildGraph.cs
+++ b/Engine/Source/Programs/AutomationTool/BuildGraph/BuildGraph.cs
@@ -1078,22 +1078,22 @@ namespace AutomationTool
 			Dictionary<string, HashSet<FileReference>> tagNameToFileSet = new Dictionary<string, HashSet<FileReference>>();
 
 			// Read all the input tags for this node, and build a list of referenced input storage blocks
-			HashSet<TempStorageBlock> inputStorageBlocks = new HashSet<TempStorageBlock>();
+			HashSet<TempStorageBlockRef> inputStorageBlocks = new HashSet<TempStorageBlockRef>();
 			foreach (BgNodeOutput input in node.Inputs)
 			{
-				TempStorageTagManifest fileList = storage.ReadFileList(input.ProducingNode.Name, input.TagName) ?? throw new InvalidOperationException();
+				TempStorageTagManifest fileList = storage.ReadTagFileList(input.ProducingNode.Name, input.TagName) ?? throw new InvalidOperationException();
 				tagNameToFileSet[input.TagName] = fileList.ToFileSet(rootDir);
 				inputStorageBlocks.UnionWith(fileList.Blocks);
 			}
 
 			// Read the manifests for all the input storage blocks
-			Dictionary<TempStorageBlock, TempStorageManifest> inputManifests = new Dictionary<TempStorageBlock, TempStorageManifest>();
+			Dictionary<TempStorageBlockRef, TempStorageBlockManifest> inputManifests = new Dictionary<TempStorageBlockRef, TempStorageBlockManifest>();
 			using (IScope scope = GlobalTracer.Instance.BuildSpan("TempStorage").WithTag("resource", "read").StartActive())
 			{
 				scope.Span.SetTag("blocks", inputStorageBlocks.Count);
-				foreach (TempStorageBlock inputStorageBlock in inputStorageBlocks)
+				foreach (TempStorageBlockRef inputStorageBlock in inputStorageBlocks)
 				{
-					TempStorageManifest manifest = storage.Retrieve(inputStorageBlock.NodeName, inputStorageBlock.OutputName, ignoreModifiedFilter);
+					TempStorageBlockManifest manifest = storage.Retrieve(inputStorageBlock.NodeName, inputStorageBlock.OutputName, ignoreModifiedFilter);
 					inputManifests[inputStorageBlock] = manifest;
 				}
 				scope.Span.SetTag("size", inputManifests.Sum(x => x.Value.GetTotalSize()));
@@ -1101,7 +1101,7 @@ namespace AutomationTool
 
 			// Read all the input storage blocks, keeping track of which block each file came from
 			Dictionary<string, TempStorageFile> inputFiles = new Dictionary<string, TempStorageFile>(FileReference.Comparer);
-			foreach (KeyValuePair<TempStorageBlock, TempStorageManifest> pair in inputManifests)
+			foreach (KeyValuePair<TempStorageBlockRef, TempStorageBlockManifest> pair in inputManifests)
 			{
 				foreach (TempStorageFile newFile in pair.Value.Files)
 				{
@@ -1238,10 +1238,10 @@ namespace AutomationTool
 			// Write all the storage blocks, and update the mapping from file to storage block
 			using (GlobalTracer.Instance.BuildSpan("TempStorage").WithTag("resource", "Write").StartActive())
 			{
-				Dictionary<FileReference, TempStorageBlock> outputFileToStorageBlock = new Dictionary<FileReference, TempStorageBlock>();
+				Dictionary<FileReference, TempStorageBlockRef> outputFileToStorageBlock = new Dictionary<FileReference, TempStorageBlockRef>();
 				foreach (KeyValuePair<string, HashSet<FileReference>> pair in outputStorageBlockToFiles)
 				{
-					TempStorageBlock outputBlock = new TempStorageBlock(node.Name, pair.Key);
+					TempStorageBlockRef outputBlock = new TempStorageBlockRef(node.Name, pair.Key);
 					foreach (FileReference file in pair.Value)
 					{
 						outputFileToStorageBlock.Add(file, outputBlock);
@@ -1257,10 +1257,10 @@ namespace AutomationTool
 				{
 					HashSet<FileReference> files = tagNameToFileSet[output.TagName];
 
-					HashSet<TempStorageBlock> storageBlocks = new HashSet<TempStorageBlock>();
+					HashSet<TempStorageBlockRef> storageBlocks = new HashSet<TempStorageBlockRef>();
 					foreach (FileReference file in files)
 					{
-						TempStorageBlock? storageBlock;
+						TempStorageBlockRef? storageBlock;
 						if (outputFileToStorageBlock.TryGetValue(file, out storageBlock))
 						{
 							storageBlocks.Add(storageBlock);
diff --git a/Engine/Source/Programs/AutomationTool/BuildGraph/TempStorage.cs b/Engine/Source/Programs/AutomationTool/BuildGraph/TempStorage.cs
--- a/Engine/Source/Programs/AutomationTool/BuildGraph/TempStorage.cs
+++ b/Engine/Source/Programs/AutomationTool/BuildGraph/TempStorage.cs
@@ -27,7 +27,7 @@ namespace AutomationTool
 	/// <summary>
 	/// Stores the name of a temp storage block
 	/// </summary>
-	public class TempStorageBlock
+	public class TempStorageBlockRef
 	{
 		/// <summary>
 		/// Name of the node
@@ -44,7 +44,7 @@ namespace AutomationTool
 		/// <summary>
 		/// Default constructor, for XML serialization.
 		/// </summary>
-		private TempStorageBlock()
+		private TempStorageBlockRef()
 		{
 			NodeName = String.Empty;
 			OutputName = String.Empty;
@@ -55,7 +55,7 @@ namespace AutomationTool
 		/// </summary>
 		/// <param name="nodeName">Name of the node</param>
 		/// <param name="outputName">Name of the node's output</param>
-		public TempStorageBlock(string nodeName, string outputName)
+		public TempStorageBlockRef(string nodeName, string outputName)
 		{
 			NodeName = nodeName;
 			OutputName = outputName;
@@ -68,7 +68,7 @@ namespace AutomationTool
 		/// <returns>True if the blocks are equivalent</returns>
 		public override bool Equals(object? other)
 		{
-			TempStorageBlock? otherBlock = other as TempStorageBlock;
+			TempStorageBlockRef? otherBlock = other as TempStorageBlockRef;
 			return otherBlock != null && NodeName == otherBlock.NodeName && OutputName == otherBlock.OutputName;
 		}
 
@@ -333,7 +333,8 @@ namespace AutomationTool
 	/// <summary>
 	/// A manifest storing information about build products for a node's output
 	/// </summary>
-	public class TempStorageManifest
+	[XmlRoot(ElementName = "TempStorageManifest")]
+	public class TempStorageBlockManifest
 	{
 		/// <summary>
 		/// List of output files
@@ -352,12 +353,12 @@ namespace AutomationTool
 		/// <summary>
 		/// Construct a static Xml serializer to avoid throwing an exception searching for the reflection info at runtime
 		/// </summary>
-		static readonly XmlSerializer s_serializer = XmlSerializer.FromTypes(new Type[] { typeof(TempStorageManifest) })[0]!;
+		static readonly XmlSerializer s_serializer = XmlSerializer.FromTypes(new Type[] { typeof(TempStorageBlockManifest) })[0]!;
 
 		/// <summary>
 		/// Construct an empty temp storage manifest
 		/// </summary>
-		private TempStorageManifest()
+		private TempStorageBlockManifest()
 		{
 			Files = Array.Empty<TempStorageFile>();
 			ZipFiles = Array.Empty<TempStorageZipFile>();
@@ -368,7 +369,7 @@ namespace AutomationTool
 		/// </summary>
 		/// <param name="files">List of full file paths</param>
 		/// <param name="rootDir">Root folder for all the files. All files must be relative to this RootDir.</param>
-		public TempStorageManifest(FileInfo[] files, DirectoryReference rootDir)
+		public TempStorageBlockManifest(FileInfo[] files, DirectoryReference rootDir)
 		{
 			Files = files.Select(x => new TempStorageFile(x, rootDir)).ToArray();
 			ZipFiles = Array.Empty<TempStorageZipFile>();
@@ -392,11 +393,11 @@ namespace AutomationTool
 		/// Load a manifest from disk
 		/// </summary>
 		/// <param name="file">File to load</param>
-		public static TempStorageManifest Load(FileReference file)
+		public static TempStorageBlockManifest Load(FileReference file)
 		{
 			using (StreamReader reader = new(file.FullName))
 			{
-				return (TempStorageManifest)(s_serializer.Deserialize(reader) ?? throw new InvalidOperationException());
+				return (TempStorageBlockManifest)(s_serializer.Deserialize(reader) ?? throw new InvalidOperationException());
 			}
 		}
 
@@ -442,7 +443,7 @@ namespace AutomationTool
 		/// </summary>
 		[XmlArray]
 		[XmlArrayItem("Block")]
-		public TempStorageBlock[] Blocks { get; set; }
+		public TempStorageBlockRef[] Blocks { get; set; }
 
 		/// <summary>
 		/// List of keys for published artifacts
@@ -463,7 +464,7 @@ namespace AutomationTool
 		{
 			LocalFiles = Array.Empty<string>();
 			ExternalFiles = Array.Empty<string>();
-			Blocks = Array.Empty<TempStorageBlock>();
+			Blocks = Array.Empty<TempStorageBlockRef>();
 			ArtifactKeys = Array.Empty<string>();
 		}
 
@@ -474,7 +475,7 @@ namespace AutomationTool
 		/// <param name="rootDir">Root folder for all the files. All files must be relative to this RootDir.</param>
 		/// <param name="blocks">Referenced storage blocks required for these files</param>
 		/// <param name="artifactKeys">Keys for published artifacts</param>
-		public TempStorageTagManifest(IEnumerable<FileReference> files, DirectoryReference rootDir, IEnumerable<TempStorageBlock> blocks, IEnumerable<string> artifactKeys)
+		public TempStorageTagManifest(IEnumerable<FileReference> files, DirectoryReference rootDir, IEnumerable<TempStorageBlockRef> blocks, IEnumerable<string> artifactKeys)
 		{
 			List<string> newLocalFiles = new List<string>();
 			List<string> newExternalFiles = new List<string>();
@@ -672,7 +673,7 @@ namespace AutomationTool
 			}
 
 			// Check that each of the tags exist
-			HashSet<TempStorageBlock> blocks = new HashSet<TempStorageBlock>();
+			HashSet<TempStorageBlockRef> blocks = new HashSet<TempStorageBlockRef>();
 			foreach (string tagName in tagNames)
 			{
 				// Check the local manifest exists
@@ -708,7 +709,7 @@ namespace AutomationTool
 			}
 
 			// Check that each of the outputs match
-			foreach (TempStorageBlock block in blocks)
+			foreach (TempStorageBlockRef block in blocks)
 			{
 				// Check the local manifest exists
 				FileReference localManifestFile = GetManifestLocation(_localDir, block.NodeName, block.OutputName);
@@ -738,7 +739,7 @@ namespace AutomationTool
 				}
 
 				// Read the manifest and check the files
-				TempStorageManifest localManifest = TempStorageManifest.Load(localManifestFile);
+				TempStorageBlockManifest localManifest = TempStorageBlockManifest.Load(localManifestFile);
 				if (localManifest.Files.Any(x => !ignoreModified.Matches(x.ToFileReference(_rootDir).FullName) && !x.Compare(_rootDir)))
 				{
 					return false;
@@ -753,7 +754,7 @@ namespace AutomationTool
 		/// <param name="nodeName">Name of the node which produced the tag set</param>
 		/// <param name="tagName">Name of the tag, with a '#' prefix</param>
 		/// <returns>The set of files</returns>
-		public TempStorageTagManifest? ReadFileList(string nodeName, string tagName)
+		public TempStorageTagManifest? ReadTagFileList(string nodeName, string tagName)
 		{
 #pragma warning disable CA1508 // False positive; FileList is always null
 			TempStorageTagManifest? fileList = null;
@@ -762,7 +763,7 @@ namespace AutomationTool
 			FileReference localFileListLocation = GetTaggedFileListLocation(_localDir, nodeName, tagName);
 			if (FileReference.Exists(localFileListLocation))
 			{
-				Logger.LogInformation("Reading local file list from {Arg0}", localFileListLocation.FullName);
+				Logger.LogInformation("Reading tag \"{NodeName}\":\"{TagName}\" from local file {File}", nodeName, tagName, localFileListLocation.FullName);
 				fileList = TempStorageTagManifest.Load(localFileListLocation);
 			}
 			else
@@ -790,7 +791,7 @@ namespace AutomationTool
 					PerformActionWithRetries(() =>
 					{
 						// Read the shared manifest
-						Logger.LogInformation("Copying shared tag set from {Arg0} to {Arg1}", sharedFileListLocation.FullName, localFileListLocation.FullName);
+						Logger.LogInformation("Copying tag \"{NodeName}\":\"{TagName}\" from {SourceFile} to {TargetFile}", nodeName, tagName, sharedFileListLocation.FullName, localFileListLocation.FullName);
 						fileList = TempStorageTagManifest.Load(sharedFileListLocation);
 					}, attempts, TimeSpan.FromSeconds(5));
 				}
@@ -816,7 +817,7 @@ namespace AutomationTool
 		/// <param name="blocks">List of referenced storage blocks</param>
 		/// <param name="artifactKeys">Keys for published artifacts</param>
 		/// <returns>The set of files</returns>
-		public void WriteFileList(string nodeName, string tagName, IEnumerable<FileReference> files, IEnumerable<TempStorageBlock> blocks, IEnumerable<string> artifactKeys)
+		public void WriteFileList(string nodeName, string tagName, IEnumerable<FileReference> files, IEnumerable<TempStorageBlockRef> blocks, IEnumerable<string> artifactKeys)
 		{
 			// Create the file list
 			TempStorageTagManifest fileList = new TempStorageTagManifest(files, _rootDir, blocks, artifactKeys);
@@ -859,13 +860,13 @@ namespace AutomationTool
 		/// <param name="buildProducts">Array of build products to be archived</param>
 		/// <param name="pushToRemote">Allow skipping the copying of this manifest to shared storage, because it's not required by any other agent</param>
 		/// <returns>The created manifest instance (which has already been saved to disk).</returns>
-		public TempStorageManifest Archive(string nodeName, string? blockName, FileReference[] buildProducts, bool pushToRemote = true)
+		public TempStorageBlockManifest Archive(string nodeName, string? blockName, FileReference[] buildProducts, bool pushToRemote = true)
 		{
 			using (IScope scope = GlobalTracer.Instance.BuildSpan("StoreToTempStorage").StartActive())
 			{
 				// Create a manifest for the given build products
 				FileInfo[] files = buildProducts.Select(x => new FileInfo(x.FullName)).ToArray();
-				TempStorageManifest manifest = new TempStorageManifest(files, _rootDir);
+				TempStorageBlockManifest manifest = new TempStorageBlockManifest(files, _rootDir);
 
 				// Create the local directory for this node
 				DirectoryReference localNodeDir = GetDirectoryForNode(_localDir, nodeName);
@@ -884,13 +885,13 @@ namespace AutomationTool
 					manifest.ZipFiles = zipFiles.Select(x => new TempStorageZipFile(x)).ToArray();
 
 					// Save the shared manifest
-					Logger.LogInformation("Saving shared manifest to {Arg0}", sharedManifestFile.FullName);
+					Logger.LogInformation("Saving block \"{NodeName}\":\"{BlockName}\" manifest to {File}", nodeName, blockName, sharedManifestFile.FullName);
 					PerformActionWithRetries(() => manifest.Save(sharedManifestFile), 3, TimeSpan.FromSeconds(5));
 				}
 
 				// Save the local manifest
 				FileReference localManifestFile = GetManifestLocation(_localDir, nodeName, blockName);
-				Logger.LogInformation("Saving local manifest to {Arg0}", localManifestFile.FullName);
+				Logger.LogInformation("Saving block \"{NodeName}\":\"{BlockName}\" manifest to {File}", nodeName, blockName, localManifestFile.FullName);
 				manifest.Save(localManifestFile);
 
 				// Update the stats
@@ -911,7 +912,7 @@ namespace AutomationTool
 		/// <param name="outputName">Name of the block to retrieve. May be null or empty.</param>
 		/// <param name="ignoreModified">Filter for files to ignore</param>
 		/// <returns>Manifest of the files retrieved</returns>
-		public TempStorageManifest Retrieve(string nodeName, string? outputName, FileFilter ignoreModified)
+		public TempStorageBlockManifest Retrieve(string nodeName, string? outputName, FileFilter ignoreModified)
 		{
 			using (IScope scope = GlobalTracer.Instance.BuildSpan("RetrieveFromTempStorage").StartActive())
 			{
@@ -920,11 +921,11 @@ namespace AutomationTool
 				bool local = FileReference.Exists(localManifestFile);
 
 				// Read the manifest, either from local storage or shared storage
-				TempStorageManifest? manifest = null;
+				TempStorageBlockManifest? manifest = null;
 				if (local)
 				{
-					Logger.LogInformation("Reading shared manifest from {Arg0}", localManifestFile.FullName);
-					manifest = TempStorageManifest.Load(localManifestFile);
+					Logger.LogInformation("Reading tag \"{NodeName}\":\"{TagName}\" manifest from {File}", nodeName, outputName, localManifestFile.FullName);
+					manifest = TempStorageBlockManifest.Load(localManifestFile);
 				}
 				else
 				{
@@ -944,8 +945,8 @@ namespace AutomationTool
 					}
 
 					// Read the shared manifest
-					Logger.LogInformation("Copying shared manifest from {Arg0} to {Arg1}", sharedManifestFile.FullName, localManifestFile.FullName);
-					PerformActionWithRetries(() => manifest = TempStorageManifest.Load(sharedManifestFile), 3, TimeSpan.FromSeconds(5));
+					Logger.LogInformation("Copying tag \"{NodeName}\":\"{TagName}\" from {SourceFile} to {TargetFile}", nodeName, outputName, sharedManifestFile.FullName, localManifestFile.FullName);
+					PerformActionWithRetries(() => manifest = TempStorageBlockManifest.Load(sharedManifestFile), 3, TimeSpan.FromSeconds(5));
 
 					// Unzip all the build products
 					DirectoryReference sharedNodeDir = GetDirectoryForNode(_sharedDir, nodeName);
@@ -1507,10 +1508,10 @@ namespace AutomationTool
 			tempStore.Archive("TestNode", "NamedOutput", namedOutput.Keys.ToArray(), true);
 			
 			// Check both outputs are still ok
-			TempStorageManifest defaultManifest = tempStore.Retrieve("TestNode", null, new FileFilter());
+			TempStorageBlockManifest defaultManifest = tempStore.Retrieve("TestNode", null, new FileFilter());
 			CheckManifest(workingDir, defaultManifest, defaultOutput);
 
-			TempStorageManifest namedManifest = tempStore.Retrieve("TestNode", "NamedOutput", new FileFilter());
+			TempStorageBlockManifest namedManifest = tempStore.Retrieve("TestNode", "NamedOutput", new FileFilter());
 			CheckManifest(workingDir, namedManifest, namedOutput);
 
 			// Delete local temp storage and the working directory and try again
@@ -1536,7 +1537,7 @@ namespace AutomationTool
 			}
 
 			// Second one should be fine
-			TempStorageManifest namedManifestFromShared = tempStore.Retrieve("TestNode", "NamedOutput", new FileFilter());
+			TempStorageBlockManifest namedManifestFromShared = tempStore.Retrieve("TestNode", "NamedOutput", new FileFilter());
 			CheckManifest(workingDir, namedManifestFromShared, namedOutput);
 		}
 
@@ -1567,7 +1568,7 @@ namespace AutomationTool
 		/// <param name="rootDir">Root directory for relative paths in the manifest</param>
 		/// <param name="manifest">Manifest to check</param>
 		/// <param name="files">Mapping of filename to timestamp as expected in the manifest</param>
-		static void CheckManifest(DirectoryReference rootDir, TempStorageManifest manifest, Dictionary<FileReference, DateTime> files)
+		static void CheckManifest(DirectoryReference rootDir, TempStorageBlockManifest manifest, Dictionary<FileReference, DateTime> files)
 		{
 			if (files.Count != manifest.Files.Length)
 			{
