Localization:
Add zh-Hans as a culture to be generated s part of the default modular localization config files for plugins using auto loc.
- Simplified chinese (zh-Hans) will be added as a supported language for plugins that use auto localization moving foward.

diff --git a/Engine/Source/Programs/AutomationTool/Localization/LocalizationConfigFileParams.cs b/Engine/Source/Programs/AutomationTool/Localization/LocalizationConfigFileParams.cs
--- a/Engine/Source/Programs/AutomationTool/Localization/LocalizationConfigFileParams.cs
+++ b/Engine/Source/Programs/AutomationTool/Localization/LocalizationConfigFileParams.cs
@@ -34,7 +34,7 @@ namespace EpicGames.Localization
 			commonSettings.ResourceName = localizationTargetName + LocalizationFileExtensions.LocalizationResourceFileExtension;
 
 			commonSettings.NativeCulture = "en";
-			commonSettings.CulturesToGenerate = new List<string>() { "en", "fr", "de", "pl", "es-419", "es", "pt-BR", "it", "ru", "ko", "tr", "ar", "ja" };
+			commonSettings.CulturesToGenerate = new List<string>() { "en", "fr", "de", "pl", "es-419", "es", "pt-BR", "it", "ru", "ko", "tr", "ar", "ja", "zh-Hans"};
 			commonSettings.CulturesToGenerate.Sort();
 			commonSettings.bSkipSourceCheck = false;
 			commonSettings.bValidateFormatPatterns = true;
