BuildGraph: Use a different folder for temp storage manifests when UAT is spawned recursively.

diff --git a/Engine/Source/Programs/AutomationTool/BuildGraph/BuildGraph.cs b/Engine/Source/Programs/AutomationTool/BuildGraph/BuildGraph.cs
--- a/Engine/Source/Programs/AutomationTool/BuildGraph/BuildGraph.cs
+++ b/Engine/Source/Programs/AutomationTool/BuildGraph/BuildGraph.cs
@@ -507,9 +507,13 @@ namespace AutomationTool
 				}
 			}
 
-			// Create the temp storage handler
+			// Get the temp storage manifest directory. When spawning buildgraph through a UAT child process, be careful not to
+			// overwrite any manifests from the parent. These may be required by the managing build system.
 			DirectoryReference rootDir = new DirectoryReference(CommandUtils.CmdEnv.LocalRoot);
-			TempStorage storage = new TempStorage(rootDir, DirectoryReference.Combine(rootDir, "Engine", "Saved", "BuildGraph"), (sharedStorageDir == null) ? null : new DirectoryReference(sharedStorageDir), writeToSharedStorage);
+			DirectoryReference manifestDir = DirectoryReference.Combine(rootDir, "Engine", "Saved", CmdEnv.IsChildInstance? "BuildGraphChildInstance" : "BuildGraph");
+
+			// Create the temp storage handler
+			TempStorage storage = new TempStorage(rootDir, manifestDir, (sharedStorageDir == null) ? null : new DirectoryReference(sharedStorageDir), writeToSharedStorage);
 			if (!resume)
 			{
 				storage.CleanLocal();
