Introduce TestData v2 output
- Move previous iterations of the HordeReport to a legacy implementation section
- Reimplement AutomatedTestSessionData base TestData v2 data model
- Change ITestReport AddEvent signature in favor of using formating with params keyword
- Remove redundancy of ITestEvent based property (UnrealTestNode.Events)
- Use Dictionary for Artifact Files in UnrealAutomatedTestResult json data model instead of specialized image comparison type

#rnx

diff --git a/Engine/Source/Programs/AutomationTool/Gauntlet/Framework/Base/Gauntlet.TestReport.cs b/Engine/Source/Programs/AutomationTool/Gauntlet/Framework/Base/Gauntlet.TestReport.cs
--- a/Engine/Source/Programs/AutomationTool/Gauntlet/Framework/Base/Gauntlet.TestReport.cs
+++ b/Engine/Source/Programs/AutomationTool/Gauntlet/Framework/Base/Gauntlet.TestReport.cs
@@ -6,7 +6,6 @@ using System.Collections.Generic;
 using System.Text.Json;
 using System.Text.Json.Serialization;
 
-
 namespace Gauntlet
 {
 	/// <summary>
@@ -68,9 +67,9 @@ namespace Gauntlet
 		/// </summary>
 		/// <param name="Type"></param>
 		/// <param name="Message"></param>
-		/// <param name="Context"></param>
+		/// <param name="Args"></param>
 		/// <returns></returns>
-		void AddEvent(EventType Type, string Message, object Context = null);
+		void AddEvent(EventType Type, string Message, params object[] Args);
 
 		/// <summary>
 		/// Attach an Artifact to the ITestReport
@@ -184,21 +183,38 @@ namespace Gauntlet
 		/// </summary>
 		/// <param name="Type"></param>
 		/// <param name="Message"></param>
-		/// <param name="Context"></param>
+		/// <param name="Args"></param>
 		/// <returns></returns>
-		public abstract void AddEvent(EventType Type, string Message, object Context = null);
+		public abstract void AddEvent(EventType Type, string Message, params object[] Args);
 
-		public void AddError(string Message, object Context = null)
+		/// <summary>
+		/// Add Error event with optional formatting inputs to report
+		/// </summary>
+		/// <param name="Message"></param>
+		/// <param name="Args"></param>
+		public void AddError(string Message, params object[] Args)
 		{
-			AddEvent(EventType.Error, Message, Context);
+			AddEvent(EventType.Error, Message, Args);
 		}
-		public void AddWarning(string Message, object Context = null)
+
+		/// <summary>
+		/// Add Warning event with optional formatting inputs to report
+		/// </summary>
+		/// <param name="Message"></param>
+		/// <param name="Args"></param>
+		public void AddWarning(string Message, params object[] Args)
 		{
-			AddEvent(EventType.Warning, Message, Context);
+			AddEvent(EventType.Warning, Message, Args);
 		}
-		public void AddInfo(string Message, object Context = null)
+
+		/// <summary>
+		/// Add Info event with optional formatting inputs to report
+		/// </summary>
+		/// <param name="Message"></param>
+		/// <param name="Args"></param>
+		public void AddInfo(string Message, params object[] Args)
 		{
-			AddEvent(EventType.Info, Message, Context);
+			AddEvent(EventType.Info, Message, Args);
 		}
 
 		/// <summary>
diff --git a/Engine/Source/Programs/AutomationTool/Gauntlet/Framework/Gauntlet.Utils.cs b/Engine/Source/Programs/AutomationTool/Gauntlet/Framework/Gauntlet.Utils.cs
--- a/Engine/Source/Programs/AutomationTool/Gauntlet/Framework/Gauntlet.Utils.cs
+++ b/Engine/Source/Programs/AutomationTool/Gauntlet/Framework/Gauntlet.Utils.cs
@@ -500,6 +500,8 @@ namespace Gauntlet
 
 			return SBuilder.ToString();
 		}
+
+		static public HashAlgorithm DefaultAlgo = SHA1.Create();
 	}
 
 	/*
@@ -1950,6 +1952,17 @@ namespace Gauntlet
 		{
 			return Regex.Replace(Name, @"[^a-z0-9_\-.]+", "_", RegexOptions.IgnoreCase);
 		}
+
+		/// <summary>
+		/// Convert file path to a Uri valid and Url encoded string
+		/// </summary>
+		/// <param name="Path"></param>
+		/// <returns></returns>
+		static public string ConvertPathToUri(string Path)
+		{
+			// Uri.AbsolutPath remove only 'file://', the extra / still need to be removed, enhance using string.Substring(1) 
+			return new Uri($"file:///{Path}").AbsolutePath.Substring(1);
+		}
 	}
 
 }
diff --git a/Engine/Source/Programs/AutomationTool/Gauntlet/Framework/Utils/Gauntlet.HordeReport.cs b/Engine/Source/Programs/AutomationTool/Gauntlet/Framework/Utils/Gauntlet.HordeReport.cs
--- a/Engine/Source/Programs/AutomationTool/Gauntlet/Framework/Utils/Gauntlet.HordeReport.cs
+++ b/Engine/Source/Programs/AutomationTool/Gauntlet/Framework/Utils/Gauntlet.HordeReport.cs
@@ -4,12 +4,15 @@ using System;
 using System.Collections.Generic;
 using System.IO;
 using System.Linq;
-using System.Text;
 using System.Text.Json;
 using System.Text.RegularExpressions;
 using System.Text.Json.Serialization;
 using AutomationTool;
 using System.Security.Cryptography;
+using Microsoft.CodeAnalysis;
+using EpicGames.Core;
+using Logging = Microsoft.Extensions.Logging;
+using System.Globalization;
 
 namespace Gauntlet
 {
@@ -42,11 +45,11 @@ namespace Gauntlet
 			/// <summary>
 			/// Horde report version
 			/// </summary>
-			public int Version { get; set; } = 1;
+			public virtual int Version => 1;
 
 			protected string OutputArtifactPath;
 			protected HashSet<string> ArtifactProcessedHashes;
-			protected Dictionary<string, object> ExtraReports;			
+			protected Dictionary<string, object> ExtraReports;
 
 			/// <summary>
 			/// Attach Artifact to the Test Report
@@ -72,14 +75,8 @@ namespace Gauntlet
 					throw new InvalidOperationException("OutputArtifactPath must be set before attaching any artifact");
 				}
 
-				string ArtifactHash;
-				{
-					// Generate a hash from the artifact path
-					using (SHA1 Sha1 = SHA1.Create())
-					{
-						ArtifactHash = Hasher.ComputeHash(ArtifactPath, Sha1, 8);
-					}
-				}
+				// Generate a hash from the artifact path
+				string ArtifactHash = Hasher.ComputeHash(ArtifactPath, Hasher.DefaultAlgo, 8);
 
 				if (ArtifactProcessedHashes == null)
 				{
@@ -134,6 +131,14 @@ namespace Gauntlet
 				Log.Verbose(string.Format("Test Report output artifact path is set to: {0}", OutputArtifactPath));
 			}
 
+			/// <summary>
+			/// Attach a log path to a device instance name
+			/// </summary>
+			/// <param name="DeviceName"></param>
+			/// <param name="LogPath"></param>
+			/// <returns></returns>
+			public abstract bool AttachDeviceLog(string InstanceName, string LogPath, string LogName);
+
 			public void AttachDependencyReport(object InReport, string Key = null)
 			{
 				if(ExtraReports == null)
@@ -157,6 +162,11 @@ namespace Gauntlet
 				return Reports;
 			}
 
+			/// <summary>
+			/// Return the Key to be used for the TestData
+			/// </summary>
+			/// <param name="BaseKey"></param>
+			/// <returns></returns>
 			public virtual string GetTestDataKey(string BaseKey = null)
 			{
 				if(BaseKey == null)
@@ -165,8 +175,36 @@ namespace Gauntlet
 				}
 				return string.Format("{0}::{1}", Type, BaseKey);
 			}
+
+			/// <summary>
+			/// Return the data to be stored in the TestData
+			/// </summary>
+			/// <returns></returns>
+			public virtual object GetTestData()
+			{
+				return this;
+			}
+
+			/// <summary>
+			/// Return the version of the TestData
+			/// </summary>
+			/// <returns></returns>
+			public virtual int GetVersion()
+			{
+				return Version;
+			}
+
+			/// <summary>
+			/// Return the metadata key/value set as Dictionary
+			/// </summary>
+			/// <returns></returns>
+			public virtual Dictionary<string, string> GetMetadata()
+			{
+				return Metadata;
+			}
 		}
 
+		#region Legacy Implementations
 		/// <summary>
 		/// Contains detailed information about device that run tests
 		/// </summary>
@@ -182,6 +220,7 @@ namespace Gauntlet
 			public int RAMInGB { get; set; }
 			public string RenderMode { get; set; }
 			public string RHI { get; set; }
+			public string AppInstanceLog { get; set; }
 		}
 
 		/// <summary>
@@ -375,6 +414,20 @@ namespace Gauntlet
 				return NewDevice;
 			}
 
+			public override bool AttachDeviceLog(string InstanceName, string LogPath, string LogName)
+			{
+				if (AttachArtifact(LogPath, LogName))
+				{
+					Device item = Devices.LastOrDefault(D => D.Instance == InstanceName);
+					if (item != null)
+					{
+						item.AppInstanceLog = LogName;
+						return true;
+					}
+				}
+				return false;
+			}
+
 			/// <summary>
 			/// Add a new TestResult to the pass results and return it 
 			/// </summary>
@@ -386,7 +439,7 @@ namespace Gauntlet
 				return NewTestResult;
 			}
 
-			public override void AddEvent(EventType Type, string Message, object Context = null)
+			public override void AddEvent(EventType Type, string Message, params object[] Args)
 			{
 				throw new System.NotImplementedException("AddEvent not implemented");
 			}
@@ -447,9 +500,9 @@ namespace Gauntlet
 							NewArtifact.Name = InTestArtifact.Name;
 							NewArtifact.Type = InTestArtifact.Type;
 							ComparisonFiles ArtifactFiles = NewArtifact.Files;
-							ArtifactFiles.Difference = InTestArtifact.Files.Difference;
-							ArtifactFiles.Approved = InTestArtifact.Files.Approved;
-							ArtifactFiles.Unapproved = InTestArtifact.Files.Unapproved;
+							ArtifactFiles.Difference = InTestArtifact.Files.GetValueOrDefault("difference");
+							ArtifactFiles.Approved = InTestArtifact.Files.GetValueOrDefault("approved");
+							ArtifactFiles.Unapproved = InTestArtifact.Files.GetValueOrDefault("unapproved");
 						}
 						foreach (UnrealAutomationEntry InTestEntry in InTestResult.Entries)
 						{
@@ -509,7 +562,7 @@ namespace Gauntlet
 		/// <summary>
 		/// Contains information about a test session 
 		/// </summary>
-		public class AutomatedTestSessionData : BaseHordeReport
+		public class AutomatedTestSessionData_legacy : BaseHordeReport
 		{
 			public override string Type
 			{
@@ -545,12 +598,9 @@ namespace Gauntlet
 				}
 				private string GenerateHash(string InName)
 				{
-					using (SHA1 Sha1 = SHA1.Create())
-					{
-						return Hasher.ComputeHash(InName, Sha1, 8);
-					}
+					return Hasher.ComputeHash(InName, Hasher.DefaultAlgo, 8);
 				}
-				public TestEvent NewEvent(EventType InType, string InMessage, string InTag, string InContext, string InDateTime)
+				public TestEvent NewEvent(string InDateTime, EventType InType, string InMessage, string InTag, string InContext = null)
 				{
 					TestEvent NewItem = new TestEvent(InType, InMessage, InTag, InContext ?? "", InDateTime);
 					Events.Add(NewItem);
@@ -685,10 +735,7 @@ namespace Gauntlet
 				private string GenerateEventHash()
 				{
 					string FilteredEvent = FilterEvent(string.Format("{0}{1}{2}{3}", Type.ToString(), Message, Context, Tag));
-					using (SHA1 Sha1 = SHA1.Create())
-					{
-						return Hasher.ComputeHash(FilteredEvent, Sha1, 8);
-					}
+					return Hasher.ComputeHash(FilteredEvent, Hasher.DefaultAlgo, 8);
 				}
 			}
 			public class TestResultData
@@ -717,7 +764,7 @@ namespace Gauntlet
 				}
 			}
 
-			public AutomatedTestSessionData(string InName) : base()
+			public AutomatedTestSessionData_legacy(string InName) : base()
 			{
 				Name = InName;
 				PreFlightChange = "";
@@ -809,25 +856,28 @@ namespace Gauntlet
 			/// <summary>
 			/// Add event to current test
 			/// </summary>
-			/// <param name="InType"></param>
-			/// <param name="InMessage"></param>
-			/// <param name="InContext"></param>
-			public override void AddEvent(EventType InType, string InMessage, object InContext = null)
+			/// <param name="Type"></param>
+			/// <param name="Message"></param>
+			/// <param name="Args"></param>
+			public override void AddEvent(EventType Type, string Message, params object[] Args)
 			{
-				AddEvent(InType, InMessage, "gauntlet", InContext == null ? null : InContext.ToString());
+				Dictionary<string, object> Properties = new Dictionary<string, object>();
+				MessageTemplate.ParsePropertyValues(Message, Args, Properties);
+				Message = MessageTemplate.Render(Message, Properties);
+				AddEvent(DateTime.UtcNow.ToString(UnrealAutomationEntry.DateTimeFormat), Type, Message, "gauntlet", null);
 			}
 			/// <summary>
 			/// Overload of AddEvent, add even to current test with date time
 			/// </summary>
+			/// <param name="InDateTime"></param>
 			/// <param name="InType"></param>
 			/// <param name="InMessage"></param>
 			/// <param name="InTag"></param>
 			/// <param name="InContext"></param>
-			/// <param name="InDateTime"></param>
-			public void AddEvent(EventType InType, string InMessage, string InTag, string InContext = null, string InDateTime = null)
+			public void AddEvent(string InDateTime, EventType InType, string InMessage, string InTag, string InContext)
 			{
 				TestResult CurrentTest = GetCurrentTest();
-				TestEvent Event = CurrentTest.NewEvent(InType, InMessage, InTag, InContext, InDateTime ?? DateTime.Now.ToString("yyyy.MM.dd-HH.mm.ss"));
+				TestEvent Event = CurrentTest.NewEvent(InDateTime, InType, InMessage, InTag, InContext);
 				if(InType == EventType.Error)
 				{
 					IndexTestError(Event.Hash, CurrentTest.TestUID, InMessage, InTag);
@@ -856,6 +906,21 @@ namespace Gauntlet
 
 				return false;
 			}
+
+			public override bool AttachDeviceLog(string InstanceName, string LogPath, string LogName)
+			{
+				if (AttachArtifact(LogPath, LogName))
+				{
+					Device item = Devices.LastOrDefault(D => D.AppInstanceName == InstanceName);
+					if (item != null)
+					{
+						item.AppInstanceLog = LogName;
+						return true;
+					}
+				}
+				return false;
+			}
+
 			public override string GetTestDataKey(string BaseKey = null)
 			{
 				return base.GetTestDataKey(); // Ignore BaseKey
@@ -880,9 +945,9 @@ namespace Gauntlet
 			/// <param name="InReportPath"></param>
 			/// <param name="InHordeArtifactPath"></param>
 			/// <returns></returns>
-			public static AutomatedTestSessionData FromUnrealAutomatedTests(UnrealAutomatedTestPassResults InTestPassResults, string InName, string InSuite, string InReportPath, string InHordeArtifactPath)
+			public static AutomatedTestSessionData_legacy FromUnrealAutomatedTests(UnrealAutomatedTestPassResults InTestPassResults, string InName, string InSuite, string InReportPath, string InHordeArtifactPath)
 			{
-				AutomatedTestSessionData OutTestPassResults = new AutomatedTestSessionData(InName);
+				AutomatedTestSessionData_legacy OutTestPassResults = new AutomatedTestSessionData_legacy(InName);
 				OutTestPassResults.SetOutputArtifactPath(InHordeArtifactPath);
 				if (InTestPassResults.Devices != null)
 				{
@@ -939,44 +1004,47 @@ namespace Gauntlet
 								Tag = "crash";
 							}
 
-							OutTestPassResults.AddEvent(InTestEntry.Event.Type, InTestEntry.Event.Message, Tag, Context, InTestEntry.Timestamp);
+							OutTestPassResults.AddEvent(InTestEntry.Timestamp, InTestEntry.Event.Type, InTestEntry.Event.Message, Tag, Context);
 							bool IsInfo = InTestEntry.Event.Type == EventType.Info;
 							// Add Artifacts
 							if (IntArtifactEntry != null)
 							{
 								List<string> FailedToAttached = new List<string>();
-								if (!IsInfo && !string.IsNullOrEmpty(IntArtifactEntry.Files.Difference))
+								string ArtifactFilePath = IntArtifactEntry.Files.GetValueOrDefault("difference");
+								if (!IsInfo && !string.IsNullOrEmpty(ArtifactFilePath))
 								{
 									if(!OutTestPassResults.AddArtifactToLastEvent(
 										"difference",
-										Path.Combine(InReportPath, IntArtifactEntry.Files.Difference),
-										IntArtifactEntry.Files.Difference
+										Path.Combine(InReportPath, ArtifactFilePath),
+										ArtifactFilePath
 									))
 									{
-										FailedToAttached.Add(IntArtifactEntry.Files.Difference);
+										FailedToAttached.Add(ArtifactFilePath);
 									}
 								}
-								if (!IsInfo && !string.IsNullOrEmpty(IntArtifactEntry.Files.Approved))
+								ArtifactFilePath = IntArtifactEntry.Files.GetValueOrDefault("approved");
+								if (!IsInfo && !string.IsNullOrEmpty(ArtifactFilePath))
 								{
 									if(!OutTestPassResults.AddArtifactToLastEvent(
 										"approved",
-										Path.Combine(InReportPath, IntArtifactEntry.Files.Approved),
-										IntArtifactEntry.Files.Approved
+										Path.Combine(InReportPath, ArtifactFilePath),
+										ArtifactFilePath
 									))
 									{
-										FailedToAttached.Add(IntArtifactEntry.Files.Approved);
+										FailedToAttached.Add(ArtifactFilePath);
 									}
 								}
-								if (!string.IsNullOrEmpty(IntArtifactEntry.Files.Unapproved))
+								ArtifactFilePath = IntArtifactEntry.Files.GetValueOrDefault("unapproved");
+								if (!string.IsNullOrEmpty(ArtifactFilePath))
 								{
-									string AbsoluteLocation = Path.Combine(InReportPath, IntArtifactEntry.Files.Unapproved);
+									string AbsoluteLocation = Path.Combine(InReportPath, ArtifactFilePath);
 									if(!OutTestPassResults.AddArtifactToLastEvent(
 										"unapproved",
 										AbsoluteLocation,
-										IntArtifactEntry.Files.Unapproved
+										ArtifactFilePath
 									))
 									{
-										FailedToAttached.Add(IntArtifactEntry.Files.Unapproved);
+										FailedToAttached.Add(ArtifactFilePath);
 									}
 									// Add Json meta data if any
 									string MetadataLocation = Utils.SystemHelpers.GetFullyQualifiedPath(Path.GetDirectoryName(AbsoluteLocation));
@@ -985,13 +1053,13 @@ namespace Gauntlet
 										string[] JsonMetadataFiles = System.IO.Directory.GetFiles(MetadataLocation, "*.json");
 										if (JsonMetadataFiles.Length > 0)
 										{
-											int LastSlash = IntArtifactEntry.Files.Unapproved.LastIndexOf("/");
-											string RelativeLocation = IntArtifactEntry.Files.Unapproved.Substring(0, LastSlash);
+											int LastSlash = ArtifactFilePath.LastIndexOf("/");
+											string RelativeLocation = ArtifactFilePath.Substring(0, LastSlash);
 											OutTestPassResults.AddEvent(
+												InTestEntry.Timestamp,
 												EventType.Info,
 												"The image reference can be updated by pointing the Screen Comparison tab from the Test Automation window to the artifacts from this test.",
-												"image comparison metadata", null,
-												InTestEntry.Timestamp
+												"image comparison metadata", null
 											);
 											foreach (string JsonFile in JsonMetadataFiles)
 											{
@@ -1082,8 +1150,12 @@ namespace Gauntlet
 			public List<String> Errors { get; set; } = new List<String>();
 			public List<String> Warnings { get; set; } = new List<String>();
 
-			public override void AddEvent(EventType Type, string Message, object Context = null)
+			public override void AddEvent(EventType Type, string Message, params object[] Args)
 			{
+				Dictionary<string, object> Properties = new Dictionary<string, object>();
+				MessageTemplate.ParsePropertyValues(Message, Args, Properties);
+				Message = MessageTemplate.Render(Message, Properties);
+
 				switch (Type)
 				{
 					case EventType.Error:
@@ -1105,106 +1177,1201 @@ namespace Gauntlet
 				}
 				return false;
 			}
+
+			public override bool AttachDeviceLog(string InstanceName, string LogPath, string LogName)
+			{
+				if (AttachArtifact(LogPath, LogName))
+				{
+					Logs.Add(string.Format("Attached Log for device {0}: {1}", InstanceName, LogName));
+				}
+				return true;
+			}
 		}
+		#endregion
 
 		/// <summary>
-		/// Container for Test Data items 
+		/// Contains information about a test session (v2)
 		/// </summary>
-		public class TestDataCollection
+		public class AutomatedTestSessionData : BaseHordeReport
 		{
-			public class DataItem
+			public override string Type => "Automated Test Session";
+
+			public override int Version => 2;
+
+			/// <summary>
+			/// Unique key that identify the test and group the different sessions
+			/// </summary>
+			public string Key { get; set; }
+
+			public override string GetTestDataKey(string BaseKey = null)
 			{
-				public string Key { get; set; }
-				public object Data { get; set; }				
+				return Key; // Ignore BaseKey
 			}
-			public TestDataCollection()
+
+			/// <summary>
+			/// Stored data of the session
+			/// </summary>
+			protected TestSession Data { get; set; }
+
+			public override object GetTestData()
+			{
+				return Data;
+			}
+
+			/// <summary>
+			/// Utility function to generate a valid key from any name
+			/// </summary>
+			/// <param name="InName"></param>
+			/// <returns></returns>
+			private static string GenerateKey(string InName)
+			{
+				return Hasher.ComputeHash(InName, Hasher.DefaultAlgo, 8);
+			}
+
+			/// <summary>
+			/// Cursor to the current phase
+			/// </summary>
+			private TestPhase CurrentPhase { get; set; }
+
+			/// <summary>
+			/// Static constructor
+			/// </summary>
+			static AutomatedTestSessionData()
+			{
+				// Make sure to register the JSON formatter for TestLogValue type
+				LogValueFormatter.RegisterFormatter(typeof(TestLogValue), new TestLogValueFormatter());
+			}
+
+			public AutomatedTestSessionData(string InKey, string InName = null) : base()
 			{
-				Items = new List<DataItem>();
+				Key = InKey;
+				Data = new TestSession(string.IsNullOrEmpty(InName)? InKey : InName, this);
+				CurrentPhase = null;
 			}
 
-			public DataItem AddNewTestReport(ITestReport InData, string InKey = null)
+			/// <summary>
+			/// Represent a Device used during a test
+			/// </summary>
+			public class TestDevice
 			{
-				if(InData is BaseHordeReport InHordeReport)
+				public string name { get; set; }
+				public string appInstanceLogPath { get; set; }
+				public Dictionary<string, string> metadata { get; set; }
+
+				public TestDevice(string InName)
 				{
-					InKey = InHordeReport.GetTestDataKey(InKey);
+					name = InName;
+					metadata = new Dictionary<string, string>();
 				}
-				DataItem NewDataItem = new DataItem();
-				NewDataItem.Key = string.IsNullOrEmpty(InKey) ? InData.Type : InKey;
-				NewDataItem.Data = InData;
+			}
 
-				var FoundItemIndex = Items.FindIndex(I => I.Key == InKey);
-				if (FoundItemIndex == -1)
+			/// <summary>
+			/// Summary for the test session
+			/// </summary>
+			public class TestSessionSummary
+			{
+				public string testName { get; set; }
+				public string dateTime { get; set; }
+				public float timeElapseSec { get; set; }
+				public int phasesTotalCount { get; set; }
+				public int phasesSucceededCount { get; set; }
+				public int phasesUndefinedCount { get; set; }
+				public int phasesFailedCount { get; set; }
+
+				public TestSessionSummary(string InName)
 				{
-					Items.Add(NewDataItem);
+					testName = InName;
 				}
-				else
+			}
+
+			[JsonConverter(typeof(JsonTryParseEnumConverter))]
+			public enum TestPhaseOutcome
+			{
+				Unknown,
+				NotRun,
+				Interrupted,
+				Failed,
+				Success,
+				Skipped
+			}
+
+			/// <summary>
+			/// Represent a phase of a test
+			/// </summary>
+			public class TestPhase
+			{
+				public string key { get; set; }
+				public string name { get; set; }
+				public string dateTime { get; set; }
+				public float timeElapseSec { get; set; }
+				public TestPhaseOutcome outcome { get; set; }
+				public bool hasError { get; set; }
+				public bool hasWarning { get; set; }
+				public List<string> deviceKeys { get; set; }
+				public string eventStreamPath { get; set; }
+
+				private TestSession session { get; set; }
+
+				private TestEventStream eventStream { get; set; }
+
+				public TestPhase(string InName, string InKey = null, TestSession InSession = null)
 				{
-					Items[FoundItemIndex] = NewDataItem;
+					key = string.IsNullOrEmpty(InKey) ? GenerateKey(InName) : ValidateKey(InKey);
+					name = InName;
+					deviceKeys = new List<string>();
+					outcome = TestPhaseOutcome.Unknown;
+					eventStream = new TestEventStream(session?.report, this);
+					AttachToSession(InSession);
 				}
 
-				var ExtraItems = InData.GetReportDependencies();
-				if (ExtraItems.Count() > 0)
+				/// <summary>
+				/// Set timing information about the phase
+				/// </summary>
+				/// <param name="StartTime"></param>
+				/// <param name="TimeElapse"></param>
+				public void SetTiming(DateTime StartTime, float TimeElapse)
 				{
-					foreach (string Key in ExtraItems.Keys)
+					dateTime = StartTime.ToString("s", CultureInfo.InvariantCulture);
+					timeElapseSec = TimeElapse;
+				}
+
+				/// <summary>
+				/// Set the outcome of the phase, it will update the test session phase counts
+				/// </summary>
+				/// <param name="State"></param>
+				public void SetOutcome(TestStateType State)
+				{
+					UpdatePhasesCount(PhaseCount.Out);
+					outcome = TestStateToPhaseOutcome(State);
+					UpdatePhasesCount();
+				}
+
+				/// <summary>
+				/// Set the outcome of the phase using Gauntlet TestResult type, it will update the test session phase counts
+				/// </summary>
+				/// <param name="State"></param>
+				public void SetOutcome(Gauntlet.TestResult State)
+				{
+					UpdatePhasesCount(PhaseCount.Out);
+					outcome = TestResultToPhaseOutcome(State);
+					UpdatePhasesCount();
+				}
+
+				/// <summary>
+				/// Attach the phase to a session, making changes to the phase outcome tracked in that session
+				/// </summary>
+				/// <param name="InSession"></param>
+				public void AttachToSession(TestSession InSession)
+				{
+					session = InSession;
+					eventStream.AttachToReport(session.report);
+					session.phases.Add(this);
+					session.summary.phasesTotalCount += 1;
+					UpdatePhasesCount();
+				}
+
+				private enum PhaseCount
+				{
+					In,
+					Out
+				}
+				/// <summary>
+				/// Update the phases count in the session based on the phase outcome
+				/// </summary>
+				/// <param name="Direction"></param>
+				private void UpdatePhasesCount(PhaseCount Direction = PhaseCount.In)
+				{
+					if (session == null)
+					{
+						return;
+					}
+					int Incr = Direction == PhaseCount.In ? 1 : -1;
+					switch (outcome)
 					{
-						DataItem ExtraDataItem = new DataItem();
-						ExtraDataItem.Key = Key;
-						ExtraDataItem.Data = ExtraItems[Key];
-						Items.Add(ExtraDataItem);
+						case TestPhaseOutcome.Failed:
+							session.summary.phasesFailedCount += Incr;
+							break;
+
+						case TestPhaseOutcome.Success:
+							session.summary.phasesSucceededCount += Incr;
+							break;
+
+						case TestPhaseOutcome.Skipped:
+							// Remove from total count
+							session.summary.phasesTotalCount -= Incr;
+							break;
+
+						default:
+							session.summary.phasesUndefinedCount += Incr;
+							break;
 					}
 				}
 
-				return NewDataItem;
+				/// <summary>
+				/// Get the event stream associated with this phase
+				/// </summary>
+				/// <returns></returns>
+				/// <exception cref="AutomationException"></exception>
+				public TestEventStream GetStream()
+				{
+					return eventStream;
+				}
 			}
 
-			public List<DataItem> Items { get; set; }
-
 			/// <summary>
-			/// Write Test Data Collection to json
+			/// fork of the test session data (summary, phases, devices, eventStreams)
 			/// </summary>
-			/// <param name="OutputTestDataFilePath"></param>
-			public void WriteToJson(string OutputTestDataFilePath, bool bIncrementNameIfFileExists = false)
+			public class TestSession
 			{
-				string OutputTestDataDir = Path.GetDirectoryName(OutputTestDataFilePath);
-				if (!Directory.Exists(OutputTestDataDir))
+				public TestSessionSummary summary { get; set; }
+				public List<TestPhase> phases { get; set; }
+				public Dictionary<string, TestDevice> devices { get; set; }
+
+				[JsonIgnore]
+				// back pointer
+				public AutomatedTestSessionData report { get; protected set; }
+
+				public TestSession(string InName, AutomatedTestSessionData InReport = null)
 				{
-					Directory.CreateDirectory(OutputTestDataDir);
+					summary = new TestSessionSummary(InName);
+					phases = new List<TestPhase>();
+					devices = new Dictionary<string, TestDevice>();
+					report = InReport;
 				}
-				if (File.Exists(OutputTestDataFilePath) && bIncrementNameIfFileExists)
+			}
+
+			/// <summary>
+			/// Set the session timing information
+			/// </summary>
+			/// <param name="StartTime"></param>
+			/// <param name="TimeElapse"></param>
+			public virtual void SetSessionTiming(DateTime StartTime, float TimeElapse)
+			{
+				Data.summary.dateTime = StartTime.ToString("s", CultureInfo.InvariantCulture);
+				Data.summary.timeElapseSec = TimeElapse;
+			}
+
+			/// <summary>
+			/// Add a Device to the session (or replace it) and return the new instance.
+			/// </summary>
+			/// <param name="InName"></param>
+			/// <param name="InKey"></param>
+			/// <returns></returns>
+			public virtual TestDevice AddDevice(string InName, string InKey = null)
+			{
+				string Key = ValidateKey(InKey ?? InName);
+				if (Data.devices.ContainsKey(Key))
 				{
-					// increment filename if file exists
-					string Ext = Path.GetExtension(OutputTestDataFilePath);
-					string Filename = OutputTestDataFilePath.Replace(Ext, "");
-					int Incr = 0;
-					do
-					{
-						Incr++;
-						OutputTestDataFilePath = string.Format("{0}{1}{2}", Filename, Incr, Ext);
-					} while (File.Exists(OutputTestDataFilePath));
+					// remove to replace
+					Data.devices.Remove(Key);
+					Log.Verbose("Device (Key) '{Name}' already exists in the report. It is being replaced.", InName, Key);
 				}
-				// write test pass summary
-				Log.Verbose("Writing Test Data Collection for Horde at {0}", OutputTestDataFilePath);
-				try
+				TestDevice Device = new TestDevice(InName);
+				Data.devices.Add(Key, Device);
+				return Device;
+			}
+
+			/// <summary>
+			/// Return the device associated with the Key
+			/// </summary>
+			/// <param name="Key"></param>
+			/// <returns></returns>
+			public virtual TestDevice GetDevice(string Key)
+			{
+				return Data.devices.GetValueOrDefault(Key);
+			}
+
+			public override bool AttachDeviceLog(string InstanceName, string LogPath, string LogName)
+			{
+				if (AttachArtifact(LogPath, LogName))
 				{
-					File.WriteAllText(OutputTestDataFilePath, JsonSerializer.Serialize(this, GetDefaultJsonOptions()));
+					var Device = Data.devices.GetValueOrDefault(InstanceName);
+					if (Device != null)
+					{
+						string LogPathForHorde = Path.GetRelativePath(Globals.UnrealRootDir, Path.Combine(OutputArtifactPath, LogName));
+						Device.appInstanceLogPath = FileUtils.ConvertPathToUri(LogPathForHorde);
+						return true;
+					}
 				}
-				catch (Exception Ex)
+
+				return false;
+			}
+
+			/// <summary>
+			/// Add a phase to the test session
+			/// </summary>
+			/// <param name="InName"></param>
+			/// <param name="InKey"></param>
+			/// <returns></returns>
+			public virtual TestPhase AddPhase(string InName, string InKey = null)
+			{
+				TestPhase Phase = new TestPhase(InName, InKey, Data);
+				CurrentPhase = Phase;
+				return Phase;
+			}
+
+			/// <summary>
+			/// Get the phase associated with Key
+			/// </summary>
+			/// <param name="Key"></param>
+			/// <returns></returns>
+			public virtual TestPhase GetPhase(string Key)
+			{
+				return Data.phases.FirstOrDefault(P => P.key == Key);
+			}
+
+			/// <summary>
+			/// Set the current session phase with the phase associated with Key 
+			/// </summary>
+			/// <param name="Key"></param>
+			/// <exception cref="AutomationException"></exception>
+			public virtual void SetCurrentPhase(string Key)
+			{
+				TestPhase Phase = GetPhase(Key);
+				if (Phase == null)
 				{
-					Log.Error("Failed to save Test Data Collection for Horde. {0}", Ex);
+					throw new AutomationException("Could not find Phase '{0}'", Key); 
 				}
+				CurrentPhase = Phase;
 			}
-		}
-		private static JsonSerializerOptions GetDefaultJsonOptions()
-		{
-			if (CommandUtils.IsBuildMachine)
+
+			/// <summary>
+			/// Set the current session phase with the input phase (and attach that phase the session with not already)
+			/// </summary>
+			/// <param name="InPhase"></param>
+			public virtual void SetCurrentPhase(TestPhase InPhase)
 			{
-				return new JsonSerializerOptions { };
+				InPhase.AttachToSession(Data);
+				CurrentPhase = InPhase;
 			}
-			// Pretty print when running locally
-			return new JsonSerializerOptions
+
+			/// <summary>
+			/// Attach the input phase to the test session
+			/// </summary>
+			/// <param name="InPhase"></param>
+			public virtual void AttachPhase(TestPhase InPhase)
 			{
-				WriteIndented = true
-			};
+				InPhase.AttachToSession(Data);
+			}
+
+			/// <summary>
+			/// Get the current session phase 
+			/// </summary>
+			/// <returns></returns>
+			public virtual TestPhase GetCurrentPhase()
+			{
+				return CurrentPhase;
+			}
+
+			/// <summary>
+			/// Represent the event stream of a phase
+			/// </summary>
+			public class TestEventStream
+			{
+				private AutomatedTestSessionData report { get; set; }
+				private TestPhase phase { get; set; }
+				private List<LogEvent> stream { get; set; }
+
+				public TestEventStream(AutomatedTestSessionData InReport, TestPhase InPhase)
+				{
+					report = InReport;
+					phase = InPhase;
+					stream = new List<LogEvent>();
+				}
+
+				/// <summary>
+				/// Return true if there is any LogEvent in this stream
+				/// </summary>
+				/// <returns></returns>
+				public bool HasAnyEvent()
+				{
+					return stream.Any();
+				}
+
+				/// <summary>
+				/// Get the list of LogEvents
+				/// </summary>
+				/// <returns></returns>
+				public IEnumerable<LogEvent> GetEvents()
+				{
+					return stream;
+				}
+
+				/// <summary>
+				/// Attach the event stream to a report
+				/// </summary>
+				/// <param name="InReport"></param>
+				public virtual void AttachToReport(AutomatedTestSessionData InReport)
+				{
+					report = InReport;
+				}
+
+				/// <summary>
+				/// Attached the input LogEvent to the stream
+				/// </summary>
+				/// <param name="Event"></param>
+				public virtual void AddEvent(LogEvent Event)
+				{
+					stream.Add(Event);
+					switch (Event.Level)
+					{
+						case Logging.LogLevel.Error:
+							phase.hasError = true;
+							break;
+
+						case Logging.LogLevel.Warning:
+							phase.hasWarning = true;
+							break;
+					}
+				}
+
+				/// <summary>
+				/// Add event with optional formatting inputs
+				/// </summary>
+				/// <param name="Time"></param>
+				/// <param name="Level"></param>
+				/// <param name="Message"></param>
+				/// <param name="Args"></param>
+				public virtual void AddEvent(DateTime Time, Logging.LogLevel Level, string Message, params object[] Args)
+				{
+					string Format = null;
+					Dictionary<string, object> Properties = null;
+					if (Args.Any())
+					{
+						Format = Message;
+						Properties = new Dictionary<string, object>();
+						MessageTemplate.ParsePropertyValues(Format, Args, Properties);
+						Message = MessageTemplate.Render(Format, Properties);
+					}
+					Logging.EventId EventId = Level == Logging.LogLevel.Critical? KnownLogEvents.Gauntlet_FatalEvent : KnownLogEvents.Gauntlet_TestEvent;
+					LogEvent Event = new LogEvent(Time, Level, EventId, Message, Format, Properties, null);
+					AddEvent(Event);
+				}
+
+				/// <summary>
+				/// Add event with optional formatting inputs
+				/// </summary>
+				/// <param name="Level"></param>
+				/// <param name="Message"></param>
+				/// <param name="Args"></param>
+				public virtual void AddEvent(Logging.LogLevel Level, string Message, params object[] Args)
+				{
+					AddEvent(DateTime.UtcNow, Level, Message, Args);
+				}
+
+				/// <summary>
+				/// Add Info event with optional formatting inputs
+				/// </summary>
+				/// <param name="Format"></param>
+				/// <param name="Args"></param>
+				public virtual void AddInfo(string Format, params object[] Args)
+				{
+					AddEvent(Logging.LogLevel.Information, Format, Args);
+				}
+
+				/// <summary>
+				/// Add Warning event with optional formatting inputs
+				/// </summary>
+				/// <param name="Format"></param>
+				/// <param name="Args"></param>
+				public virtual void AddWarning(string Format, params object[] Args)
+				{
+					AddEvent(Logging.LogLevel.Warning, Format, Args);
+				}
+
+				/// <summary>
+				/// Add Error event with optional formatting inputs
+				/// </summary>
+				/// <param name="Format"></param>
+				/// <param name="Args"></param>
+				public virtual void AddError(string Format, params object[] Args)
+				{
+					AddEvent(Logging.LogLevel.Error, Format, Args);
+				}
+
+				/// <summary>
+				/// Add event using ITestEvent as input
+				/// </summary>
+				/// <param name="Event"></param>
+				public virtual void AddEvent(ITestEvent Event)
+				{
+					string Message = Event.Summary;
+					Logging.LogLevel Level = EventSeverityToLogLevel(Event.Severity);
+					List<object> Params = new List<object>();
+					if (Event.Details.Any())
+					{
+						Message += "\n{Details}";
+						Params.Add(string.Join("\n", Event.Details));
+					}
+					if (Event.Callstack.Any())
+					{
+						Message += "\n{Callstack}";
+						Params.Add(string.Join("\n", Event.Callstack));
+					}
+					AddEvent(Event.Time, Level, Message, Args: Params.ToArray());
+				}
+
+				/// <summary>
+				/// Add Artifact event
+				/// </summary>
+				/// <param name="Time"></param>
+				/// <param name="Level"></param>
+				/// <param name="Format"></param>
+				/// <param name="Artifact"></param>
+				public virtual void AddArtifact(DateTime Time, Logging.LogLevel Level, string Format, ILogArtifact Artifact)
+				{
+					if (report == null)
+					{
+						AddWarning("Phase not attached to a report. Failed to attached Artifact {Name}.", Artifact.ToString());
+					}
+					else
+					{
+						Artifact.AttachTo(report);
+					}
+
+					var Properties = Artifact.ToProperties();
+					string Message = MessageTemplate.Render(Format, Properties);
+
+					LogEvent Event = new LogEvent(Time, Level, KnownLogEvents.Gauntlet_TestEvent, Message, Format, Properties, null);
+					AddEvent(Event);
+				}
+
+				/// <summary>
+				/// Add Image Comparison event
+				/// </summary>
+				/// <param name="Time"></param>
+				/// <param name="Level"></param>
+				/// <param name="Format"></param>
+				/// <param name="Artifact"></param>
+				public virtual void AddImageComparison(DateTime Time, Logging.LogLevel Level, string Format, ImageComparisonFiles Artifact)
+				{
+					bool IsError = Level == Logging.LogLevel.Error;
+					if (!IsError)
+					{
+						// Don't push references if it is not a failure
+						Artifact.Files[TestLogValueTypes.approved] = null;
+						Artifact.Files[TestLogValueTypes.approvedMetadata] = null;
+					}
+
+					AddArtifact(Time, Level, Format, Artifact);
+					if (IsError)
+					{
+						AddInfo("The image reference can be updated by pointing the Screen Comparison tab from the Test Automation window to the artifacts from this test.");
+					}
+				}
+			}
+
+			/// <summary>
+			/// Attached the input LogEvent to the stream to current event stream
+			/// </summary>
+			/// <param name="Event"></param>
+			public virtual void AddEvent(LogEvent Event)
+			{
+				GetEventStream().AddEvent(Event);
+			}
+
+			/// <summary>
+			/// Add event with optional formatting inputs to current event stream
+			/// </summary>
+			/// <param name="Time"></param>
+			/// <param name="Level"></param>
+			/// <param name="Message"></param>
+			/// <param name="Args"></param>
+			public virtual void AddEvent(DateTime Time, Logging.LogLevel Level, string Message, params object[] Args)
+			{
+				GetEventStream().AddEvent(Time, Level, Message, Args);
+			}
+
+			/// <summary>
+			/// Add event with optional formatting inputs to current event stream
+			/// </summary>
+			/// <param name="Type"></param>
+			/// <param name="Message"></param>
+			/// <param name="Args"></param>
+			public override void AddEvent(EventType Type, string Message, params object[] Args)
+			{
+				AddEvent(DateTime.UtcNow, EventTypeToLogLevel(Type), Message, Args);
+			}
+
+			/// <summary>
+			/// Attach Artifact to the report
+			/// </summary>
+			/// <param name="FilePath"></param>
+			/// <param name="Name"></param>
+			/// <returns></returns>
+			protected virtual bool AttachArtifact(ref string FilePath, string Name = null)
+			{
+				if (AttachArtifact(FilePath, Name))
+				{
+					FilePath = Path.Combine(OutputArtifactPath, Name ?? Path.GetFileName(FilePath));
+					FilePath = FileUtils.ConvertPathToUri(Path.GetRelativePath(Globals.UnrealRootDir, FilePath));
+					return true;
+				}
+				FilePath = null;
+				return false;
+			}
+
+			/// <summary>
+			/// Attach artifact to the report, warn if it fails to do it
+			/// </summary>
+			/// <param name="FilePath"></param>
+			/// <param name="Name"></param>
+			/// <returns></returns>
+			protected virtual bool AttachArtifactOrWarn(ref string FilePath, string Name = null)
+			{
+				if (!AttachArtifact(ref FilePath, Name))
+				{
+					AddWarning("Failed to attached Artifact {FilePath}.", Path.GetRelativePath(Globals.UnrealRootDir, FilePath));
+					return false;
+				}
+				return true;
+			}
+
+			private static class TestLogValueTypes
+			{
+				// Log entry
+				public static readonly Utf8String Type = new Utf8String("$type");
+				public static readonly Utf8String Text = new Utf8String("$text");
+				// Image compare
+				public static readonly string imageCompare = "Image Compare";
+				public static readonly string unapproved = "unapproved";
+				public static readonly string unapprovedMetadata = "unapproved_metadata";
+				public static readonly string approved = "approved";
+				public static readonly string approvedMetadata = "approved_metadata";
+				public static readonly string difference = "difference";
+				public static readonly string differenceReport = "difference_report";
+			}
+
+			/// <summary>
+			/// Represent a Log value. Necessary to support LogValue annotation in TestData.
+			/// </summary>
+			public class TestLogValue
+			{
+				public Utf8String Type { get; protected set; }
+				public string Text { get; protected set; }
+				public Dictionary<Utf8String, object> Properties { get; protected set; }
+
+				public TestLogValue(string InType, string InText, Dictionary<Utf8String, object> InProperties)
+				{
+					Type = new Utf8String(InType);
+					Text = InText;
+					Properties = InProperties;
+				}
+
+				public override string ToString()
+				{
+					return Text;
+				}
+			}
+
+			/// <summary>
+			/// TestLogValue Formatter
+			/// </summary>
+			public class TestLogValueFormatter : ILogValueFormatter
+			{
+				public void Format(object value, Utf8JsonWriter writer)
+				{
+					TestLogValue entry = (TestLogValue)value;
+					writer.WriteStartObject();
+					writer.WriteString(TestLogValueTypes.Type, entry.Type);
+					writer.WriteString(TestLogValueTypes.Text, entry.ToString());
+					if (entry.Properties != null)
+					{
+						foreach (KeyValuePair<Utf8String, object> pair in entry.Properties)
+						{
+							writer.WritePropertyName(pair.Key);
+							writer.WriteStringValue(pair.Value.ToString());
+						}
+					}
+					writer.WriteEndObject();
+				}
+			}
+
+			/// <summary>
+			/// Interface for artifact in Log event
+			/// </summary>
+			public interface ILogArtifact
+			{
+				public bool AttachTo(AutomatedTestSessionData Report);
+				Dictionary<string, object> ToProperties();
+			}
+
+			/// <summary>
+			/// Generic Artifact Files utility objects to attached files to report and generate event properties
+			/// </summary>
+			public class ArtifactFiles : ILogArtifact
+			{
+				public string Type { get; protected set; }
+				public string Token { get; protected set; }
+				public string Name { get; protected set; }
+				public Dictionary<string, string> Files { get; protected set; }
+				public string RootPath { get; protected set; }
+
+				public ArtifactFiles(string InType, string InName, Dictionary<string, string> InFiles, string InRootPath)
+				{
+					Type = InType;
+					Token = Type.ToLower().Replace(" ", "_");
+					Name = InName;
+					Files = InFiles;
+					RootPath = InRootPath;
+				}
+
+				private bool AttachFileTo(ref string FilePath, AutomatedTestSessionData Report)
+				{
+					if (string.IsNullOrEmpty(FilePath))
+					{
+						// Ignore this file
+						return true;
+					}
+					string Name = null;
+					if (!Path.IsPathFullyQualified(FilePath))
+					{
+						Name = FilePath;
+						if (!string.IsNullOrEmpty(RootPath))
+							FilePath = Path.Combine(RootPath, FilePath);
+					}
+					return Report.AttachArtifactOrWarn(ref FilePath, Name);
+				}
+
+				public bool AttachTo(AutomatedTestSessionData Report)
+				{
+					bool AllAttached = true;
+					foreach (string Key in Files.Keys.ToList())
+					{
+						string FilePath = Files[Key];
+						if (!AttachFileTo(ref FilePath, Report))
+						{
+							AllAttached = AllAttached && false;
+						}
+						Files[Key] = FilePath;
+					}
+
+					return AllAttached;
+				}
+
+				public Dictionary<string, object> ToProperties()
+				{
+					Dictionary<Utf8String, object> Prop = new Dictionary<Utf8String, object>();
+					foreach(var Item in Files)
+					{
+						if (!string.IsNullOrEmpty(Item.Value))
+						{
+							Prop.Add(new Utf8String(Item.Key.ToLower()), Item.Value);
+						}
+					}
+					TestLogValue Value = new TestLogValue(Type, Name, Prop);
+					Dictionary<string, object> EventProperties = new Dictionary<string, object> { { Token, Value } };
+
+					return EventProperties;
+				}
+			}
+
+			/// <summary>
+			/// Utility object to attached ImageComparison artifacts and generate related event properties
+			/// </summary>
+			public class ImageComparisonFiles : ArtifactFiles, ILogArtifact
+			{
+				public ImageComparisonFiles(string InName, Dictionary<string, string> InFiles, string InRootPath) : base(TestLogValueTypes.imageCompare, InName, InFiles, InRootPath)
+				{
+					string FilePath = Files.GetValueOrDefault(TestLogValueTypes.unapproved);
+					if (!string.IsNullOrEmpty(FilePath))
+					{
+						Files.Add(TestLogValueTypes.unapprovedMetadata, Path.ChangeExtension(FilePath, ".json"));
+					}
+					FilePath = Files.GetValueOrDefault(TestLogValueTypes.approved);
+					if (!string.IsNullOrEmpty(FilePath))
+					{
+						Files.Add(TestLogValueTypes.approvedMetadata, Path.ChangeExtension(FilePath, ".json"));
+					}
+					FilePath = Files.GetValueOrDefault(TestLogValueTypes.difference);
+					if (!string.IsNullOrEmpty(FilePath))
+					{
+						Files.Add(TestLogValueTypes.differenceReport, Path.Combine(Path.GetDirectoryName(FilePath), "Report.json"));
+					}
+				}
+			}
+
+			private static Logging.LogLevel EventTypeToLogLevel(EventType Type)
+			{
+				Logging.LogLevel Level = Logging.LogLevel.None;
+				switch (Type)
+				{
+					case EventType.Info:
+						Level = Logging.LogLevel.Information;
+						break;
+
+					case EventType.Warning:
+						Level = Logging.LogLevel.Warning;
+						break;
+
+					case EventType.Error:
+						Level = Logging.LogLevel.Error;
+						break;
+				}
+
+				return Level;
+			}
+
+			private static Logging.LogLevel EventSeverityToLogLevel(EventSeverity Type)
+			{
+				Logging.LogLevel Level = Logging.LogLevel.None;
+				switch (Type)
+				{
+					case EventSeverity.Info:
+						Level = Logging.LogLevel.Information;
+						break;
+
+					case EventSeverity.Warning:
+						Level = Logging.LogLevel.Warning;
+						break;
+
+					case EventSeverity.Error:
+						Level = Logging.LogLevel.Error;
+						break;
+
+					case EventSeverity.Fatal:
+						Level = Logging.LogLevel.Critical;
+						break;
+				}
+
+				return Level;
+			}
+
+			private static TestPhaseOutcome TestStateToPhaseOutcome(TestStateType State)
+			{
+				TestPhaseOutcome Outcome = TestPhaseOutcome.Unknown;
+				switch (State)
+				{
+					case TestStateType.Fail:
+						Outcome = TestPhaseOutcome.Failed;
+						break;
+
+					case TestStateType.InProcess:
+						Outcome = TestPhaseOutcome.Interrupted;
+						break;
+
+					case TestStateType.NotRun:
+						Outcome = TestPhaseOutcome.NotRun;
+						break;
+
+					case TestStateType.Success:
+						Outcome = TestPhaseOutcome.Success;
+						break;
+
+					case TestStateType.Skipped:
+						Outcome = TestPhaseOutcome.Skipped;
+						break;
+				}
+				return Outcome;
+			}
+
+			private static TestPhaseOutcome TestResultToPhaseOutcome(Gauntlet.TestResult State)
+			{
+				TestPhaseOutcome Outcome = TestPhaseOutcome.Unknown;
+				switch (State)
+				{
+					case Gauntlet.TestResult.Failed:
+						Outcome = TestPhaseOutcome.Failed;
+						break;
+
+					case Gauntlet.TestResult.TimedOut:
+					case Gauntlet.TestResult.Cancelled:
+						Outcome = TestPhaseOutcome.Interrupted;
+						break;
+
+					case Gauntlet.TestResult.Invalid:
+						Outcome = TestPhaseOutcome.NotRun;
+						break;
+
+					case Gauntlet.TestResult.Passed:
+						Outcome = TestPhaseOutcome.Success;
+						break;
+
+					case Gauntlet.TestResult.InsufficientDevices:
+						Outcome = TestPhaseOutcome.Skipped;
+						break;
+				}
+				return Outcome;
+			}
+
+			/// <summary>
+			/// Return the event stream of the target phase or current phase if not specified
+			/// </summary>
+			/// <param name="InPhaseKey"></param>
+			/// <returns></returns>
+			/// <exception cref="AutomationException"></exception>
+			public virtual TestEventStream GetEventStream(string InPhaseKey = null)
+			{
+				if (InPhaseKey == null && CurrentPhase == null)
+				{
+					AddPhase("Main");
+				}
+				string PhaseKey = string.IsNullOrEmpty(InPhaseKey) ? CurrentPhase.key : InPhaseKey;
+				TestPhase Phase = GetPhase(PhaseKey);
+				if (Phase == null)
+				{
+					throw new AutomationException("No Phase '{0}' set in Test '{1}'", PhaseKey, Key);
+				}
+				return Phase.GetStream();
+			}
+
+			public override Dictionary<string, object> GetReportDependencies()
+			{
+				// store the event streams
+				string OutputDir = Path.Combine(OutputArtifactPath, "EventStreams");
+				if (!Directory.Exists(OutputDir))
+				{
+					Directory.CreateDirectory(OutputDir);
+				}
+				Log.Verbose("Writing Phase Event Streams at {0}", OutputDir);
+				foreach (TestPhase Phase in Data.phases)
+				{
+					string OutputFilePath = null;
+					TestEventStream Stream = Phase.GetStream();
+					// skip if the event stream is empty
+					if (Stream.HasAnyEvent())
+					{
+						string FileName = $"EventStreams/{Phase.key}.json";
+						OutputFilePath = Path.Combine(OutputArtifactPath, FileName);
+						// write phase event stream
+						try
+						{
+							File.WriteAllText(OutputFilePath, JsonSerializer.Serialize(Stream.GetEvents(), GetDefaultJsonOptions()));
+							AttachArtifact(ref OutputFilePath, FileName);
+						}
+						catch (Exception Ex)
+						{
+							Log.Error("Failed to save Event Stream for '{0}'. {1}", Phase.key, Ex);
+						}
+					}
+					// update phase
+					Phase.eventStreamPath = OutputFilePath;
+				}
+
+				return base.GetReportDependencies();
+			}
+
+			private static DateTime ConvertUETimeStringToDateTime(string StringTime)
+			{
+				return UnrealAutomationEntry.GetTimestampAsDateTime(StringTime);
+			}
+
+			/// <summary>
+			/// Convert UnrealAutomatedTestPassResults to Horde data model
+			/// </summary>
+			/// <param name="InTestPassResults"></param>
+			/// <param name="InName"></param>
+			/// <param name="InReportPath"></param>
+			/// <param name="InHordeArtifactPath"></param>
+			/// <returns></returns>
+			public static AutomatedTestSessionData FromUnrealAutomatedTests(UnrealAutomatedTestPassResults InTestPassResults, string InKey, string InName, string InReportPath, string InHordeArtifactPath)
+			{
+				AutomatedTestSessionData OutTestPassResults = new AutomatedTestSessionData(InKey, InName);
+				OutTestPassResults.SetOutputArtifactPath(InHordeArtifactPath);
+				if (InTestPassResults.Devices != null)
+				{
+					foreach (UnrealAutomationDevice InDevice in InTestPassResults.Devices)
+					{
+						TestDevice ConvertedDevice = OutTestPassResults.AddDevice(InDevice.DeviceName, InDevice.Instance);
+						ConvertedDevice.appInstanceLogPath = InDevice.AppInstanceLog;
+						ConvertedDevice.metadata.Add("platform", InDevice.Platform);
+						ConvertedDevice.metadata.Add("os_version", InDevice.OSVersion);
+						ConvertedDevice.metadata.Add("model", InDevice.Model);
+						ConvertedDevice.metadata.Add("gpu", InDevice.GPU);
+						ConvertedDevice.metadata.Add("cpumodel", InDevice.CPUModel);
+						ConvertedDevice.metadata.Add("ram_in_gb", InDevice.RAMInGB.ToString());
+						ConvertedDevice.metadata.Add("render_mode", InDevice.RenderMode);
+						ConvertedDevice.metadata.Add("rhi", InDevice.RHI);
+					}
+				}
+				DateTime SessionTime = ConvertUETimeStringToDateTime(InTestPassResults.ReportCreatedOn);
+				OutTestPassResults.SetSessionTiming(SessionTime, InTestPassResults.TotalDuration);
+				if (InTestPassResults.Tests != null)
+				{
+					foreach (UnrealAutomatedTestResult InTestResult in InTestPassResults.Tests)
+					{
+						TestPhase Phase = OutTestPassResults.AddPhase(InTestResult.FullTestPath);
+						Phase.deviceKeys = InTestResult.DeviceInstance;
+						string TestDateTime = InTestResult.DateTime;
+						DateTime PhaseTime = ConvertUETimeStringToDateTime(TestDateTime);
+						Phase.SetTiming(PhaseTime, InTestResult.Duration);
+						Phase.SetOutcome(InTestResult.State);
+						TestEventStream Stream = Phase.GetStream();
+						foreach (UnrealAutomationEntry InTestEntry in InTestResult.Entries)
+						{
+							string Artifact = InTestEntry.Event.Artifact;
+							// If Artifact values is not null nor a bunch on 0, then we have a file attachment.
+							if (!string.IsNullOrEmpty(Artifact) && Artifact.Substring(0, 4) != "0000")
+							{
+								UnrealAutomationArtifact ArtifactEntry = InTestResult.Artifacts.Find(A => A.Id == Artifact);
+								if (ArtifactEntry != null)
+								{
+									DateTime EventTime = ConvertUETimeStringToDateTime(InTestEntry.Timestamp);
+									Logging.LogLevel LogLevel = EventTypeToLogLevel(InTestEntry.Event.Type);
+									string EntryName = ArtifactEntry.Name;
+									if (ArtifactEntry.Type == "Comparison")
+									{
+										ImageComparisonFiles ComparisonFiles = new ImageComparisonFiles(EntryName, ArtifactEntry.Files, InReportPath);
+										string Format = InTestEntry.Event.Message.Replace(EntryName, $"{{{ComparisonFiles.Token}}}");
+										Stream.AddImageComparison(EventTime, LogLevel, Format, ComparisonFiles);
+
+									}
+									else
+									{
+										ArtifactFiles Files = new ArtifactFiles(ArtifactEntry.Type, EntryName, ArtifactEntry.Files, InReportPath);
+										string Format = InTestEntry.Event.Message.Replace(EntryName, $"{{{Files.Token}}}");
+										Stream.AddArtifact(EventTime, LogLevel, Format, Files);
+									}
+									continue;
+								}
+							}
+							Stream.AddEvent(InTestEntry.AsLogEvent());
+						}
+					}
+				}
+				return OutTestPassResults;
+			}
+		}
+
+		/// <summary>
+		/// Container for Test Data items 
+		/// </summary>
+		public class TestDataCollection
+		{
+			public class DataItem
+			{
+				public string key { get; set; }
+				public object data { get; set; }
+				public int version { get; set; } = 1;
+				public Dictionary<string, string> metadata { get; set; }
+			}
+			public TestDataCollection()
+			{
+				items = new List<DataItem>();
+			}
+
+			public DataItem AddNewTestReport(object InData, string InKey = null)
+			{
+				int Version = 1;
+				Dictionary<string, string> Metadata = null;
+				Dictionary<string, object> ExtraItems = null;
+				if (InData is BaseHordeReport InHordeReport)
+				{
+					ExtraItems = InHordeReport.GetReportDependencies();
+					InKey = InHordeReport.GetTestDataKey(InKey);
+					InData = InHordeReport.GetTestData();
+					Version = InHordeReport.GetVersion();
+					Metadata = InHordeReport.GetMetadata();
+				}
+				DataItem NewDataItem = new DataItem()
+				{
+					key = InKey,
+					data = InData,
+					version = Version,
+					metadata = Metadata
+				};
+
+				var FoundItemIndex = items.FindIndex(I => I.key == InKey);
+				if (FoundItemIndex == -1)
+				{
+					items.Add(NewDataItem);
+				}
+				else
+				{
+					items[FoundItemIndex] = NewDataItem;
+				}
+
+				if (ExtraItems != null && ExtraItems.Count() > 0)
+				{
+					foreach (string Key in ExtraItems.Keys)
+					{
+						DataItem ExtraDataItem = new DataItem()
+						{
+							key = Key,
+							data = ExtraItems[Key]
+						};
+						items.Add(ExtraDataItem);
+					}
+				}
+
+				return NewDataItem;
+			}
+
+			public List<DataItem> items { get; set; }
+
+			/// <summary>
+			/// Write Test Data Collection to json
+			/// </summary>
+			/// <param name="OutputTestDataFilePath"></param>
+			public void WriteToJson(string OutputTestDataFilePath, bool bIncrementNameIfFileExists = false)
+			{
+				string OutputTestDataDir = Path.GetDirectoryName(OutputTestDataFilePath);
+				if (!Directory.Exists(OutputTestDataDir))
+				{
+					Directory.CreateDirectory(OutputTestDataDir);
+				}
+				if (File.Exists(OutputTestDataFilePath) && bIncrementNameIfFileExists)
+				{
+					// increment filename if file exists
+					string Ext = Path.GetExtension(OutputTestDataFilePath);
+					string Filename = OutputTestDataFilePath.Replace(Ext, "");
+					int Incr = 0;
+					do
+					{
+						Incr++;
+						OutputTestDataFilePath = string.Format("{0}{1}{2}", Filename, Incr, Ext);
+					} while (File.Exists(OutputTestDataFilePath));
+				}
+				// write test pass summary
+				Log.Verbose("Writing Test Data Collection for Horde at {0}", OutputTestDataFilePath);
+				try
+				{
+					File.WriteAllText(OutputTestDataFilePath, JsonSerializer.Serialize(this, GetDefaultJsonOptions()));
+				}
+				catch (Exception Ex)
+				{
+					Log.Error("Failed to save Test Data Collection for Horde. {0}", Ex);
+				}
+			}
+		}
+		private static JsonSerializerOptions GetDefaultJsonOptions()
+		{
+			if (CommandUtils.IsBuildMachine)
+			{
+				return new JsonSerializerOptions { };
+			}
+			// Pretty print when running locally
+			return new JsonSerializerOptions
+			{
+				WriteIndented = true
+			};
+		}
+
+		/// <summary>
+		/// Validate Key name, intended to validate a key that is going to be used as property name (ie: Phase key).
+		/// </summary>
+		/// <param name="InKey"></param>
+		/// <returns></returns>
+		/// <exception cref="AutomationException"></exception>
+		private static string ValidateKey(string InKey)
+		{
+			if (string.IsNullOrEmpty(InKey))
+			{
+				throw new AutomationException("Key must not be empty.");
+			}
+			if (InKey.StartsWith("$"))
+			{
+				throw new AutomationException("Key must not start with the '$' character.");
+			}
+			if (InKey.Contains("."))
+			{
+				throw new AutomationException("Key must not contain '.' characters.");
+			}
+			if (InKey.Length > 1024)
+			{
+				throw new AutomationException("Key must not have more than 1024 characters.");
+			}
+			return InKey;
 		}
 	}
 }
diff --git a/Engine/Source/Programs/AutomationTool/Gauntlet/Unreal/Automation/UE.Automation.cs b/Engine/Source/Programs/AutomationTool/Gauntlet/Unreal/Automation/UE.Automation.cs
--- a/Engine/Source/Programs/AutomationTool/Gauntlet/Unreal/Automation/UE.Automation.cs
+++ b/Engine/Source/Programs/AutomationTool/Gauntlet/Unreal/Automation/UE.Automation.cs
@@ -529,7 +529,7 @@ namespace UE
 
 			if (GetConfiguration() is AutomationTestConfig Config)
 			{
-				if (Config.ResumeOnCriticalFailure && string.IsNullOrEmpty(Config.ReportExportPath))
+				if ((Config.WriteTestResultsForHorde || Config.ResumeOnCriticalFailure) && string.IsNullOrEmpty(Config.ReportExportPath))
 				{
 					Config.ReportExportPath = Path.Combine(Globals.TempDir, "TestReport");
 				}
@@ -653,10 +653,10 @@ namespace UE
 							switch (Entry.Level)
 							{
 								case UnrealLog.LogLevel.Error:
-									Events.Add(new UnrealAutomationEvent(EventType.Error, Entry.Message));
+									TestNodeEvents.Add(new UnrealTestEvent(EventSeverity.Error, Entry.Message, Enumerable.Empty<string>()));
 									break;
 								case UnrealLog.LogLevel.Warning:
-									Events.Add(new UnrealAutomationEvent(EventType.Warning, Entry.Message));
+									TestNodeEvents.Add(new UnrealTestEvent(EventSeverity.Warning, Entry.Message, Enumerable.Empty<string>()));
 									break;
 							}
 						}
@@ -828,12 +828,11 @@ namespace UE
 							TempReport.SetOutputArtifactPath(HordeArtifactPath);
 							foreach (UnrealRoleArtifacts Artifact in SessionArtifacts)
 							{
-								string LogName = Path.GetFullPath(Artifact.LogPath).Replace(Path.GetFullPath(Context.Options.LogDir), "").TrimStart(Path.DirectorySeparatorChar);
-								TempReport.AttachArtifact(Artifact.LogPath, LogName);
-								// Reference last run instance log
-								if (Artifact.SessionRole.RoleType == MainRole.Type)
+								string LogName = Path.GetRelativePath(Path.GetFullPath(Context.Options.LogDir), Path.GetFullPath(Artifact.LogPath));
+								if (TempReport.AttachArtifact(Artifact.LogPath, LogName) && Artifact.SessionRole.RoleType == MainRole.Type)
 								{
-									JsonTestPassResults.Devices.Last().AppInstanceLog = LogName.Replace("\\", "/");
+									// Reference last run instance log
+									JsonTestPassResults.Devices.Last().AppInstanceLog = Gauntlet.FileUtils.ConvertPathToUri(Path.GetRelativePath(Globals.UnrealRootDir, Path.Combine(HordeArtifactPath, LogName)));
 								}
 							}
 							JsonTestPassResults.WriteToJson();
@@ -1074,19 +1073,19 @@ namespace UE
 					{
 						Log.Error(KnownLogEvents.Gauntlet_TestEvent, " * No tests were executed.");
 
-						IEnumerable<UnrealAutomationEvent> Errors = Events.Where(E => E.IsError).Distinct();
-						IEnumerable<UnrealAutomationEvent> Warnings = Events.Where(E => E.IsWarning).Distinct();
+						IEnumerable<UnrealTestEvent> Errors = TestNodeEvents.Where(E => E.IsError).Distinct();
+						IEnumerable<UnrealTestEvent> Warnings = TestNodeEvents.Where(E => E.IsWarning).Distinct();
 						if (Errors.Any() || Warnings.Any())
 						{
 							Log.Info("   See log above for details.");
 							foreach (var Error in Errors.Take(MaxEventsDisplayPerTest))
 							{
-								Log.Error(KnownLogEvents.Gauntlet_UnrealEngineTestEvent, "    " + Error.FormatToString());
+								Log.Error(KnownLogEvents.Gauntlet_UnrealEngineTestEvent, "    " + Error.Summary);
 							}
 							NotifyMoreIfNeeded(Errors);
 							foreach (var Warning in Warnings.Take(MaxEventsDisplayPerTest))
 							{
-								Log.Warning(KnownLogEvents.Gauntlet_UnrealEngineTestEvent, "    " + Warning.FormatToString());
+								Log.Warning(KnownLogEvents.Gauntlet_UnrealEngineTestEvent, "    " + Warning.Summary);
 							}
 							NotifyMoreIfNeeded(Warnings);
 							Log.Info("");
diff --git a/Engine/Source/Programs/AutomationTool/Gauntlet/Unreal/Base/Gauntlet.UnrealTestNode.cs b/Engine/Source/Programs/AutomationTool/Gauntlet/Unreal/Base/Gauntlet.UnrealTestNode.cs
--- a/Engine/Source/Programs/AutomationTool/Gauntlet/Unreal/Base/Gauntlet.UnrealTestNode.cs
+++ b/Engine/Source/Programs/AutomationTool/Gauntlet/Unreal/Base/Gauntlet.UnrealTestNode.cs
@@ -10,6 +10,7 @@ using System.Linq;
 using EpicGames.Core;
 using Microsoft.Extensions.Logging;
 using AutomationUtils.Matchers;
+using static Gauntlet.HordeReport;
 
 namespace Gauntlet
 {
@@ -78,11 +79,11 @@ namespace Gauntlet
 		/// </summary>
 		public override IEnumerable<string> GetWarnings()
 		{
-			IEnumerable<string> WarningList = Events.Where(E => E.IsWarning).Select(E => E.Message);
+			IEnumerable<string> WarningList = TestNodeEvents.Where(E => E.IsWarning).Select(E => E.Summary);
 			
 			if (RoleResults != null)
 			{
-				WarningList = WarningList.Union(RoleResults.SelectMany(R => R.Events.Where(E => E.Severity == EventSeverity.Warning)).Select(E => E.Summary));
+				WarningList = WarningList.Union(RoleResults.SelectMany(R => R.Events.Where(E => E.IsWarning)).Select(E => E.Summary));
 			}
 
 			return WarningList.ToArray();
@@ -94,11 +95,11 @@ namespace Gauntlet
 		/// </summary>
 		public override IEnumerable<string> GetErrors()
 		{
-			IEnumerable<string> ErrorList = Events.Where(E => E.IsError).Select(E => E.Message);
+			IEnumerable<string> ErrorList = TestNodeEvents.Where(E => E.IsError).Select(E => E.Summary);
 			
 			if (RoleResults != null)
 			{
-				ErrorList = ErrorList.Union(RoleResults.SelectMany(R => R.Events.Where(E => E.Severity == EventSeverity.Error || E.Severity == EventSeverity.Fatal)).Select(E => E.Summary));
+				ErrorList = ErrorList.Union(RoleResults.SelectMany(R => R.Events.Where(E => E.IsError)).Select(E => E.Summary));
 			}
 
 			return ErrorList.ToArray();			
@@ -120,14 +121,6 @@ namespace Gauntlet
 			));
 		}
 
-		/// <summary>
-		/// Returns the test URL Link
-		/// </summary>
-		public virtual string GetURLLink()
-		{
-			return "";
-		}
-
 		/// <summary>
 		/// Report an error
 		/// </summary>
@@ -135,7 +128,7 @@ namespace Gauntlet
 		public virtual void ReportError(string Message, params object[] Args)
 		{
 			Message = string.Format(Message, Args);
-			Events.Add(new UnrealAutomationEvent(EventType.Error, Message));
+			TestNodeEvents.Add(new UnrealTestEvent(EventSeverity.Error, Message, Enumerable.Empty<string>()));
 			if (!LogWarningsAndErrorsAfterSummary) { Log.Error(KnownLogEvents.Gauntlet_TestEvent, Message); }
 			if (GetTestStatus() == TestStatus.Complete && GetTestResult() == TestResult.Passed)
 			{
@@ -150,7 +143,7 @@ namespace Gauntlet
 		public virtual void ReportWarning(string Message, params object[] Args)
 		{
 			Message = string.Format(Message, Args);
-			Events.Add(new UnrealAutomationEvent(EventType.Warning, Message));
+			TestNodeEvents.Add(new UnrealTestEvent(EventSeverity.Warning, Message, Enumerable.Empty<string>()));
 			if (!LogWarningsAndErrorsAfterSummary) { Log.Warning(KnownLogEvents.Gauntlet_TestEvent, Message); }
 		}
 
@@ -231,12 +224,6 @@ namespace Gauntlet
 		/// </summary>
 		public IEnumerable<UnrealRoleArtifacts> SessionArtifacts { get; private set; }
 
-		/// <summary>
-		/// Error and warning collection.
-		/// </summary>
-		protected List<UnrealAutomationEvent> Events { get; private set; } = new List<UnrealAutomationEvent>();
-
-
 		/// <summary>
 		/// Collection of events thrown by gauntlet itself during the test run.
 		/// </summary>
@@ -1106,13 +1093,18 @@ namespace Gauntlet
 
 			// Create results from all roles from these artifacts
 			RoleResults = CreateRoleResultsFromArtifacts(InReason, SessionArtifacts);
+			// Update Test node result from Role processed results when they don't match
+			if (GetTestResult() == TestResult.Passed && RoleResults.Any(R => R.ProcessResult != UnrealProcessResult.ExitOk))
+			{
+				SetTestResult(TestResult.Failed);
+			}
 
 			string Message = string.Empty;
 			ITestReport Report = null;
 
 			try
 			{
-				// Check if the deprecated signature is overriden, call it anyway if that the case.
+				// Check if the deprecated signature is overridden, call it anyway if that the case.
 				var OldSignature = new[] { typeof(TestResult), typeof(UnrealTestContext), typeof(UnrealBuildSource), typeof(IEnumerable<UnrealRoleResult>), typeof(string) };
 				var Simplifiedignature = new[] { typeof(TestResult) };
 				if (!Utils.InterfaceHelpers.HasOverriddenMethod(this.GetType(), "CreateReport", Simplifiedignature) && Utils.InterfaceHelpers.HasOverriddenMethod(this.GetType(), "CreateReport", OldSignature))
@@ -1282,8 +1274,7 @@ namespace Gauntlet
 			if (GetCachedConfiguration().WriteTestResultsForHorde)
 			{
 				// write test report for Horde
-				HordeReport.SimpleTestReport HordeTestReport = CreateSimpleReportForHorde(Result);
-				return HordeTestReport;
+				return CreateSimpleReportForHorde(Result);
 			}
 
 			return null;
@@ -1299,13 +1290,33 @@ namespace Gauntlet
 			if (GetCachedConfiguration().WriteTestResultsForHorde)
 			{
 				// write test report for Horde
-				HordeReport.SimpleTestReport HordeTestReport = CreateSimpleReportForHorde(Result);
-				return HordeTestReport;
+				return CreateSimpleReportForHorde(Result);
 			}
 
 			return null;
 		}
 
+		/// <summary>
+		/// Unique key to identify and cross reference test name across test sessions 
+		/// </summary>
+		protected virtual string HordeReportTestKey
+		{
+			get
+			{
+				if (string.IsNullOrEmpty(GetCachedConfiguration().HordeTestDataKey))
+				{
+					if (Globals.Params.ParseParam("UseTestDataV2"))
+					{
+						return Type;
+					}
+
+					return Name + " " + GetMainRoleContextString();
+				}
+
+				return GetCachedConfiguration().HordeTestDataKey;
+			}
+		}
+
 		/// <summary>
 		/// Optional override for the test report name, useful when a single test node has different testing modes
 		/// </summary>
@@ -1326,59 +1337,102 @@ namespace Gauntlet
 		/// Generate a Simple Test Report from the results of this test
 		/// </summary>
 		/// <param name="Result"></param>
-		protected virtual HordeReport.SimpleTestReport CreateSimpleReportForHorde(TestResult Result)
+		protected virtual ITestReport CreateSimpleReportForHorde(TestResult Result)
 		{
-			if (string.IsNullOrEmpty(GetCachedConfiguration().HordeTestDataKey))
-			{
-				GetCachedConfiguration().HordeTestDataKey = Name + " " + GetMainRoleContextString();
-			}
-
-			HordeReport.SimpleTestReport HordeTestReport = new HordeReport.SimpleTestReport(UnrealTestResult, Context, GetCachedConfiguration());
-			HordeTestReport.TestName = HordeReportTestName;
-			HordeTestReport.ReportCreatedOn = DateTime.Now.ToString();
-			HordeTestReport.TotalDurationSeconds = (float) (DateTime.Now - SessionStartTime).TotalSeconds;
-			HordeTestReport.Description = GetMainRoleContextString();
-			HordeTestReport.URLLink = GetURLLink();
-			HordeTestReport.Errors.AddRange(GetErrorsAndAbnornalExits());
-			if (!string.IsNullOrEmpty(CancellationReason))
+			BaseHordeReport HordeTestReport = null;
+			IEnumerable<UnrealTestRole> RoleList = null;
+			if (Globals.Params.ParseParam("UseTestDataV2"))
 			{
-				HordeTestReport.Errors.Add(CancellationReason);
+				RoleList = new UnrealTestRole[] { GetCachedConfiguration().GetMainRequiredRole() };
+				AutomatedTestSessionData TestReport = new AutomatedTestSessionData(HordeReportTestKey, HordeReportTestName);
+				DateTime SessionTime = SessionStartTime.ToUniversalTime();
+				float TimeElapse = (float)(DateTime.Now - SessionStartTime).TotalSeconds;
+				TestReport.SetSessionTiming(SessionTime, TimeElapse);
+				// Simple report has one phase
+				var Phase = TestReport.AddPhase("Main");
+				Phase.SetTiming(SessionTime, TimeElapse);
+				// Propagate events
+				var Stream = Phase.GetStream();
+				foreach(UnrealTestEvent Event in TestNodeEvents)
+				{
+					Stream.AddEvent(Event);
+				}
+				if (RoleResults != null)
+				{
+					foreach(UnrealRoleResult Item in RoleResults)
+					{
+						string RoleType = Item.Artifacts.SessionRole.RoleType.ToString();
+						foreach (UnrealTestEvent Event in Item.Events)
+						{
+							Stream.AddEvent(Event);
+						}
+						// Add Processed result
+						if (Item.ProcessResult != UnrealProcessResult.ExitOk)
+						{
+							Stream.AddError($"{RoleType}: {Item.Summary}");
+						}
+						// Since we are going through the Role, lets add the related device
+						ITargetDevice TargetDevice = Item.Artifacts.AppInstance.Device;
+						string DeviceName = TargetDevice.Platform == BuildHostPlatform.Current.Platform? System.Environment.MachineName : TargetDevice.Name;
+						string DeviceKey = TargetDevice.Name.Replace(".", "-"); // Key must not contains dot
+						var ReportDevice = TestReport.AddDevice(DeviceName, DeviceKey);
+						ReportDevice.metadata.Add("Platform", TargetDevice.Platform.ToString());
+						ReportDevice.metadata.Add("Role", RoleType);
+						Phase.deviceKeys.Add(DeviceKey);
+					}
+				}
+				// Outcome
+				Phase.SetOutcome(GetTestResult());
+				HordeTestReport = TestReport;
 			}
-			HordeTestReport.Warnings.AddRange(GetWarnings());
-			HordeTestReport.HasSucceeded = !(Result == TestResult.Failed || Result == TestResult.TimedOut || HordeTestReport.Errors.Count > 0);
-			if (HordeTestReport.Errors.Count > 0 && Result == TestResult.Passed)
+			else
 			{
-				SetUnrealTestResult(TestResult.Failed);
+				RoleList = GetCachedConfiguration().RequiredRoles.Values.SelectMany(V => V);
+				SimpleTestReport SimpleReport = new SimpleTestReport(UnrealTestResult, Context, GetCachedConfiguration());
+				SimpleReport.TestName = HordeReportTestName;
+				SimpleReport.ReportCreatedOn = DateTime.UtcNow.ToString();
+				SimpleReport.TotalDurationSeconds = (float)(DateTime.Now - SessionStartTime).TotalSeconds;
+				SimpleReport.Description = GetMainRoleContextString();
+				SimpleReport.Errors.AddRange(GetErrorsAndAbnornalExits());
+				if (!string.IsNullOrEmpty(CancellationReason))
+				{
+					SimpleReport.Errors.Add(CancellationReason);
+				}
+				SimpleReport.Warnings.AddRange(GetWarnings());
+				SimpleReport.HasSucceeded = !(Result == TestResult.Failed || Result == TestResult.TimedOut || SimpleReport.Errors.Count > 0);
+				if (SimpleReport.Errors.Count > 0 && Result == TestResult.Passed)
+				{
+					SetUnrealTestResult(TestResult.Failed);
+				}
+				SimpleReport.Status = GetTestResult().ToString();
+				HordeTestReport = SimpleReport;
 			}
-			HordeTestReport.Status = GetTestResult().ToString();
-			string HordeArtifactPath = string.IsNullOrEmpty(GetCachedConfiguration().HordeArtifactPath) ? HordeReport.DefaultArtifactsDir : GetCachedConfiguration().HordeArtifactPath;
-			HordeTestReport.SetOutputArtifactPath(HordeArtifactPath);
+			string HordeArtifactPath = string.IsNullOrEmpty(GetCachedConfiguration().HordeArtifactPath) ? DefaultArtifactsDir : GetCachedConfiguration().HordeArtifactPath;
+			string ArtifactPathLabel = new DirectoryInfo(ArtifactPath).Name;
+			HordeTestReport.SetOutputArtifactPath(Path.Combine(HordeArtifactPath, ArtifactPathLabel));
 			if (SessionArtifacts != null)
 			{
 				foreach (UnrealRoleResult RoleResult in RoleResults)
 				{
-					string LogName = string.Empty;
 					string LogPath = RoleResult.Artifacts.LogPath;
 					if(!string.IsNullOrEmpty(LogPath) && File.Exists(LogPath))
 					{
-						LogName = Path.GetFullPath(LogPath).Replace(Path.GetFullPath(Context.Options.LogDir), "").TrimStart(Path.DirectorySeparatorChar);
-						HordeTestReport.AttachArtifact(RoleResult.Artifacts.LogPath, LogName);
-					}
+						string LogName = FileUtils.ConvertPathToUri(Path.GetRelativePath(Path.GetFullPath(ArtifactPath), Path.GetFullPath(LogPath)));
+						if (HordeTestReport is BaseHordeReport Report)
+						{
+							string DeviceKey = RoleResult.Artifacts.AppInstance.Device.Name.Replace(".", "-");
+							Report.AttachDeviceLog(DeviceKey, LogPath, LogName);
+						}
+						else
+						{
+							HordeTestReport.AttachArtifact(LogPath, LogName);
+						}
 
-					UnrealLog LogSummary = RoleResult.LogSummary;
-					if (LogSummary.Errors.Count() > 0)
-					{
-						HordeTestReport.Warnings.Add(
-							string.Format(
-								"Log Parsing: FatalErrors={0}, Ensures={1}, Errors={2}, Warnings={3}, Log={4}",
-								(LogSummary.FatalError != null ? 1 : 0), LogSummary.Ensures.Count(), LogSummary.Errors.Count(), LogSummary.Warnings.Count(), LogName
-							)
-						);
 					}
 				}
 			}
 			// Metadata
-			SetReportMetadata(HordeTestReport, GetCachedConfiguration().RequiredRoles.Values.SelectMany(V => V));
+			SetReportMetadata(HordeTestReport, RoleList);
 
 			return HordeTestReport;
 		}
@@ -1398,35 +1452,51 @@ namespace Gauntlet
 				string HordeArtifactPath = GetCachedConfiguration().HordeArtifactPath;
 				Log.Verbose("Reading json Unreal Automated test report from {Path}", JsonReportPath);
 				UnrealAutomatedTestPassResults JsonTestPassResults = UnrealAutomatedTestPassResults.LoadFromJson(JsonReportPath);
+				// With UE Test Automation, we care only for one role.
 				var MainRole = GetCachedConfiguration().GetMainRequiredRole();
-				string HordeTestDataKey = string.IsNullOrEmpty(GetCachedConfiguration().HordeTestDataKey) ? Name + " " + GetMainRoleContextString() : GetCachedConfiguration().HordeTestDataKey;
-				GetCachedConfiguration().HordeTestDataKey = HordeTestDataKey;
+				var RoleList = new List<UnrealTestRole>() { MainRole };
 				// Convert test results for Horde
-				HordeReport.AutomatedTestSessionData HordeTestPassResults = HordeReport.AutomatedTestSessionData.FromUnrealAutomatedTests(
-					JsonTestPassResults, Type, Suite, UnrealAutomatedTestReportPath, HordeArtifactPath
-				);
-				// Make a copy of the report in the old way - until we decide the transition is over
-				HordeReport.UnrealEngineTestPassResults CopyTestPassResults = HordeReport.UnrealEngineTestPassResults.FromUnrealAutomatedTests(JsonTestPassResults, ReportURL);
-				CopyTestPassResults.CopyTestResultsArtifacts(UnrealAutomatedTestReportPath, HordeArtifactPath);
-				HordeTestPassResults.AttachDependencyReport(CopyTestPassResults, HordeTestDataKey);
-				// Pre Flight information
-				HordeTestPassResults.PreFlightChange = GetCachedConfiguration().PreFlightChange;
+				ITestReport HordeTestPassResults = null;
+				if (Globals.Params.ParseParam("UseTestDataV2"))
+				{
+					string ArtifactPathLabel = new DirectoryInfo(ArtifactPath).Name;
+					HordeArtifactPath = Path.Combine(HordeArtifactPath, ArtifactPathLabel);
+					HordeTestPassResults = AutomatedTestSessionData.FromUnrealAutomatedTests(
+						JsonTestPassResults, HordeReportTestKey, HordeReportTestName, UnrealAutomatedTestReportPath, HordeArtifactPath
+					);
+				}
+				else
+				{
+					AutomatedTestSessionData_legacy LegacyTestPassResults = AutomatedTestSessionData_legacy.FromUnrealAutomatedTests(
+						JsonTestPassResults, Type, Suite, UnrealAutomatedTestReportPath, HordeArtifactPath
+					);
+					// Make a copy of the report in the old way - until we decide the transition is over
+					UnrealEngineTestPassResults CopyTestPassResults = UnrealEngineTestPassResults.FromUnrealAutomatedTests(JsonTestPassResults, ReportURL);
+					CopyTestPassResults.CopyTestResultsArtifacts(UnrealAutomatedTestReportPath, HordeArtifactPath);
+					SetReportMetadata(CopyTestPassResults, RoleList); // Set the metadata to the old report
+					LegacyTestPassResults.AttachDependencyReport(CopyTestPassResults, HordeReportTestKey);
+					// Pre Flight information
+					LegacyTestPassResults.PreFlightChange = GetCachedConfiguration().PreFlightChange;
+					HordeTestPassResults = LegacyTestPassResults;
+				}
 				// Metadata
-				// With UE Test Automation, we care only for one role.
-				var RoleList = new List<UnrealTestRole>() { MainRole };
 				SetReportMetadata(HordeTestPassResults, RoleList);
-				SetReportMetadata(CopyTestPassResults, RoleList); // Set the metadata to the old report too
-				// Attached test Artifacts
+				// Attached test log artifacts
 				if (SessionArtifacts != null)
 				{
 					foreach (UnrealRoleArtifacts Artifact in SessionArtifacts)
 					{
-						string LogName = Path.GetFullPath(Artifact.LogPath).Replace(Path.GetFullPath(Context.Options.LogDir), "").TrimStart(Path.DirectorySeparatorChar);
-						HordeTestPassResults.AttachArtifact(Artifact.LogPath, LogName);
-						// Reference last run instance log
 						if (Artifact.SessionRole.RoleType == MainRole.Type)
 						{
-							HordeTestPassResults.Devices.Last().AppInstanceLog = LogName.Replace("\\", "/");
+							string LogName = FileUtils.ConvertPathToUri(Path.GetRelativePath(Path.GetFullPath(ArtifactPath), Path.GetFullPath(Artifact.LogPath)));
+							if (HordeTestPassResults is BaseHordeReport Report)
+							{
+								Report.AttachDeviceLog(JsonTestPassResults.Devices.Last().Instance, Artifact.LogPath, LogName);
+							}
+							else
+							{
+								HordeTestPassResults.AttachArtifact(Artifact.LogPath, LogName);
+							}
 						}
 					}
 				}
@@ -1475,7 +1545,7 @@ namespace Gauntlet
 				FileName + ".TestData.json"
 			);
 			HordeReport.TestDataCollection HordeTestDataCollection = new HordeReport.TestDataCollection();
-			HordeTestDataCollection.AddNewTestReport(Report, GetCachedConfiguration().HordeTestDataKey);
+			HordeTestDataCollection.AddNewTestReport(Report, HordeReportTestKey);
 			HordeTestDataCollection.WriteToJson(HordeTestDataFilePath, !AutomationTool.Automation.IsBuildMachine);
 		}
 
@@ -1724,7 +1794,7 @@ namespace Gauntlet
 
 			IEnumerable<UnrealTestEvent> EventList = CreateEventListFromArtifact(InReason, InRoleArtifacts, LogSummary, ProcessResult);
 
-			// if the test is stopping for a reason other than completion, mark this as failing incase derived classes
+			// if the test is stopping for a reason other than completion, mark this as failing in case derived classes
 			// don't do the right thing
 			if (InReason == StopReason.MaxDuration)
 			{
diff --git a/Engine/Source/Programs/AutomationTool/Gauntlet/Unreal/Utils/Gauntlet.UnrealTestResults.cs b/Engine/Source/Programs/AutomationTool/Gauntlet/Unreal/Utils/Gauntlet.UnrealTestResults.cs
--- a/Engine/Source/Programs/AutomationTool/Gauntlet/Unreal/Utils/Gauntlet.UnrealTestResults.cs
+++ b/Engine/Source/Programs/AutomationTool/Gauntlet/Unreal/Utils/Gauntlet.UnrealTestResults.cs
@@ -12,7 +12,6 @@ using Microsoft.Extensions.Logging;
 using Logging = Microsoft.Extensions.Logging;
 using System.Globalization;
 using AutomationUtils.Matchers;
-using Polly;
 
 namespace Gauntlet
 {
@@ -43,15 +42,6 @@ namespace Gauntlet
 		[JsonPropertyName("appInstanceLog")]
 		public string AppInstanceLog { get; set; }
 	}
-	public class UnrealAutomationComparisonFiles
-	{
-		[JsonPropertyName("difference")]
-		public string Difference { get; set; }
-		[JsonPropertyName("approved")]
-		public string Approved { get; set; }
-		[JsonPropertyName("unapproved")]
-		public string Unapproved { get; set; }
-	}
 	public class UnrealAutomationArtifact
 	{
 		[JsonPropertyName("id")]
@@ -61,7 +51,7 @@ namespace Gauntlet
 		[JsonPropertyName("type")]
 		public string Type { get; set; }
 		[JsonPropertyName("files")]
-		public UnrealAutomationComparisonFiles Files { get; set; }
+		public Dictionary<string, string> Files { get; set; }
 	}
 	public class UnrealAutomationEvent
 	{
@@ -173,21 +163,26 @@ namespace Gauntlet
 		[JsonPropertyName("timeStamp")]
 		public string Timestamp { get; set; }
 
-		private DateTime GetTimestampAsDateTime()
+		public const string InvalidDateTime = "0001.01.01-00.00.00";
+		public const string DateTimeFormat = "yyyy.MM.dd-HH.mm.ss";
+
+		public static DateTime GetTimestampAsDateTime(string StringTime)
 		{
-			if (string.IsNullOrEmpty(Timestamp) || Timestamp == "0001.01.01-00.00.00")
+			DateTime Time = DateTime.UtcNow;
+			if (!string.IsNullOrEmpty(StringTime) && StringTime != InvalidDateTime)
 			{
-				// Special case: when UE Test DateTime is set this way, it means the value was 0 or null before getting converted to json.
-				Timestamp = DateTime.UtcNow.ToString("yyyy.MM.dd-HH.mm.ss");
+				DateTime.TryParseExact(StringTime, DateTimeFormat, CultureInfo.InvariantCulture, DateTimeStyles.None, out Time);
 			}
-			DateTime Time = DateTime.UtcNow;
-			DateTime.TryParseExact(Timestamp, "yyyy.MM.dd-HH.mm.ss", CultureInfo.InvariantCulture, DateTimeStyles.None, out Time);
 			return Time;
 		}
 
+		/// <summary>
+		/// Produce a LogEvent instance from this UnrealAutomationEntry
+		/// </summary>
+		/// <returns></returns>
 		public LogEvent AsLogEvent()
 		{
-			DateTime Time = GetTimestampAsDateTime();
+			DateTime Time = GetTimestampAsDateTime(Timestamp);
 			Logging.LogLevel Level = Event.IsError ?
 										(Event.IsCriticalFailure ?
 											Logging.LogLevel.Critical : Logging.LogLevel.Error)
@@ -229,10 +224,16 @@ namespace Gauntlet
 					}
 					Properties.Add("SourceFile", Filename);
 					Properties.Add("Line", LineNumber.ToString());
-					Format = (Format ?? Message) + " [{SourceFile}:{Line}]";
+					Format = (Format ?? Message) + " [{SourceFile}({Line})]";
 				}
 			}
 
+			if (Format != null)
+			{
+				// Make sure Message is produced from Format
+				Message = MessageTemplate.Render(Format, Properties);
+			}
+
 			return new LogEvent(Time, Level, EventIdType, Message, Format, Properties, null);
 		}
 
