AutomationTool: Bump target framework to .NET 8

#rnx

diff --git a/Engine/Source/Programs/AutomationTool/AutomationTool.csproj b/Engine/Source/Programs/AutomationTool/AutomationTool.csproj
--- a/Engine/Source/Programs/AutomationTool/AutomationTool.csproj
+++ b/Engine/Source/Programs/AutomationTool/AutomationTool.csproj
@@ -10,13 +10,13 @@
   <Choose>
     <When Condition="'$(IsWindows)'=='true'">
       <PropertyGroup>
-        <TargetFramework>net6.0-windows10.0.19041</TargetFramework>
+        <TargetFramework>net8.0-windows10.0.19041</TargetFramework>
         <UseWindowsForms>true</UseWindowsForms>
       </PropertyGroup>
     </When>
     <Otherwise>
       <PropertyGroup>
-        <TargetFramework>net6.0</TargetFramework>
+        <TargetFramework>net8.0</TargetFramework>
       </PropertyGroup>
     </Otherwise>
   </Choose>
