use Role.IsEditor() to include CookedEditor as valid role to parse its artifacts for errors

#rnx

diff --git a/Engine/Source/Programs/AutomationTool/Gauntlet/Unreal/Automation/UE.Automation.cs b/Engine/Source/Programs/AutomationTool/Gauntlet/Unreal/Automation/UE.Automation.cs
--- a/Engine/Source/Programs/AutomationTool/Gauntlet/Unreal/Automation/UE.Automation.cs
+++ b/Engine/Source/Programs/AutomationTool/Gauntlet/Unreal/Automation/UE.Automation.cs
@@ -910,7 +910,7 @@ namespace UE
 			UnrealProcessResult UnrealResult = base.GetExitCodeAndReason(InReason, InLog, InArtifacts, out ExitReason, out ExitCode);
 
 			// The editor is an additional arbiter of success
-			if (InArtifacts.SessionRole.RoleType == UnrealTargetRole.Editor
+			if (InArtifacts.SessionRole.RoleType.IsEditor()
 				&& InLog.HasAbnormalExit == false)
 			{
 				// if no fatal errors, check test results
@@ -1004,7 +1004,7 @@ namespace UE
 			base.LogTestSummaryHeader();
 
 			// Everything we need is in the editor artifacts
-			var EditorRole = RoleResults.Where(R => R.Artifacts.SessionRole.RoleType == UnrealTargetRole.Editor).FirstOrDefault();
+			var EditorRole = RoleResults.Where(R => R.Artifacts.SessionRole.RoleType.IsEditor()).FirstOrDefault();
 
 			if (EditorRole != null)
 			{
@@ -1227,7 +1227,7 @@ namespace UE
 
 			foreach (var Role in GetRolesThatFailed())
 			{
-				if (Role.Artifacts.SessionRole.RoleType == UnrealTargetRole.Editor)
+				if (Role.Artifacts.SessionRole.RoleType.IsEditor())
 				{
 					AutomationLogParser Parser = new AutomationLogParser(Role.LogSummary.FullLogContent);
 					AllErrors.AddRange(
