[Gauntlet] Move FullClean and SkipInstall static properties to DevicePool and out of UnrealAppConfig
This mainly improves separation of concern. Additionally, desktops should now be able to skipinstall on build machines not using reservation blocks.

diff --git a/Engine/Source/Programs/AutomationTool/Gauntlet/Framework/Devices/Gauntlet.DevicePool.cs b/Engine/Source/Programs/AutomationTool/Gauntlet/Framework/Devices/Gauntlet.DevicePool.cs
--- a/Engine/Source/Programs/AutomationTool/Gauntlet/Framework/Devices/Gauntlet.DevicePool.cs
+++ b/Engine/Source/Programs/AutomationTool/Gauntlet/Framework/Devices/Gauntlet.DevicePool.cs
@@ -361,6 +361,12 @@ namespace Gauntlet
 		/// </summary>
 		public static DevicePool Instance { get; private set; } = new DevicePool();
 
+		public static bool SkipInstall;
+
+		public static bool FullClean;
+
+		public static bool DeviceReservationBlock;
+
 		/// <summary>
 		/// Device reservation service URL
 		/// </summary>
@@ -427,11 +433,6 @@ namespace Gauntlet
 		/// </summary>
 		private string LocalTempDir;
 
-		/// <summary>
-		/// Whether or not Unique temporary directories should be created
-		/// </summary>
-		private bool bUniqueTemps;
-
 		/// <summary>
 		/// The maximum number of problem devices to report to device backend
 		/// This mitigates issues on builders, such as hung local processes, incorrectly reporting problem devices
@@ -451,6 +452,10 @@ namespace Gauntlet
 			{
 				Instance = this;
 			}
+
+			SkipInstall = Globals.Params.ParseParams("SkipInstall", "SkipCopy", "SkipDeploy");
+			FullClean = Globals.Params.ParseParam("FullClean");
+			DeviceReservationBlock = false;
 		}
 
 		#region IDisposable Support
@@ -580,10 +585,9 @@ namespace Gauntlet
 			return Constraints[Device];
 		}
 
-		public void SetLocalOptions(string InLocalTemp, bool InUniqueTemps, string InDeviceURL = "")
+		public void SetLocalOptions(string InLocalTemp, bool InUniqueTemps = false, string InDeviceURL = "")
 		{
 			LocalTempDir = InLocalTemp;
-			bUniqueTemps = InUniqueTemps;
 			DeviceURL = InDeviceURL;
 		}
 
@@ -1015,18 +1019,15 @@ namespace Gauntlet
 
 				if (DeviceReservation.InstallRequired == true)
 				{
-					UnrealAppConfig.ForceSkipInstall = false;
-					UnrealAppConfig.ForceFullClean = true;
+					SkipInstall = false;
+					FullClean = true;
+					DeviceReservationBlock = true;
 				}
 				else if (DeviceReservation.InstallRequired == false)
 				{
-					UnrealAppConfig.ForceSkipInstall = true;
-					UnrealAppConfig.ForceFullClean = false;
-				}
-				else
-				{
-					UnrealAppConfig.ForceSkipInstall = null;
-					UnrealAppConfig.ForceFullClean = null;
+					SkipInstall = true;
+					FullClean = false;
+					DeviceReservationBlock = true;
 				}
 
 				// Add target devices from reservation
@@ -1464,10 +1465,7 @@ namespace Gauntlet
 			// When using device reservation blocks, we don't want to fully clean the cache and lose previously installed builds.
 			// If bRetainBuilds evaluates to true, it means we are in the second step or beyond in a device reservation block.
 			// In this case we'll just delete the left over UserDir which should already have been emptied by UnrealSession.
-			bool? bForceClean = UnrealAppConfig.ForceFullClean;
-			bool? bSkipInstall = UnrealAppConfig.ForceSkipInstall;
-			bool bUsingReservationBlock = bForceClean.HasValue && bSkipInstall.HasValue;
-			bool bRetainBuilds = bUsingReservationBlock && !bForceClean.Value && bSkipInstall.Value;
+			bool bRetainBuilds = SkipInstall && !FullClean;
 
 			if(bRetainBuilds)
 			{
@@ -1504,7 +1502,7 @@ namespace Gauntlet
 					{
 						// If we fail to acquire the default client cache while using device reservation blocks,
 						// we can't ensure future tests will have their cache directories mapped to the correct build location
-						if(bUsingReservationBlock)
+						if (DeviceReservationBlock)
 						{
 							throw new AutomationException("Failed to clean default client device cache {0}. {1}", ClientCache, Ex);
 						}
diff --git a/Engine/Source/Programs/AutomationTool/Gauntlet/Unreal/Base/Gauntlet.UnrealAppConfig.cs b/Engine/Source/Programs/AutomationTool/Gauntlet/Unreal/Base/Gauntlet.UnrealAppConfig.cs
--- a/Engine/Source/Programs/AutomationTool/Gauntlet/Unreal/Base/Gauntlet.UnrealAppConfig.cs
+++ b/Engine/Source/Programs/AutomationTool/Gauntlet/Unreal/Base/Gauntlet.UnrealAppConfig.cs
@@ -104,22 +104,11 @@ namespace Gauntlet
 		public OverlayExecutable OverlayExecutable{ get; set; }
 
 		// Prevents installing a build on device
-		public bool SkipInstall => ForceSkipInstall.HasValue ? ForceSkipInstall.Value : _SkipInstall;
+		public bool SkipInstall => DevicePool.SkipInstall;
 
 		// Performs a full clean on the device before installing
-		public bool FullClean => ForceFullClean.HasValue ? ForceFullClean.Value : _FullClean;
+		public bool FullClean => DevicePool.FullClean;
 
-		// Force a full clean
-		public static bool? ForceFullClean = null;
-
-		// Force a skip install
-		public static bool? ForceSkipInstall = null;
-
-		[AutoParamWithNames(false, "SkipInstall", "SkipDeploy", "SkipCopy")]
-		private bool _SkipInstall { get; set; }
-
-		[AutoParamWithNames(false, "FullClean")]
-		private bool _FullClean { get; set; }
 
 		/// <summary>
 		/// Constructor that sets some required values to defaults
@@ -131,7 +120,6 @@ namespace Gauntlet
 			CommandLine = "";
 			Configuration = UnrealTargetConfiguration.Development;
 			Sandbox = "Gauntlet";
-			AutoParam.ApplyParamsAndDefaults(this, Globals.Params.AllArguments);
 		}
 	}
 }
\ No newline at end of file
