[Gauntlet] Fix CIS with params function

diff --git a/Engine/Source/Programs/AutomationTool/Gauntlet/Framework/Gauntlet.Params.cs b/Engine/Source/Programs/AutomationTool/Gauntlet/Framework/Gauntlet.Params.cs
--- a/Engine/Source/Programs/AutomationTool/Gauntlet/Framework/Gauntlet.Params.cs
+++ b/Engine/Source/Programs/AutomationTool/Gauntlet/Framework/Gauntlet.Params.cs
@@ -40,6 +40,25 @@ namespace Gauntlet
 			return false;
 		}
 
+		/// <summary>
+		/// Parses the argument list for a set of parameters and returns whether any are defined or not.
+		/// Useful for cases where you want one setting to have multiple possible command line triggers
+		/// </summary>
+		/// <param name="Params"></param>
+		/// <returns></returns>
+		public bool ParseParams(params string[] Params)
+		{
+			foreach(string Param in Params)
+			{
+				if(ParseParam(Param))
+				{
+					return true;
+				}
+			}
+
+			return false;
+		}
+
 		/// <summary>
 		/// Parses multiple values from a param. The multiple values can either be specified as repeated
 		/// arguments, e.g. -foo=one -foo=two, or when 'CommaSeparated' is true as a comma-separated 
