[Gauntlet] Horde Test Report Moving version and metadata properties into data blob property

#rnx

diff --git a/Engine/Source/Programs/AutomationTool/Gauntlet/Framework/Base/Gauntlet.TestReport.cs b/Engine/Source/Programs/AutomationTool/Gauntlet/Framework/Base/Gauntlet.TestReport.cs
--- a/Engine/Source/Programs/AutomationTool/Gauntlet/Framework/Base/Gauntlet.TestReport.cs
+++ b/Engine/Source/Programs/AutomationTool/Gauntlet/Framework/Base/Gauntlet.TestReport.cs
@@ -139,12 +139,18 @@ namespace Gauntlet
 	{
 		public BaseTestReport()
 		{
-			Metadata = new Dictionary<string, string>();
+			metadata = new Dictionary<string, string>();
 		}
+
+		/// <summary>
+		/// Horde report version
+		/// </summary>
+		public abstract int version { get; }
+
 		/// <summary>
 		/// Metadata blackboard
 		/// </summary>
-		public Dictionary<string, string> Metadata { get; set; }
+		public Dictionary<string, string> metadata { get; private set; }
 
 		/// <summary>
 		/// Return report type
@@ -176,7 +182,7 @@ namespace Gauntlet
 		/// <returns></returns>
 		public virtual void SetMetadata(string Key, string Value)
 		{
-			Metadata[Key] = Value;
+			metadata[Key] = Value;
 		}
 
 		/// <summary>
diff --git a/Engine/Source/Programs/AutomationTool/Gauntlet/Framework/Utils/Gauntlet.HordeReport.cs b/Engine/Source/Programs/AutomationTool/Gauntlet/Framework/Utils/Gauntlet.HordeReport.cs
--- a/Engine/Source/Programs/AutomationTool/Gauntlet/Framework/Utils/Gauntlet.HordeReport.cs
+++ b/Engine/Source/Programs/AutomationTool/Gauntlet/Framework/Utils/Gauntlet.HordeReport.cs
@@ -13,6 +13,7 @@ using Microsoft.CodeAnalysis;
 using EpicGames.Core;
 using Logging = Microsoft.Extensions.Logging;
 using System.Globalization;
+using static Gauntlet.HordeReport.TestDataCollection;
 
 namespace Gauntlet
 {
@@ -39,13 +40,12 @@ namespace Gauntlet
 			}
 		}
 
-		public abstract class BaseHordeReport : BaseTestReport
+		public abstract class BaseHordeReport : BaseTestReport, IDataBlob
 		{
-
 			/// <summary>
 			/// Horde report version
 			/// </summary>
-			public virtual int Version => 1;
+			public override int version => 1;
 
 			protected string OutputArtifactPath;
 			protected HashSet<string> ArtifactProcessedHashes;
@@ -181,28 +181,10 @@ namespace Gauntlet
 			/// Return the data to be stored in the TestData
 			/// </summary>
 			/// <returns></returns>
-			public virtual object GetTestData()
+			public virtual IDataBlob GetTestData()
 			{
 				return this;
 			}
-
-			/// <summary>
-			/// Return the version of the TestData
-			/// </summary>
-			/// <returns></returns>
-			public virtual int GetVersion()
-			{
-				return Version;
-			}
-
-			/// <summary>
-			/// Return the metadata key/value set as Dictionary
-			/// </summary>
-			/// <returns></returns>
-			public virtual Dictionary<string, string> GetMetadata()
-			{
-				return Metadata;
-			}
 		}
 
 		#region Legacy Implementations
@@ -1197,7 +1179,7 @@ namespace Gauntlet
 		{
 			public override string Type => "Automated Test Session";
 
-			public override int Version => 2;
+			public override int version => 2;
 
 			/// <summary>
 			/// Unique key that identify the test and group the different sessions
@@ -1214,7 +1196,7 @@ namespace Gauntlet
 			/// </summary>
 			protected TestSession Data { get; set; }
 
-			public override object GetTestData()
+			public override IDataBlob GetTestData()
 			{
 				return Data;
 			}
@@ -1422,11 +1404,13 @@ namespace Gauntlet
 			/// <summary>
 			/// fork of the test session data (summary, phases, devices, eventStreams)
 			/// </summary>
-			public class TestSession
+			public class TestSession : IDataBlob
 			{
 				public TestSessionSummary summary { get; set; }
 				public List<TestPhase> phases { get; set; }
 				public Dictionary<string, TestDevice> devices { get; set; }
+				public int version { get { return report.version; } }
+				public Dictionary<string, string> metadata { get { return report.metadata; } }
 
 				[JsonIgnore]
 				// back pointer
@@ -2240,12 +2224,15 @@ namespace Gauntlet
 		/// </summary>
 		public class TestDataCollection
 		{
+			public interface IDataBlob
+			{
+				int version { get; }
+				Dictionary<string, string> metadata { get; }
+			}
 			public class DataItem
 			{
 				public string key { get; set; }
 				public object data { get; set; }
-				public int version { get; set; } = 1;
-				public Dictionary<string, string> metadata { get; set; }
 			}
 			public TestDataCollection()
 			{
@@ -2254,23 +2241,17 @@ namespace Gauntlet
 
 			public DataItem AddNewTestReport(object InData, string InKey = null)
 			{
-				int Version = 1;
-				Dictionary<string, string> Metadata = null;
 				Dictionary<string, object> ExtraItems = null;
 				if (InData is BaseHordeReport InHordeReport)
 				{
 					ExtraItems = InHordeReport.GetReportDependencies();
 					InKey = InHordeReport.GetTestDataKey(InKey);
 					InData = InHordeReport.GetTestData();
-					Version = InHordeReport.GetVersion();
-					Metadata = InHordeReport.GetMetadata();
 				}
 				DataItem NewDataItem = new DataItem()
 				{
 					key = InKey,
 					data = InData,
-					version = Version,
-					metadata = Metadata
 				};
 
 				var FoundItemIndex = items.FindIndex(I => I.key == InKey);
