[Gauntlet] Adds IBuildValidtor type to support a user defined way of selecting the "Latest Known Good" (LKG) build
To use this, simply define a type that Implement IBuildValidator and override GetLatestGoodBuild() to return a string to the path of the latest known good build
Then, run a gauntlet test where -build="LKG" OR -build="LatestGood"

diff --git a/Engine/Source/Programs/AutomationTool/Gauntlet/Framework/Base/Gauntlet.BuildValidator.cs b/Engine/Source/Programs/AutomationTool/Gauntlet/Framework/Base/Gauntlet.BuildValidator.cs
new file mode 100644
--- /dev/null
+++ b/Engine/Source/Programs/AutomationTool/Gauntlet/Framework/Base/Gauntlet.BuildValidator.cs
@@ -0,0 +1,19 @@
+// Copyright Epic Games, Inc. All Rights Reserved.
+
+using System;
+using System.Collections.Generic;
+using UnrealBuildTool;
+
+namespace Gauntlet
+{
+	/// <summary>
+	/// Utility for returning the path to a "Latest Good" build
+	/// This is used to determine build paths when -build=LKG or -build=LatestGood
+	/// </summary>
+	public interface IBuildValidator
+	{
+		string Name { get; }
+		bool CanSupportProject(string InProjectName);
+		string GetLatestGoodBuild();
+	}
+}
diff --git a/Engine/Source/Programs/AutomationTool/Gauntlet/Unreal/BuildSource/Gauntlet.UnrealBuildSource.cs b/Engine/Source/Programs/AutomationTool/Gauntlet/Unreal/BuildSource/Gauntlet.UnrealBuildSource.cs
--- a/Engine/Source/Programs/AutomationTool/Gauntlet/Unreal/BuildSource/Gauntlet.UnrealBuildSource.cs
+++ b/Engine/Source/Programs/AutomationTool/Gauntlet/Unreal/BuildSource/Gauntlet.UnrealBuildSource.cs
@@ -203,6 +203,34 @@ namespace Gauntlet
 				string BinariesPath = Path.Combine(ProjectPath.Directory.FullName, "Binaries");
 				OutBuildPaths = Directory.Exists(BinariesPath) ? new string[] { StagedPath, BinariesPath } : new string[] { StagedPath };
 			}
+			else if (BuildDir.Name.Equals("LatestGood", StringComparison.OrdinalIgnoreCase)
+				|| BuildDir.Name.Equals("LKG", StringComparison.OrdinalIgnoreCase))
+			{
+				string RequestedValidator = Globals.Params.ParseValue("BuildValidator", null);
+				IBuildValidator Validator = Utils.InterfaceHelpers.FindImplementations<IBuildValidator>(true)
+					.Where(Validator => Validator.CanSupportProject(ProjectName))
+					.Where(Validator => string.IsNullOrEmpty(RequestedValidator) || Validator.Name.Equals(RequestedValidator, StringComparison.OrdinalIgnoreCase))
+					.FirstOrDefault();
+
+				if(Validator == null)
+				{
+					Log.Error("No build validator that can support project {ProjectName} was found.", ProjectName);
+					return false;
+				}
+
+				string LatestGoodBuild = Validator.GetLatestGoodBuild();
+				if(string.IsNullOrEmpty(LatestGoodBuild))
+				{
+					Log.Error("No latest good build was able to be found!");
+					return false;
+				}
+
+				OutBuildPaths = new[] { LatestGoodBuild };
+				OutBuildName = Path.GetFileName(LatestGoodBuild);
+
+				Log.Info("{Validator} selected {Build} as the latest good build. Proceeding with this build", Validator.GetType().Name, OutBuildName);
+				return true;
+			}
 			else
 			{
 				// todo - make this more generic
