[Gauntlet] Revert change of case for Metadata property in legacy Horde reports

#rnx

diff --git a/Engine/Source/Programs/AutomationTool/Gauntlet/Framework/Base/Gauntlet.TestReport.cs b/Engine/Source/Programs/AutomationTool/Gauntlet/Framework/Base/Gauntlet.TestReport.cs
--- a/Engine/Source/Programs/AutomationTool/Gauntlet/Framework/Base/Gauntlet.TestReport.cs
+++ b/Engine/Source/Programs/AutomationTool/Gauntlet/Framework/Base/Gauntlet.TestReport.cs
@@ -150,7 +150,7 @@ namespace Gauntlet
 		/// <summary>
 		/// Metadata blackboard
 		/// </summary>
-		public Dictionary<string, string> metadata { get; private set; }
+		public virtual Dictionary<string, string> metadata { get; protected set; }
 
 		/// <summary>
 		/// Return report type
diff --git a/Engine/Source/Programs/AutomationTool/Gauntlet/Framework/Utils/Gauntlet.HordeReport.cs b/Engine/Source/Programs/AutomationTool/Gauntlet/Framework/Utils/Gauntlet.HordeReport.cs
--- a/Engine/Source/Programs/AutomationTool/Gauntlet/Framework/Utils/Gauntlet.HordeReport.cs
+++ b/Engine/Source/Programs/AutomationTool/Gauntlet/Framework/Utils/Gauntlet.HordeReport.cs
@@ -369,6 +369,9 @@ namespace Gauntlet
 				get { return "Unreal Automated Tests"; }
 			}
 
+			[JsonPropertyName("Metadata")]
+			public override Dictionary<string, string> metadata { get; protected set; }
+
 			public UnrealEngineTestPassResults() : base()
 			{
 				Devices = new List<Device>();
@@ -552,6 +555,9 @@ namespace Gauntlet
 				get { return "Automated Test Session"; }
 			}
 
+			[JsonPropertyName("Metadata")]
+			public override Dictionary<string, string> metadata { get; protected set; }
+
 			public class TestResult
 			{
 				public string Name { get; set; }
@@ -1096,6 +1102,9 @@ namespace Gauntlet
 				get { return "Simple Report"; }
 			}
 
+			[JsonPropertyName("Metadata")]
+			public override Dictionary<string, string> metadata { get; protected set; }
+
 			public SimpleTestReport() : base()
 			{
 
@@ -1897,7 +1906,7 @@ namespace Gauntlet
 					Type = InType;
 					Token = Type.ToLower().Replace(" ", "_");
 					Name = InName;
-					Files = InFiles;
+					Files = new(InFiles);
 					RootPath = InRootPath;
 				}
 
