[Gauntlet] Allow generating v1 and v2 TestData at the same time

#rnx

diff --git a/Engine/Source/Programs/AutomationTool/Gauntlet/Unreal/Base/Gauntlet.UnrealTestNode.cs b/Engine/Source/Programs/AutomationTool/Gauntlet/Unreal/Base/Gauntlet.UnrealTestNode.cs
--- a/Engine/Source/Programs/AutomationTool/Gauntlet/Unreal/Base/Gauntlet.UnrealTestNode.cs
+++ b/Engine/Source/Programs/AutomationTool/Gauntlet/Unreal/Base/Gauntlet.UnrealTestNode.cs
@@ -1303,18 +1303,18 @@ namespace Gauntlet
 		{
 			get
 			{
-				if (string.IsNullOrEmpty(GetCachedConfiguration().HordeTestDataKey))
-				{
-					if (Globals.Params.ParseParam("UseTestDataV2"))
-					{
-						return Type;
-					}
-
-					return Name + " " + GetMainRoleContextString();
-				}
+				return Type;
+			}
+		}
 
-				return GetCachedConfiguration().HordeTestDataKey;
+		private string GetHordeTestDataKey_legacy()
+		{
+			if (string.IsNullOrEmpty(GetCachedConfiguration().HordeTestDataKey))
+			{
+				return Name + " " + GetMainRoleContextString();
 			}
+
+			return GetCachedConfiguration().HordeTestDataKey;
 		}
 
 		/// <summary>
@@ -1340,10 +1340,8 @@ namespace Gauntlet
 		protected virtual ITestReport CreateSimpleReportForHorde(TestResult Result)
 		{
 			BaseHordeReport HordeTestReport = null;
-			IEnumerable<UnrealTestRole> RoleList = null;
 			if (Globals.Params.ParseParam("UseTestDataV2"))
 			{
-				RoleList = new UnrealTestRole[] { GetCachedConfiguration().GetMainRequiredRole() };
 				AutomatedTestSessionData TestReport = new AutomatedTestSessionData(HordeReportTestKey, HordeReportTestName);
 				DateTime SessionTime = SessionStartTime.ToUniversalTime();
 				float TimeElapse = (float)(DateTime.Now - SessionStartTime).TotalSeconds;
@@ -1353,13 +1351,13 @@ namespace Gauntlet
 				Phase.SetTiming(SessionTime, TimeElapse);
 				// Propagate events
 				var Stream = Phase.GetStream();
-				foreach(UnrealTestEvent Event in TestNodeEvents)
+				foreach (UnrealTestEvent Event in TestNodeEvents)
 				{
 					Stream.AddEvent(Event);
 				}
 				if (RoleResults != null)
 				{
-					foreach(UnrealRoleResult Item in RoleResults)
+					foreach (UnrealRoleResult Item in RoleResults)
 					{
 						string RoleType = Item.Artifacts.SessionRole.RoleType.ToString();
 						foreach (UnrealTestEvent Event in Item.Events)
@@ -1373,7 +1371,7 @@ namespace Gauntlet
 						}
 						// Since we are going through the Role, lets add the related device
 						ITargetDevice TargetDevice = Item.Artifacts.AppInstance.Device;
-						string DeviceName = TargetDevice.Platform == BuildHostPlatform.Current.Platform? System.Environment.MachineName : TargetDevice.Name;
+						string DeviceName = TargetDevice.Platform == BuildHostPlatform.Current.Platform ? System.Environment.MachineName : TargetDevice.Name;
 						string DeviceKey = TargetDevice.Name.Replace(".", "-"); // Key must not contains dot
 						var ReportDevice = TestReport.AddDevice(DeviceName, DeviceKey);
 						ReportDevice.metadata.Add("Platform", TargetDevice.Platform.ToString());
@@ -1384,28 +1382,21 @@ namespace Gauntlet
 				// Outcome
 				Phase.SetOutcome(GetTestResult());
 				HordeTestReport = TestReport;
+				// Metadata
+				SetReportMetadata(HordeTestReport, new UnrealTestRole[] { GetCachedConfiguration().GetMainRequiredRole() });
 			}
-			else
+			if (HordeTestReport == null || Globals.Params.ParseParam("AddTestDataV1"))
 			{
-				RoleList = GetCachedConfiguration().RequiredRoles.Values.SelectMany(V => V);
-				SimpleTestReport SimpleReport = new SimpleTestReport(UnrealTestResult, Context, GetCachedConfiguration());
-				SimpleReport.TestName = HordeReportTestName;
-				SimpleReport.ReportCreatedOn = DateTime.UtcNow.ToString();
-				SimpleReport.TotalDurationSeconds = (float)(DateTime.Now - SessionStartTime).TotalSeconds;
-				SimpleReport.Description = GetMainRoleContextString();
-				SimpleReport.Errors.AddRange(GetErrorsAndAbnornalExits());
-				if (!string.IsNullOrEmpty(CancellationReason))
+				SimpleTestReport SimpleReport = CreateSimpleReportForHorde_legacy(Result);
+				// Set or Attach
+				if(HordeTestReport == null)
 				{
-					SimpleReport.Errors.Add(CancellationReason);
+					HordeTestReport = SimpleReport;
 				}
-				SimpleReport.Warnings.AddRange(GetWarnings());
-				SimpleReport.HasSucceeded = !(Result == TestResult.Failed || Result == TestResult.TimedOut || SimpleReport.Errors.Count > 0);
-				if (SimpleReport.Errors.Count > 0 && Result == TestResult.Passed)
+				else
 				{
-					SetUnrealTestResult(TestResult.Failed);
+					HordeTestReport.AttachDependencyReport(SimpleReport.GetTestData(), GetHordeTestDataKey_legacy());
 				}
-				SimpleReport.Status = GetTestResult().ToString();
-				HordeTestReport = SimpleReport;
 			}
 			string HordeArtifactPath = string.IsNullOrEmpty(GetCachedConfiguration().HordeArtifactPath) ? DefaultArtifactsDir : GetCachedConfiguration().HordeArtifactPath;
 			string ArtifactPathLabel = new DirectoryInfo(ArtifactPath).Name;
@@ -1431,12 +1422,35 @@ namespace Gauntlet
 					}
 				}
 			}
-			// Metadata
-			SetReportMetadata(HordeTestReport, RoleList);
 
 			return HordeTestReport;
 		}
 
+		private SimpleTestReport CreateSimpleReportForHorde_legacy(TestResult Result)
+		{
+			SimpleTestReport SimpleReport = new SimpleTestReport(UnrealTestResult, Context, GetCachedConfiguration());
+			SimpleReport.TestName = HordeReportTestName;
+			SimpleReport.ReportCreatedOn = DateTime.UtcNow.ToString();
+			SimpleReport.TotalDurationSeconds = (float)(DateTime.Now - SessionStartTime).TotalSeconds;
+			SimpleReport.Description = GetMainRoleContextString();
+			SimpleReport.Errors.AddRange(GetErrorsAndAbnornalExits());
+			if (!string.IsNullOrEmpty(CancellationReason))
+			{
+				SimpleReport.Errors.Add(CancellationReason);
+			}
+			SimpleReport.Warnings.AddRange(GetWarnings());
+			SimpleReport.HasSucceeded = !(Result == TestResult.Failed || Result == TestResult.TimedOut || SimpleReport.Errors.Count > 0);
+			if (SimpleReport.Errors.Count > 0 && Result == TestResult.Passed)
+			{
+				SetUnrealTestResult(TestResult.Failed);
+			}
+			SimpleReport.Status = GetTestResult().ToString();
+			// Metadata
+			SetReportMetadata(SimpleReport, GetCachedConfiguration().RequiredRoles.Values.SelectMany(V => V));
+
+			return SimpleReport;
+		}
+
 		/// <summary>
 		/// Generate report from Unreal Automated Test Results
 		/// </summary>
@@ -1449,35 +1463,48 @@ namespace Gauntlet
 			string JsonReportPath = Path.Combine(UnrealAutomatedTestReportPath, "index.json");
 			if (File.Exists(JsonReportPath))
 			{
-				string HordeArtifactPath = GetCachedConfiguration().HordeArtifactPath;
 				Log.Verbose("Reading json Unreal Automated test report from {Path}", JsonReportPath);
 				UnrealAutomatedTestPassResults JsonTestPassResults = UnrealAutomatedTestPassResults.LoadFromJson(JsonReportPath);
 				// With UE Test Automation, we care only for one role.
 				var MainRole = GetCachedConfiguration().GetMainRequiredRole();
 				var RoleList = new List<UnrealTestRole>() { MainRole };
 				// Convert test results for Horde
-				ITestReport HordeTestPassResults = null;
+				BaseHordeReport HordeTestPassResults = null;
 				if (Globals.Params.ParseParam("UseTestDataV2"))
 				{
 					string ArtifactPathLabel = new DirectoryInfo(ArtifactPath).Name;
-					HordeArtifactPath = Path.Combine(HordeArtifactPath, ArtifactPathLabel);
+					string HordeArtifactPath = Path.Combine(GetCachedConfiguration().HordeArtifactPath, ArtifactPathLabel);
 					HordeTestPassResults = AutomatedTestSessionData.FromUnrealAutomatedTests(
 						JsonTestPassResults, HordeReportTestKey, HordeReportTestName, UnrealAutomatedTestReportPath, HordeArtifactPath
 					);
 				}
-				else
+				if (HordeTestPassResults == null || Globals.Params.ParseParam("AddTestDataV1"))
 				{
+					string HordeArtifactPath = GetCachedConfiguration().HordeArtifactPath;
 					AutomatedTestSessionData_legacy LegacyTestPassResults = AutomatedTestSessionData_legacy.FromUnrealAutomatedTests(
 						JsonTestPassResults, Type, Suite, UnrealAutomatedTestReportPath, HordeArtifactPath
 					);
+					// Pre Flight information
+					LegacyTestPassResults.PreFlightChange = GetCachedConfiguration().PreFlightChange;
+					// Set or Attach
+					if (HordeTestPassResults == null)
+					{
+						HordeTestPassResults = LegacyTestPassResults;
+					}
+					else
+					{
+						SetReportMetadata(LegacyTestPassResults, RoleList);
+						HordeTestPassResults.AttachDependencyReport(LegacyTestPassResults.GetTestData(), GetHordeTestDataKey_legacy());
+						foreach (var Item in LegacyTestPassResults.GetReportDependencies())
+						{
+							HordeTestPassResults.AttachDependencyReport(Item.Value, Item.Key);
+						}
+					}
 					// Make a copy of the report in the old way - until we decide the transition is over
 					UnrealEngineTestPassResults CopyTestPassResults = UnrealEngineTestPassResults.FromUnrealAutomatedTests(JsonTestPassResults, ReportURL);
 					CopyTestPassResults.CopyTestResultsArtifacts(UnrealAutomatedTestReportPath, HordeArtifactPath);
 					SetReportMetadata(CopyTestPassResults, RoleList); // Set the metadata to the old report
-					LegacyTestPassResults.AttachDependencyReport(CopyTestPassResults, HordeReportTestKey);
-					// Pre Flight information
-					LegacyTestPassResults.PreFlightChange = GetCachedConfiguration().PreFlightChange;
-					HordeTestPassResults = LegacyTestPassResults;
+					HordeTestPassResults.AttachDependencyReport(CopyTestPassResults, GetHordeTestDataKey_legacy());
 				}
 				// Metadata
 				SetReportMetadata(HordeTestPassResults, RoleList);
@@ -1546,7 +1573,7 @@ namespace Gauntlet
 				FileName + ".TestData.json"
 			);
 			HordeReport.TestDataCollection HordeTestDataCollection = new HordeReport.TestDataCollection();
-			HordeTestDataCollection.AddNewTestReport(Report, HordeReportTestKey);
+			HordeTestDataCollection.AddNewTestReport(Report, GetHordeTestDataKey_legacy());
 			HordeTestDataCollection.WriteToJson(HordeTestDataFilePath, !AutomationTool.Automation.IsBuildMachine);
 		}
 
