[Gauntlet] Fix Metadata duplication in legacy Horde Report

#rnx

diff --git a/Engine/Source/Programs/AutomationTool/Gauntlet/Framework/Base/Gauntlet.TestReport.cs b/Engine/Source/Programs/AutomationTool/Gauntlet/Framework/Base/Gauntlet.TestReport.cs
--- a/Engine/Source/Programs/AutomationTool/Gauntlet/Framework/Base/Gauntlet.TestReport.cs
+++ b/Engine/Source/Programs/AutomationTool/Gauntlet/Framework/Base/Gauntlet.TestReport.cs
@@ -139,18 +139,18 @@ namespace Gauntlet
 	{
 		public BaseTestReport()
 		{
-			metadata = new Dictionary<string, string>();
+			Metadata = new Dictionary<string, string>();
 		}
 
 		/// <summary>
 		/// Horde report version
 		/// </summary>
-		public abstract int version { get; }
+		public abstract int Version { get; }
 
 		/// <summary>
 		/// Metadata blackboard
 		/// </summary>
-		public virtual Dictionary<string, string> metadata { get; protected set; }
+		public Dictionary<string, string> Metadata { get; private set; }
 
 		/// <summary>
 		/// Return report type
@@ -182,7 +182,7 @@ namespace Gauntlet
 		/// <returns></returns>
 		public virtual void SetMetadata(string Key, string Value)
 		{
-			metadata[Key] = Value;
+			Metadata[Key] = Value;
 		}
 
 		/// <summary>
diff --git a/Engine/Source/Programs/AutomationTool/Gauntlet/Framework/Utils/Gauntlet.HordeReport.cs b/Engine/Source/Programs/AutomationTool/Gauntlet/Framework/Utils/Gauntlet.HordeReport.cs
--- a/Engine/Source/Programs/AutomationTool/Gauntlet/Framework/Utils/Gauntlet.HordeReport.cs
+++ b/Engine/Source/Programs/AutomationTool/Gauntlet/Framework/Utils/Gauntlet.HordeReport.cs
@@ -40,12 +40,12 @@ namespace Gauntlet
 			}
 		}
 
-		public abstract class BaseHordeReport : BaseTestReport, IDataBlob
+		public abstract class BaseHordeReport : BaseTestReport
 		{
 			/// <summary>
 			/// Horde report version
 			/// </summary>
-			public override int version => 1;
+			public override int Version => 1;
 
 			protected string OutputArtifactPath;
 			protected HashSet<string> ArtifactProcessedHashes;
@@ -181,7 +181,7 @@ namespace Gauntlet
 			/// Return the data to be stored in the TestData
 			/// </summary>
 			/// <returns></returns>
-			public virtual IDataBlob GetTestData()
+			public virtual object GetTestData()
 			{
 				return this;
 			}
@@ -369,9 +369,6 @@ namespace Gauntlet
 				get { return "Unreal Automated Tests"; }
 			}
 
-			[JsonPropertyName("Metadata")]
-			public override Dictionary<string, string> metadata { get; protected set; }
-
 			public UnrealEngineTestPassResults() : base()
 			{
 				Devices = new List<Device>();
@@ -555,8 +552,6 @@ namespace Gauntlet
 				get { return "Automated Test Session"; }
 			}
 
-			[JsonPropertyName("Metadata")]
-			public override Dictionary<string, string> metadata { get; protected set; }
 
 			public class TestResult
 			{
@@ -1102,9 +1097,6 @@ namespace Gauntlet
 				get { return "Simple Report"; }
 			}
 
-			[JsonPropertyName("Metadata")]
-			public override Dictionary<string, string> metadata { get; protected set; }
-
 			public SimpleTestReport() : base()
 			{
 
@@ -1188,7 +1180,7 @@ namespace Gauntlet
 		{
 			public override string Type => "Automated Test Session";
 
-			public override int version => 2;
+			public override int Version => 2;
 
 			/// <summary>
 			/// Unique key that identify the test and group the different sessions
@@ -1418,8 +1410,8 @@ namespace Gauntlet
 				public TestSessionSummary summary { get; set; }
 				public List<TestPhase> phases { get; set; }
 				public Dictionary<string, TestDevice> devices { get; set; }
-				public int version { get { return report.version; } }
-				public Dictionary<string, string> metadata { get { return report.metadata; } }
+				public int version { get { return report.Version; } }
+				public Dictionary<string, string> metadata { get { return report.Metadata; } }
 
 				[JsonIgnore]
 				// back pointer
