Add -projectprogramoverride to allow BuildCookRun to override passed in uproject paramter when building programs.

diff --git a/Engine/Source/Programs/AutomationTool/AutomationUtils/BuildCommand.cs b/Engine/Source/Programs/AutomationTool/AutomationUtils/BuildCommand.cs
--- a/Engine/Source/Programs/AutomationTool/AutomationUtils/BuildCommand.cs
+++ b/Engine/Source/Programs/AutomationTool/AutomationUtils/BuildCommand.cs
@@ -251,6 +251,46 @@ namespace AutomationTool
 				return int.Parse(Value);
 			}
 		}
+		/// <summary>
+		/// Parses project name string into a FileReference. Can be "Game" or "Game.uproject" or "Path/To/Game.uproject"
+		/// </summary>
+		/// <param name="OriginalProjectName">In project string to parse</param>
+		/// <returns>FileReference to uproject</returns>
+		public FileReference ParseProjectString(string OriginalProjectName)
+		{
+			FileReference ProjectFullPath = null;
+
+			if (string.IsNullOrEmpty(OriginalProjectName))
+			{
+				return null;
+			}
+
+			var ProjectName = OriginalProjectName;
+			ProjectName = ProjectName.Trim(new char[] { '\"' });
+			if (ProjectName.IndexOfAny(new char[] { '\\', '/' }) < 0)
+			{
+				ProjectName = CombinePaths(CmdEnv.LocalRoot, ProjectName, ProjectName + ".uproject");
+			}
+			else if (!FileExists_NoExceptions(ProjectName))
+			{
+				ProjectName = CombinePaths(CmdEnv.LocalRoot, ProjectName);
+			}
+			if (FileExists_NoExceptions(ProjectName))
+			{
+				ProjectFullPath = new FileReference(ProjectName);
+			}
+			else
+			{
+				var Branch = new BranchInfo();
+				var GameProj = Branch.FindGame(OriginalProjectName);
+				if (GameProj != null)
+				{
+					ProjectFullPath = GameProj.FilePath;
+				}
+			}
+
+			return ProjectFullPath;
+		}
 
 		public FileReference ParseProjectParam()
 		{
@@ -279,32 +319,11 @@ namespace AutomationTool
 					return null;
 				}
 
-				var ProjectName = OriginalProjectName;
-				ProjectName = ProjectName.Trim(new char[] { '\"' });
-				if (ProjectName.IndexOfAny(new char[] { '\\', '/' }) < 0)
-				{
-					ProjectName = CombinePaths(CmdEnv.LocalRoot, ProjectName, ProjectName + ".uproject");
-				}
-				else if (!FileExists_NoExceptions(ProjectName))
-				{
-					ProjectName = CombinePaths(CmdEnv.LocalRoot, ProjectName);
-				}
-				if (FileExists_NoExceptions(ProjectName))
-				{
-					ProjectFullPath = new FileReference(ProjectName);
-				}
-				else
+				ProjectFullPath = ParseProjectString(OriginalProjectName);
+
+				if (ProjectFullPath == null || !FileExists_NoExceptions(ProjectFullPath.FullName))
 				{
-					var Branch = new BranchInfo();
-					var GameProj = Branch.FindGame(OriginalProjectName);
-					if (GameProj != null)
-					{
-						ProjectFullPath = GameProj.FilePath;
-					}
-					if (ProjectFullPath == null || !FileExists_NoExceptions(ProjectFullPath.FullName))
-					{
-						throw new AutomationException("Could not find a project file {0}.", ProjectName);
-					}
+					throw new AutomationException("Could not find a project file {0}.", OriginalProjectName);
 				}
 			}
 
diff --git a/Engine/Source/Programs/AutomationTool/AutomationUtils/ProjectParams.cs b/Engine/Source/Programs/AutomationTool/AutomationUtils/ProjectParams.cs
--- a/Engine/Source/Programs/AutomationTool/AutomationUtils/ProjectParams.cs
+++ b/Engine/Source/Programs/AutomationTool/AutomationUtils/ProjectParams.cs
@@ -295,6 +295,7 @@ namespace AutomationTool
 			//
 
 			this.RawProjectPath = InParams.RawProjectPath;
+			this.RawProgramProjectOverride = InParams.RawProgramProjectOverride;
 			this.MapsToCook = InParams.MapsToCook;
 			this.MapIniSectionsToCook = InParams.MapIniSectionsToCook;
 			this.DirectoriesToCook = InParams.DirectoriesToCook;
@@ -458,7 +459,7 @@ namespace AutomationTool
 		/// </summary>
 		public ProjectParams(			
 			FileReference RawProjectPath,
-
+			FileReference RawProgramProjectOverride = null,
 			BuildCommand Command = null,
 			string Device = null,			
 			string MapToRun = null,	
@@ -637,6 +638,11 @@ namespace AutomationTool
 			}
 			catch { this.ProjectDescriptor = new ProjectDescriptor(); }
 
+			if(RawProgramProjectOverride != null)
+			{
+				this.RawProgramProjectOverride = RawProgramProjectOverride;
+			}
+
 			if (DirectoriesToCook != null)
 			{
 				this.DirectoriesToCook = DirectoriesToCook;
@@ -2878,13 +2884,22 @@ namespace AutomationTool
 
 		public FileReference CodeBasedUprojectPath
 		{
-            get { return IsCodeBasedProject ? RawProjectPath : null; }
+            get 
+			{ 
+				if(RawProgramProjectOverride != null) { return RawProgramProjectOverride; }
+				return IsCodeBasedProject ? RawProjectPath : null; 
+			}
 		}
 		/// <summary>
 		/// True if this project is a program.
 		/// </summary>
 		public bool IsProgramTarget { get; private set; }
 
+		/// <summary>
+		/// Returns override for Program Project files located in non-engine projects.
+		/// </summary>
+		public FileReference RawProgramProjectOverride { get; set; }
+
 		/// <summary>
 		/// Path where the project's game (or program) binaries are built for the given target platform.
 		/// </summary>
@@ -3105,6 +3120,15 @@ namespace AutomationTool
                 throw new AutomationException("RawProjectPath {0} file must exist", RawProjectPath);
             }
 
+			if (RawProgramProjectOverride != null && !RawProgramProjectOverride.HasExtension(".uproject"))
+			{
+				throw new AutomationException("RawProgramProjectPath {0} must end with .uproject", RawProgramProjectOverride);
+			}
+			if (RawProgramProjectOverride != null && !CommandUtils.FileExists(RawProgramProjectOverride.FullName))
+			{
+				throw new AutomationException("RawProgramProjectOverride {0} file must exist", RawProgramProjectOverride);
+			}
+
 			if (FileServer && !Cook && !CookInEditor)
 			{
 				throw new AutomationException("Only cooked builds can use a fileserver, use -cook or -CookInEditor");
@@ -3368,6 +3392,7 @@ namespace AutomationTool
 				Logger.LogDebug("AppLocalDirectory={AppLocalDirectory}", AppLocalDirectory);
 				Logger.LogDebug("NoBootstrapExe={NoBootstrapExe}", NoBootstrapExe);
 				Logger.LogDebug("RawProjectPath={RawProjectPath}", RawProjectPath);
+				Logger.LogDebug("RawProgramProjectOverride={RawProgramProjectOverride}", RawProgramProjectOverride);
 				Logger.LogDebug("Run={Run}", Run);
 				Logger.LogDebug("ServerConfigsToBuild={Arg0}", string.Join(",", ServerConfigsToBuild));
 				Logger.LogDebug("ServerCookedTargets={Arg0}", ServerCookedTargets.ToString());
diff --git a/Engine/Source/Programs/AutomationTool/Scripts/BuildCookRun.Automation.cs b/Engine/Source/Programs/AutomationTool/Scripts/BuildCookRun.Automation.cs
--- a/Engine/Source/Programs/AutomationTool/Scripts/BuildCookRun.Automation.cs
+++ b/Engine/Source/Programs/AutomationTool/Scripts/BuildCookRun.Automation.cs
@@ -62,6 +62,12 @@ public class BuildCookRun : BuildCommand, IProjectParamsHelpers
 			Params.DirectoriesToCook = new ParamList<string>(DirectoriesToCook.Split('+'));
 		}
 
+		var ProjectProgramOverride = ParseParamValue("projectprogramoverride");
+		if (!String.IsNullOrEmpty(ProjectProgramOverride))
+		{
+			Params.RawProgramProjectOverride = ParseProjectString(ProjectProgramOverride);
+		}
+
 		var DDCGraph = ParseParamValue("ddc");
 		if (!String.IsNullOrEmpty(DDCGraph))
 		{
diff --git a/Engine/Source/Programs/AutomationTool/Scripts/BuildProjectCommand.Automation.cs b/Engine/Source/Programs/AutomationTool/Scripts/BuildProjectCommand.Automation.cs
--- a/Engine/Source/Programs/AutomationTool/Scripts/BuildProjectCommand.Automation.cs
+++ b/Engine/Source/Programs/AutomationTool/Scripts/BuildProjectCommand.Automation.cs
@@ -163,7 +163,8 @@ namespace AutomationScripts
 						UnrealTargetPlatform CrashReportPlatform = Platform.GetPlatform(ClientPlatformType).CrashReportPlatform ?? ClientPlatformType;
 						CrashReportPlatforms.Add(CrashReportPlatform);
 						string Arch = Params.IsProgramTarget ? ProgramBuildArgs : ClientBuildArgs;
-						ProjectAgenda.AddTargets(Params.ClientCookedTargets.ToArray(), ClientPlatformType, BuildConfig, Params.CodeBasedUprojectPath, 
+						FileReference ProjectRefrence = Params.CodeBasedUprojectPath;
+						ProjectAgenda.AddTargets(Params.ClientCookedTargets.ToArray(), ClientPlatformType, BuildConfig, ProjectRefrence, 
 							InAddArgs: $" -remoteini=\"{Params.RawProjectPath.Directory}\" {AdditionalArgs} {Arch}");
 					}
 				}
diff --git a/Engine/Source/Programs/AutomationTool/Scripts/CopyBuildToStagingDirectory.Automation.cs b/Engine/Source/Programs/AutomationTool/Scripts/CopyBuildToStagingDirectory.Automation.cs
--- a/Engine/Source/Programs/AutomationTool/Scripts/CopyBuildToStagingDirectory.Automation.cs
+++ b/Engine/Source/Programs/AutomationTool/Scripts/CopyBuildToStagingDirectory.Automation.cs
@@ -5332,6 +5332,12 @@ namespace AutomationScripts
 						{
 							ReceiptBaseDir = EngineDir;
 						}
+
+						//if RawProgramProjectOverride then the binary will likely be under <Game>/Binaries
+						if (Params.RawProgramProjectOverride != null)
+						{
+							ReceiptBaseDir = DirectoryReference.FromFile(Params.RawProgramProjectOverride);
+						}
 					}
 
 					Platform PlatformInstance = Platform.Platforms[StagePlatform];
