Updating UploadFilestoS3 to use assume role
Adding job to upload QFE to s3
Adding job to upload linux toolchain to s3
Adjusting remaining s3 jobs for UE5 to use assumerole

#rnx

diff --git a/Engine/Source/Programs/AutomationTool/AutomationTool.csproj b/Engine/Source/Programs/AutomationTool/AutomationTool.csproj
--- a/Engine/Source/Programs/AutomationTool/AutomationTool.csproj
+++ b/Engine/Source/Programs/AutomationTool/AutomationTool.csproj
@@ -95,6 +95,7 @@
   </ItemGroup>
 
   <ItemGroup>
+    <PackageReference Include="AWSSDK.SecurityToken" Version="3.7.300.105" />
     <PackageReference Include="Microsoft.Build" Version="17.1.0" ExcludeAssets="runtime" PrivateAssets="all" />
     <PackageReference Include="Microsoft.Build.Locator" Version="1.4.1" PrivateAssets="all" />
     <!-- Added as a temporary workaround to running from visual studio where the system version of dotnet is run. 
