Fix for BP only project quick launch fail

diff --git a/Engine/Source/Programs/AutomationTool/Mac/MacPlatform.Automation.cs b/Engine/Source/Programs/AutomationTool/Mac/MacPlatform.Automation.cs
--- a/Engine/Source/Programs/AutomationTool/Mac/MacPlatform.Automation.cs
+++ b/Engine/Source/Programs/AutomationTool/Mac/MacPlatform.Automation.cs
@@ -587,6 +587,12 @@ public class MacPlatform : ApplePlatform
 			Int32 BaseDirLen = Params.BaseStageDirectory.Length;
 			string StageSubDir = ClientApp.Substring(BaseDirLen, ClientApp.IndexOf("/", BaseDirLen + 1) - BaseDirLen);
 			ClientApp = CombinePaths(Params.BaseStageDirectory, StageSubDir, $"{ExeName}.app/Contents/MacOS/{ExeName}");
+			if (!File.Exists(ClientApp))
+			{
+				// Could be blueprint only projects which ClientApp would be pointing at non-existing UnrealGame/UnrealClient
+				ExeName = Params.RawProjectPath.GetFileNameWithoutAnyExtensions();
+				ClientApp = CombinePaths(Params.BaseStageDirectory, StageSubDir, $"{ExeName}.app/Contents/MacOS/{ExeName}");
+			}
 		}
 		else if (!File.Exists(ClientApp))
 		{
