AutomationTool: Update BuildCommand to use nullable annotations.

#rnx

diff --git a/Engine/Source/Programs/AutomationTool/AutomationUtils/BuildCommand.cs b/Engine/Source/Programs/AutomationTool/AutomationUtils/BuildCommand.cs
--- a/Engine/Source/Programs/AutomationTool/AutomationUtils/BuildCommand.cs
+++ b/Engine/Source/Programs/AutomationTool/AutomationUtils/BuildCommand.cs
@@ -7,6 +7,9 @@ using System.Threading.Tasks;
 using EpicGames.Core;
 using UnrealBuildBase;
 using Microsoft.Extensions.Logging;
+using System.Diagnostics.CodeAnalysis;
+
+#nullable enable
 
 namespace AutomationTool
 {
@@ -18,13 +21,14 @@ namespace AutomationTool
 		/// <summary>
 		/// Command line parameters for this command (empty by non-null by default)
 		/// </summary>
-		private string[] CommandLineParams = new string[0];
 		public string[] Params
 		{
-			get { return CommandLineParams; }
-			set { CommandLineParams = value; }
+			get { return _commandLineParams; }
+			set { _commandLineParams = value; }
 		}
 
+		private string[] _commandLineParams = new string[0];
+
 		/// <summary>
 		/// Parses the command's Params list for a parameter and returns whether it is defined or not.
 		/// </summary>
@@ -43,7 +47,8 @@ namespace AutomationTool
 		/// <param name="Default"></param>
 		/// <param name="ObsoleteParam"></param>
 		/// <returns>Returns the value or Default if the parameter was not found.</returns>
-		public string ParseParamValue(string Param, string Default = null, string ObsoleteParam = null)
+		[return: NotNullIfNotNull("Default")]
+		public string? ParseParamValue(string Param, string? Default = null, string? ObsoleteParam = null)
 		{
 			string ParamValue = ParseParamValue(Params, Param, null);
 
@@ -73,7 +78,7 @@ namespace AutomationTool
 		/// </summary>
 		/// <param name="Param"></param>
 		/// <returns></returns>
-		public string ParseOptionalStringParam(string Param)
+		public string? ParseOptionalStringParam(string Param)
 		{
 			return ParseParamValue(Param, null);
 		}
@@ -85,12 +90,12 @@ namespace AutomationTool
 		/// <returns>Value of the argument</returns>
 		public string ParseRequiredStringParam(string Param)
 		{
-			string Value = ParseOptionalStringParam(Param);
-			if(Value == null)
+			string? value = ParseOptionalStringParam(Param);
+			if(value == null)
 			{
 				throw new AutomationException("Missing -{0}=... parameter", Param);
 			}
-			return Value;
+			return value;
 		}
 
 		/// <summary>
@@ -98,16 +103,16 @@ namespace AutomationTool
 		/// </summary>
 		/// <param name="Param">Name of the argument</param>
 		/// <returns>Value of the argument</returns>
-		public FileReference ParseOptionalFileReferenceParam(string Param)
+		public FileReference? ParseOptionalFileReferenceParam(string Param)
 		{
-			string StringValue = ParseParamValue(Param);
-			if(StringValue == null)
+			string? stringValue = ParseParamValue(Param);
+			if(stringValue == null)
 			{
 				return null;
 			}
 			else
 			{
-				return new FileReference(StringValue);
+				return new FileReference(stringValue);
 			}
 		}
 
@@ -118,12 +123,12 @@ namespace AutomationTool
 		/// <returns>Value of the argument</returns>
 		public FileReference ParseRequiredFileReferenceParam(string Param)
 		{
-			FileReference Value = ParseOptionalFileReferenceParam(Param);
-			if(Value == null)
+			FileReference? value = ParseOptionalFileReferenceParam(Param);
+			if(value == null)
 			{
 				throw new AutomationException("Missing -{0}=... parameter", Param);
 			}
-			return Value;
+			return value;
 		}
 
 		/// <summary>
@@ -131,16 +136,16 @@ namespace AutomationTool
 		/// </summary>
 		/// <param name="Param">Name of the argument</param>
 		/// <returns>Value of the argument</returns>
-		public DirectoryReference ParseOptionalDirectoryReferenceParam(string Param)
+		public DirectoryReference? ParseOptionalDirectoryReferenceParam(string Param)
 		{
-			string StringValue = ParseOptionalStringParam(Param);
-			if(StringValue == null)
+			string? stringValue = ParseOptionalStringParam(Param);
+			if(stringValue == null)
 			{
 				return null;
 			}
 			else
 			{
-				return new DirectoryReference(StringValue);
+				return new DirectoryReference(stringValue);
 			}
 		}
 
@@ -151,12 +156,12 @@ namespace AutomationTool
 		/// <returns>Value of the argument</returns>
 		public DirectoryReference ParseRequiredDirectoryReferenceParam(string Param)
 		{
-			DirectoryReference Value = ParseOptionalDirectoryReferenceParam(Param);
-			if(Value == null)
+			DirectoryReference? value = ParseOptionalDirectoryReferenceParam(Param);
+			if(value == null)
 			{
 				throw new AutomationException("Missing -{0}=... parameter", Param);
 			}
-			return Value;
+			return value;
 		}
 
 		/// <summary>
@@ -166,17 +171,17 @@ namespace AutomationTool
 		/// <returns>Returns the value that was parsed.</returns>
 		public Nullable<T> ParseOptionalEnumParam<T>(string Param) where T : struct
 		{
-			string ValueString = ParseParamValue(Param);
-			if(ValueString == null)
+			string? valueString = ParseParamValue(Param);
+			if(valueString == null)
 			{
 				return null;
 			}
 			else
 			{
 				T Value;
-				if(!Enum.TryParse<T>(ValueString, out Value))
+				if(!Enum.TryParse<T>(valueString, out Value))
 				{
-					throw new AutomationException("'{0}' is not a valid value for {1}", ValueString, typeof(T).Name);
+					throw new AutomationException("'{0}' is not a valid value for {1}", valueString, typeof(T).Name);
 				}
 				return Value;
 			}
@@ -189,12 +194,12 @@ namespace AutomationTool
 		/// <returns>Returns the value that was parsed.</returns>
 		public T ParseRequiredEnumParamEnum<T>(string Param) where T : struct
 		{
-			Nullable<T> Value = ParseOptionalEnumParam<T>(Param);
-			if(!Value.HasValue)
+			Nullable<T> value = ParseOptionalEnumParam<T>(Param);
+			if(!value.HasValue)
 			{
 				throw new AutomationException("Missing -{0}=... parameter", Param);
 			}
-			return Value.Value;
+			return value.Value;
 		}
 
 		/// <summary>
@@ -241,113 +246,113 @@ namespace AutomationTool
 		/// <returns>Returns the value or Default if the parameter was not found.</returns>
 		public int? ParseParamNullableInt(string Param)
 		{
-			string Value = ParseParamValue(Params, Param, null);
-			if(Value == null)
+			string value = ParseParamValue(Params, Param, null);
+			if(value == null)
 			{
 				return null;
 			}
 			else
 			{
-				return int.Parse(Value);
+				return int.Parse(value);
 			}
 		}
 		/// <summary>
 		/// Parses project name string into a FileReference. Can be "Game" or "Game.uproject" or "Path/To/Game.uproject"
 		/// </summary>
-		/// <param name="OriginalProjectName">In project string to parse</param>
+		/// <param name="originalProjectName">In project string to parse</param>
 		/// <returns>FileReference to uproject</returns>
-		public FileReference ParseProjectString(string OriginalProjectName)
+		public FileReference? ParseProjectString(string originalProjectName)
 		{
-			FileReference ProjectFullPath = null;
+			FileReference? projectFullPath = null;
 
-			if (string.IsNullOrEmpty(OriginalProjectName))
+			if (string.IsNullOrEmpty(originalProjectName))
 			{
 				return null;
 			}
 
-			var ProjectName = OriginalProjectName;
-			ProjectName = ProjectName.Trim(new char[] { '\"' });
-			if (ProjectName.IndexOfAny(new char[] { '\\', '/' }) < 0)
+			string projectName = originalProjectName;
+			projectName = projectName.Trim(new char[] { '\"' });
+			if (projectName.IndexOfAny(new char[] { '\\', '/' }) < 0)
 			{
-				ProjectName = CombinePaths(CmdEnv.LocalRoot, ProjectName, ProjectName + ".uproject");
+				projectName = CombinePaths(CmdEnv.LocalRoot, projectName, projectName + ".uproject");
 			}
-			else if (!FileExists_NoExceptions(ProjectName))
+			else if (!FileExists_NoExceptions(projectName))
 			{
-				ProjectName = CombinePaths(CmdEnv.LocalRoot, ProjectName);
+				projectName = CombinePaths(CmdEnv.LocalRoot, projectName);
 			}
-			if (FileExists_NoExceptions(ProjectName))
+			if (FileExists_NoExceptions(projectName))
 			{
-				ProjectFullPath = new FileReference(ProjectName);
+				projectFullPath = new FileReference(projectName);
 			}
 			else
 			{
 				var Branch = new BranchInfo();
-				var GameProj = Branch.FindGame(OriginalProjectName);
+				var GameProj = Branch.FindGame(originalProjectName);
 				if (GameProj != null)
 				{
-					ProjectFullPath = GameProj.FilePath;
+					projectFullPath = GameProj.FilePath;
 				}
 			}
 
-			return ProjectFullPath;
+			return projectFullPath;
 		}
 
-		public FileReference ParseProjectParam()
+		public FileReference? ParseProjectParam()
 		{
-			FileReference ProjectFullPath = null;
+			FileReference? projectFullPath = null;
 
-			var bForeign = ParseParam("foreign");
-			var bForeignCode = ParseParam("foreigncode");
+			bool bForeign = ParseParam("foreign");
+			bool bForeignCode = ParseParam("foreigncode");
 			if (bForeign)
 			{
-				var DestSample = ParseParamValue("DestSample", "CopiedHoverShip");
-				var Dest = ParseParamValue("ForeignDest", CombinePaths(@"C:\testue\foreign\", DestSample + "_ _Dir"));
-				ProjectFullPath = new FileReference(CombinePaths(Dest, DestSample + ".uproject"));
+				string destSample = ParseParamValue("DestSample", "CopiedHoverShip");
+				string dest = ParseParamValue("ForeignDest", CombinePaths(@"C:\testue\foreign\", destSample + "_ _Dir"));
+				projectFullPath = new FileReference(CombinePaths(dest, destSample + ".uproject"));
 			}
 			else if (bForeignCode)
 			{
-				var DestSample = ParseParamValue("DestSample", "PlatformerGame");
-				var Dest = ParseParamValue("ForeignDest", CombinePaths(@"C:\testue\foreign\", DestSample + "_ _Dir"));
-				ProjectFullPath = new FileReference(CombinePaths(Dest, DestSample + ".uproject"));
+				string destSample = ParseParamValue("DestSample", "PlatformerGame");
+				string dest = ParseParamValue("ForeignDest", CombinePaths(@"C:\testue\foreign\", destSample + "_ _Dir"));
+				projectFullPath = new FileReference(CombinePaths(dest, destSample + ".uproject"));
 			}
 			else
 			{
-				var OriginalProjectName = ParseParamValue("project", "");
+				string originalProjectName = ParseParamValue("project", "");
 
-				if (string.IsNullOrEmpty(OriginalProjectName))
+				if (string.IsNullOrEmpty(originalProjectName))
 				{
 					return null;
 				}
 
-				ProjectFullPath = ParseProjectString(OriginalProjectName);
+				projectFullPath = ParseProjectString(originalProjectName);
 
-				if (ProjectFullPath == null || !FileExists_NoExceptions(ProjectFullPath.FullName))
+				if (projectFullPath == null || !FileExists_NoExceptions(projectFullPath.FullName))
 				{
-					throw new AutomationException("Could not find a project file {0}.", OriginalProjectName);
+					throw new AutomationException("Could not find a project file {0}.", originalProjectName);
 				}
 			}
 
-			return ProjectFullPath;
+			return projectFullPath;
 		}
 
 		/// <summary>
 		/// Checks that all of the required params are present, throws an exception if not
 		/// </summary>
-		/// <param name="Args"></param>
-		public void CheckParamsArePresent(params string[] Args)
+		/// <param name="args"></param>
+		public void CheckParamsArePresent(params string[] args)
 		{
-			List<string> MissingParams = new List<string>();
-			foreach (string Arg in Args)
+			List<string> missingParams = new List<string>();
+			foreach (string arg in args)
 			{
-				if (ParseParamValue(Arg, null) == null)
+				if (ParseParamValue(arg, null) == null)
 				{
-					MissingParams.Add(Arg);
+					missingParams.Add(arg);
 				}
 			}
 
-			if (MissingParams.Count > 0)
+			if (missingParams.Count > 0)
 			{
-				throw new AutomationException("Params {0} are missing but required. Required params are {1}", string.Join(",", MissingParams), string.Join(",", Args));
+				throw new AutomationException("Params {0} are missing but required. Required params are {1}", string.Join(",", missingParams), string.Join(",", args));
 			}
 		}
 
diff --git a/Engine/Source/Programs/AutomationTool/BuildGraph/BuildGraph.cs b/Engine/Source/Programs/AutomationTool/BuildGraph/BuildGraph.cs
--- a/Engine/Source/Programs/AutomationTool/BuildGraph/BuildGraph.cs
+++ b/Engine/Source/Programs/AutomationTool/BuildGraph/BuildGraph.cs
@@ -268,19 +268,19 @@ namespace AutomationTool
 		public override async Task<ExitCode> ExecuteAsync()
 		{
 			// Parse the command line parameters
-			string className = ParseParamValue("Class", null);
-			string scriptFileName = ParseParamValue("Script", null);
+			string? className = ParseParamValue("Class", null);
+			string? scriptFileName = ParseParamValue("Script", null);
 			string[] targetNames = ParseParamValues("Target").SelectMany(x => x.Split(';', '+').Select(y => y.Trim()).Where(y => y.Length > 0)).ToArray();
-			string documentationFileName = ParseParamValue("Documentation", null);
-			string schemaFileName = ParseParamValue("Schema", null);
-			string importSchemaFileName = ParseParamValue("ImportSchema", null);
-			string exportFileName = ParseParamValue("Export", null);
-			string hordeExportFileName = ParseParamValue("HordeExport", null);
-			string preprocessedFileName = ParseParamValue("Preprocess", null);
-			string sharedStorageDir = ParseParamValue("SharedStorageDir", null);
-			string singleNodeName = ParseParamValue("SingleNode", null);
-			string triggerName = ParseParamValue("Trigger", null);
-			string tokenSignature = ParseParamValue("TokenSignature", null);
+			string? documentationFileName = ParseParamValue("Documentation", null);
+			string? schemaFileName = ParseParamValue("Schema", null);
+			string? importSchemaFileName = ParseParamValue("ImportSchema", null);
+			string? exportFileName = ParseParamValue("Export", null);
+			string? hordeExportFileName = ParseParamValue("HordeExport", null);
+			string? preprocessedFileName = ParseParamValue("Preprocess", null);
+			string? sharedStorageDir = ParseParamValue("SharedStorageDir", null);
+			string? singleNodeName = ParseParamValue("SingleNode", null);
+			string? triggerName = ParseParamValue("Trigger", null);
+			string? tokenSignature = ParseParamValue("TokenSignature", null);
 			bool skipTargetsWithoutTokens = ParseParam("SkipTargetsWithoutTokens");
 			bool resume = singleNodeName != null || ParseParam("Resume");
 			bool listOnly = ParseParam("ListOnly");
@@ -288,8 +288,8 @@ namespace AutomationTool
 			bool writeToSharedStorage = ParseParam("WriteToSharedStorage") || CommandUtils.IsBuildMachine;
 			bool publicTasksOnly = ParseParam("PublicTasksOnly");
 			bool skipValidation = ParseParam("SkipValidation");
-			string reportName = ParseParamValue("ReportName", null);
-			string branchOverride = ParseParamValue("Branch", null);
+			string? reportName = ParseParamValue("ReportName", null);
+			string? branchOverride = ParseParamValue("Branch", null);
 
 			GraphPrintOptions printOptions = GraphPrintOptions.ShowCommandLineOptions;
 			if (ParseParam("ShowDeps"))
@@ -364,7 +364,7 @@ namespace AutomationTool
 			}
 
 			// If the -project flag is given, pass useful information into the graph
-			FileReference projectFile = ParseProjectParam();
+			FileReference? projectFile = ParseProjectParam();
 			if (projectFile != null)
 			{
 				defaultProperties["ProjectName"] = projectFile.GetFileNameWithoutExtension();
diff --git a/Engine/Source/Programs/AutomationTool/BuildGraph/TempStorage.cs b/Engine/Source/Programs/AutomationTool/BuildGraph/TempStorage.cs
--- a/Engine/Source/Programs/AutomationTool/BuildGraph/TempStorage.cs
+++ b/Engine/Source/Programs/AutomationTool/BuildGraph/TempStorage.cs
@@ -1609,7 +1609,7 @@ namespace AutomationTool
 		/// </summary>
 		public override void ExecuteBuild()
 		{
-			string tempStorageDir = ParseParamValue("TempStorageDir", null);
+			string? tempStorageDir = ParseParamValue("TempStorageDir", null);
 			if (tempStorageDir == null)
 			{
 				throw new AutomationException("Missing -TempStorageDir parameter");
@@ -1621,7 +1621,7 @@ namespace AutomationTool
 				return;
 			}
 
-			string days = ParseParamValue("Days", null);
+			string? days = ParseParamValue("Days", null);
 			if (days == null)
 			{
 				throw new AutomationException("Missing -Days parameter");
diff --git a/Engine/Source/Programs/AutomationTool/Scripts/GeneratePlatformReport.cs b/Engine/Source/Programs/AutomationTool/Scripts/GeneratePlatformReport.cs
--- a/Engine/Source/Programs/AutomationTool/Scripts/GeneratePlatformReport.cs
+++ b/Engine/Source/Programs/AutomationTool/Scripts/GeneratePlatformReport.cs
@@ -37,7 +37,7 @@ namespace AutomationTool
 		{
 			// Parse parameters
 			string[] Platforms = ParseParamValue("Platforms", "").Split('+', StringSplitOptions.RemoveEmptyEntries);
-			string CsvFile = ParseParamValue("CsvOut");
+			string? CsvFile = ParseParamValue("CsvOut");
 			string DDPISection = ParseParamValue("DDPISection","DataDrivenPlatformInfo");
 			bool bFriendlyNames = !ParseParam("UseFieldNames");
 			bool bOpenReport = ParseParam("OpenReport");
