Capture summary in UnrealAutomationEntry to allow Horde issue to match them.
+ Improve 'Critical error' capture in UnrealLogParser

#rnx

diff --git a/Engine/Source/Programs/AutomationTool/Gauntlet/Unreal/Utils/Gauntlet.UnrealLogParser.cs b/Engine/Source/Programs/AutomationTool/Gauntlet/Unreal/Utils/Gauntlet.UnrealLogParser.cs
--- a/Engine/Source/Programs/AutomationTool/Gauntlet/Unreal/Utils/Gauntlet.UnrealLogParser.cs
+++ b/Engine/Source/Programs/AutomationTool/Gauntlet/Unreal/Utils/Gauntlet.UnrealLogParser.cs
@@ -959,7 +959,7 @@ namespace Gauntlet
 		/// <returns></returns>
 		public UnrealLog.CallstackMessage GetFatalError()
 		{
-			string[] ErrorMsgMatches = new string[] { @"(Fatal Error:.+)", @"Critical error: =+[\s\n]+(?:.+?\s*Error:\s*)?(.+)", @"(Assertion Failed:.+)", @"(Unhandled Exception:.+)", @"(LowLevelFatalError.+)" };
+			string[] ErrorMsgMatches = new string[] { @"(Fatal Error:.+)", @"Critical error: =+\s+(?:[\S\s]+?\s*Error: +)?(.+)", @"(Assertion Failed:.+)", @"(Unhandled Exception:.+)", @"(LowLevelFatalError.+)" };
 
 			var Traces = ParseTracedErrors(ErrorMsgMatches).Concat(GetASanErrors());
 
diff --git a/Engine/Source/Programs/AutomationTool/Gauntlet/Unreal/Utils/Gauntlet.UnrealTestResults.cs b/Engine/Source/Programs/AutomationTool/Gauntlet/Unreal/Utils/Gauntlet.UnrealTestResults.cs
--- a/Engine/Source/Programs/AutomationTool/Gauntlet/Unreal/Utils/Gauntlet.UnrealTestResults.cs
+++ b/Engine/Source/Programs/AutomationTool/Gauntlet/Unreal/Utils/Gauntlet.UnrealTestResults.cs
@@ -12,6 +12,7 @@ using Microsoft.Extensions.Logging;
 using Logging = Microsoft.Extensions.Logging;
 using System.Globalization;
 using AutomationUtils.Matchers;
+using System.Text.RegularExpressions;
 
 namespace Gauntlet
 {
@@ -166,6 +167,12 @@ namespace Gauntlet
 		public const string InvalidDateTime = "0001.01.01-00.00.00";
 		public const string DateTimeFormat = "yyyy.MM.dd-HH.mm.ss";
 
+		/// <summary>
+		/// Regex pattern to match assertion, ensure and other crash report summary from UE
+		/// ie: Assertion failed: !This->CurrentlySerializingObject [File:.\\Runtime/CoreUObject/Private/Misc/GCObjectReferencer.cpp] [Line: 54] 
+		/// </summary>
+		private static Regex SummaryPattern = new Regex(@"\w[\w ]+:(?:\s+.+)?\s+\[[^]]+\.[^]]+\](?:\s*\[[^]]+\])?(?:\s*\n\w+.+)? *", RegexOptions.Multiline | RegexOptions.ExplicitCapture);
+
 		public static DateTime GetTimestampAsDateTime(string StringTime)
 		{
 			DateTime Time = DateTime.UtcNow;
@@ -205,8 +212,20 @@ namespace Gauntlet
 				}
 				else
 				{
-					Properties.Add("Callstack", Message);
-					Format = "{Callstack}";
+					Match MatchSummary = SummaryPattern.Match(Message);
+					if (MatchSummary.Success)
+					{
+						Group MatchedGroup = MatchSummary.Groups[0];
+						int Index = MatchedGroup.Index;
+						Format = $"{Message.Substring(0, Index)}{{Summary}}\n{{Callstack}}";
+						Properties.Add("Summary", MatchedGroup.Value);
+						Properties.Add("Callstack", Message.Substring(Index + MatchedGroup.Length).TrimStart('\n'));
+					}
+					else
+					{
+						Properties.Add("Callstack", Message);
+						Format = "{Callstack}";
+					}
 				}
 			}
 			else
