AutomationTool: Add a service provider to the CommandUtils class, and add utility methods to allow retrieving secrets from Horde.

#rnx

diff --git a/Engine/Source/Programs/AutomationTool/AutomationUtils/Automation.cs b/Engine/Source/Programs/AutomationTool/AutomationUtils/Automation.cs
--- a/Engine/Source/Programs/AutomationTool/AutomationUtils/Automation.cs
+++ b/Engine/Source/Programs/AutomationTool/AutomationUtils/Automation.cs
@@ -8,10 +8,12 @@ using System.Reflection;
 using System.Diagnostics;
 using UnrealBuildTool;
 using EpicGames.Core;
+using EpicGames.Horde;
 using OpenTracing.Util;
 using UnrealBuildBase;
 using System.Threading.Tasks;
 using Microsoft.Extensions.Logging;
+using Microsoft.Extensions.DependencyInjection;
 
 using static AutomationTool.CommandUtils;
 
@@ -255,42 +257,56 @@ namespace AutomationTool
 		/// <param name="Commands"></param>
 		public static async Task<ExitCode> ExecuteAsync(List<CommandInfo> CommandsToExecute, Dictionary<string, Type> Commands)
 		{
-			Logger.LogInformation("Executing commands...");
-			for (int CommandIndex = 0; CommandIndex < CommandsToExecute.Count; ++CommandIndex)
+			ServiceCollection serviceCollection = new ServiceCollection();
+			serviceCollection.AddLogging(builder => builder.AddEpicDefault());
+			serviceCollection.AddHorde(options => options.AllowAuthPrompt = !CommandUtils.IsBuildMachine);
+
+			await using ServiceProvider serviceProvider = serviceCollection.BuildServiceProvider();
+			try
 			{
-				var CommandInfo = CommandsToExecute[CommandIndex];
-				Logger.LogDebug("Attempting to execute {CommandInfo}", CommandInfo.ToString());
-				Type CommandType;
-				if (!Commands.TryGetValue(CommandInfo.CommandName, out CommandType))
-				{
-					throw new AutomationException("Failed to find command {0}", CommandInfo.CommandName);
-				}
+				CommandUtils.ServiceProvider = serviceProvider;
 
-				BuildCommand Command = (BuildCommand)Activator.CreateInstance(CommandType);
-				Command.Params = CommandInfo.Arguments.ToArray();
-				try
+				Logger.LogInformation("Executing commands...");
+				for (int CommandIndex = 0; CommandIndex < CommandsToExecute.Count; ++CommandIndex)
 				{
-					ExitCode Result = await Command.ExecuteAsync();
-					if(Result != ExitCode.Success)
+					var CommandInfo = CommandsToExecute[CommandIndex];
+					Logger.LogDebug("Attempting to execute {CommandInfo}", CommandInfo.ToString());
+					Type CommandType;
+					if (!Commands.TryGetValue(CommandInfo.CommandName, out CommandType))
 					{
-						return Result;
+						throw new AutomationException("Failed to find command {0}", CommandInfo.CommandName);
 					}
-					Logger.LogInformation("BUILD SUCCESSFUL");
-				}
-				finally
-				{
-					// dispose of the class if necessary
-					var CommandDisposable = Command as IDisposable;
-					if (CommandDisposable != null)
+
+					BuildCommand Command = (BuildCommand)Activator.CreateInstance(CommandType);
+					Command.Params = CommandInfo.Arguments.ToArray();
+					try
 					{
-						CommandDisposable.Dispose();
+						ExitCode Result = await Command.ExecuteAsync();
+						if (Result != ExitCode.Success)
+						{
+							return Result;
+						}
+						Logger.LogInformation("BUILD SUCCESSFUL");
+					}
+					finally
+					{
+						// dispose of the class if necessary
+						var CommandDisposable = Command as IDisposable;
+						if (CommandDisposable != null)
+						{
+							CommandDisposable.Dispose();
+						}
 					}
-				}
 
-				// Make sure there's no directories on the stack.
-				CommandUtils.ClearDirStack();
+					// Make sure there's no directories on the stack.
+					CommandUtils.ClearDirStack();
+				}
+				return ExitCode.Success;
+			}
+			finally
+			{
+				CommandUtils.ServiceProvider = null!;
 			}
-			return ExitCode.Success;
 		}
 
 		/// <summary>
diff --git a/Engine/Source/Programs/AutomationTool/AutomationUtils/AutomationUtils.Automation.csproj b/Engine/Source/Programs/AutomationTool/AutomationUtils/AutomationUtils.Automation.csproj
--- a/Engine/Source/Programs/AutomationTool/AutomationUtils/AutomationUtils.Automation.csproj
+++ b/Engine/Source/Programs/AutomationTool/AutomationUtils/AutomationUtils.Automation.csproj
@@ -54,6 +54,9 @@
   <ItemGroup>
     <ProjectReference Include="..\..\Shared\EpicGames.Core\EpicGames.Core.csproj" PrivateAssets="All"><Private>true</Private></ProjectReference>
     <ProjectReference Include="..\..\Shared\EpicGames.Build\EpicGames.Build.csproj" PrivateAssets="All"> <Private>true</Private> </ProjectReference>
+    <ProjectReference Include="..\..\Shared\EpicGames.Horde\EpicGames.Horde.csproj" PrivateAssets="All">
+      <Private>true</Private>
+    </ProjectReference>
     <ProjectReference Include="..\..\Shared\EpicGames.Perforce\EpicGames.Perforce.csproj" />
     <ProjectReference Include="..\..\UnrealBuildTool\UnrealBuildTool.csproj" PrivateAssets="All"><Private>true</Private></ProjectReference>
   </ItemGroup>
diff --git a/Engine/Source/Programs/AutomationTool/AutomationUtils/CommandUtils.cs b/Engine/Source/Programs/AutomationTool/AutomationUtils/CommandUtils.cs
--- a/Engine/Source/Programs/AutomationTool/AutomationUtils/CommandUtils.cs
+++ b/Engine/Source/Programs/AutomationTool/AutomationUtils/CommandUtils.cs
@@ -1,4 +1,7 @@
 // Copyright Epic Games, Inc. All Rights Reserved.
+
+using EpicGames.Core;
+using JetBrains.Annotations;
 using Microsoft.Extensions.Logging;
 using System;
 using System.Collections;
@@ -14,11 +17,8 @@ using System.Threading;
 using System.Threading.Tasks;
 using System.Xml;
 using System.Xml.Serialization;
-
-using EpicGames.Core;
 using UnrealBuildBase;
 using UnrealBuildTool;
-using JetBrains.Annotations;
 
 using static AutomationTool.CommandUtils;
 
@@ -37,7 +37,7 @@ namespace AutomationTool
 		}
 
 		public ParamList(ICollection<T> Collection)
-			: base(Collection != null ? Collection : new T[] {})
+			: base(Collection != null ? Collection : new T[] { })
 		{
 		}
 
@@ -50,7 +50,7 @@ namespace AutomationTool
 				{
 					Text += ", ";
 				}
-				Text += this[Index].ToString();				
+				Text += this[Index].ToString();
 			}
 			return Text;
 		}
@@ -109,6 +109,11 @@ namespace AutomationTool
 		/// </summary>
 		public static ILogger Logger => Log.Logger;
 
+		/// <summary>
+		/// Shared service provider instance, constructed on startup
+		/// </summary>
+		public static IServiceProvider ServiceProvider { get; set; } = null!;
+
 		/// <summary>
 		/// Global activity source for UAT. Can be used to add trace data and additional context to logs.
 		/// </summary>
@@ -269,7 +274,7 @@ namespace AutomationTool
 		[StringFormatMethod("Format")]
 		public static void LogWithVerbosity(EpicGames.Core.LogEventType Verbosity, string Format, params object[] Args)
 		{
-            EpicGames.Core.Log.WriteLine(Verbosity, Format, Args);
+			EpicGames.Core.Log.WriteLine(Verbosity, Format, Args);
 		}
 
 		/// <summary>
@@ -279,7 +284,7 @@ namespace AutomationTool
 		/// <param name="Message">Text</param>
 		public static void LogWithVerbosity(EpicGames.Core.LogEventType Verbosity, string Message)
 		{
-            EpicGames.Core.Log.WriteLine(Verbosity, Message);
+			EpicGames.Core.Log.WriteLine(Verbosity, Message);
 		}
 
 		/// <summary>
@@ -289,12 +294,12 @@ namespace AutomationTool
 		/// <param name="Ex">Exception</param>
 		public static void LogWithVerbosity(EpicGames.Core.LogEventType Verbosity, Exception Ex)
 		{
-            EpicGames.Core.Log.WriteLine(Verbosity, LogUtils.FormatException(Ex));
+			EpicGames.Core.Log.WriteLine(Verbosity, LogUtils.FormatException(Ex));
 		}
 
 		public static void LogPushProgress(bool bShowProgress, int Numerator, int Denominator)
 		{
-			if(bShowProgress)
+			if (bShowProgress)
 			{
 				Logger.LogInformation("[@progress push {Numerator}/{Denominator} skipline]", Numerator, Denominator);
 			}
@@ -302,7 +307,7 @@ namespace AutomationTool
 
 		public static void LogPopProgress(bool bShowProgress)
 		{
-			if(bShowProgress)
+			if (bShowProgress)
 			{
 				Logger.LogInformation("[@progress pop skipline]");
 			}
@@ -310,7 +315,7 @@ namespace AutomationTool
 
 		public static void LogIncrementProgress(bool bShowProgress, int Numerator, int Denominator)
 		{
-			if(bShowProgress)
+			if (bShowProgress)
 			{
 				Logger.LogInformation("[@progress increment {Numerator}/{Denominator} skipline]", Numerator, Denominator);
 			}
@@ -319,7 +324,7 @@ namespace AutomationTool
 		[StringFormatMethod("Format")]
 		public static void LogSetProgress(bool bShowProgress, string Format, params string[] Args)
 		{
-			if(bShowProgress)
+			if (bShowProgress)
 			{
 				Logger.LogInformation("[@progress '{Message}' skipline]", String.Format(Format, Args));
 			}
@@ -328,7 +333,7 @@ namespace AutomationTool
 		[StringFormatMethod("Format")]
 		public static void LogSetProgress(bool bShowProgress, int Numerator, int Denominator, string Format, params string[] Args)
 		{
-			if(bShowProgress)
+			if (bShowProgress)
 			{
 				Logger.LogInformation("[@progress {Numerator}/{Denominator} '{Arg2}' skipline]", Numerator, Denominator, String.Format(Format, Args));
 			}
@@ -403,21 +408,21 @@ namespace AutomationTool
 		/// <param name="PathToSearch">Paths to search</param>
 		/// <returns>An array of files found in the specified paths</returns>
 		public static string[] FindFiles_NoExceptions(bool bQuiet, string SearchPattern, bool Recursive, string PathToSearch)
-        {
-            List<string> FoundFiles = new List<string>();
+		{
+			List<string> FoundFiles = new List<string>();
 
 			var NormalizedPath = ConvertSeparators(PathSeparator.Default, PathToSearch);
-            if (DirectoryExists(NormalizedPath))
-            {
-                var FoundInPath = InternalUtils.SafeFindFiles(NormalizedPath, SearchPattern, Recursive, bQuiet);
-                if (FoundInPath != null)
-                {
-                    FoundFiles.AddRange(FoundInPath);
-                }
-            }
+			if (DirectoryExists(NormalizedPath))
+			{
+				var FoundInPath = InternalUtils.SafeFindFiles(NormalizedPath, SearchPattern, Recursive, bQuiet);
+				if (FoundInPath != null)
+				{
+					FoundFiles.AddRange(FoundInPath);
+				}
+			}
 
 			return FoundFiles.ToArray();
-        }
+		}
 		/// <summary>
 		/// Finds files in specified paths. 
 		/// </summary>
@@ -501,13 +506,13 @@ namespace AutomationTool
 		/// <param name="bQuiet">When true, logging is suppressed.</param>
 		/// <param name="FileName">Filename</param>
 		public static void DeleteFile(bool bQuiet, string FileName)
-        {
-            var NormalizedFilename = ConvertSeparators(PathSeparator.Default, FileName);
-            if (!InternalUtils.SafeDeleteFile(NormalizedFilename, bQuiet))
-            {
-                throw new AutomationException(String.Format("Failed to delete file '{0}'", NormalizedFilename));
-            }
-        }
+		{
+			var NormalizedFilename = ConvertSeparators(PathSeparator.Default, FileName);
+			if (!InternalUtils.SafeDeleteFile(NormalizedFilename, bQuiet))
+			{
+				throw new AutomationException(String.Format("Failed to delete file '{0}'", NormalizedFilename));
+			}
+		}
 
 		/// <summary>
 		/// Deletes a file(s). 
@@ -531,8 +536,8 @@ namespace AutomationTool
 		/// If the deletion of the file fails, prints a warning.
 		/// </summary>
 		/// <param name="Filename">Filename</param>
-        /// <param name="bQuiet">if true, then don't retry and don't print much.</param>
-        public static bool DeleteFile_NoExceptions(string Filename, bool bQuiet = false)
+		/// <param name="bQuiet">if true, then don't retry and don't print much.</param>
+		public static bool DeleteFile_NoExceptions(string Filename, bool bQuiet = false)
 		{
 			bool Result = true;
 			var NormalizedFilename = ConvertSeparators(PathSeparator.Default, Filename);
@@ -590,23 +595,23 @@ namespace AutomationTool
 			bool Result = true;
 
 			var NormalizedDirectory = ConvertSeparators(PathSeparator.Default, Directory);
-            try
-            {
-                if (!InternalUtils.SafeDeleteDirectory(NormalizedDirectory, bQuiet))
-                {
+			try
+			{
+				if (!InternalUtils.SafeDeleteDirectory(NormalizedDirectory, bQuiet))
+				{
 					Logger.LogWarning("Failed to delete directory '{NormalizedDirectory}'", NormalizedDirectory);
-                    Result = false;
-                }
-            }
-            catch (Exception Ex)
-            {
+					Result = false;
+				}
+			}
+			catch (Exception Ex)
+			{
 				if (!bQuiet)
 				{
 					Logger.LogWarning("Failed to delete directory, exception '{NormalizedDirectory}'", NormalizedDirectory);
 					Logger.LogWarning(Ex, "{Text}", Ex.Message);
 				}
-                Result = false;
-            }
+				Result = false;
+			}
 
 			return Result;
 		}
@@ -736,8 +741,8 @@ namespace AutomationTool
 		/// <summary>
 		/// Creates a directory. Throws an exception on failure.
 		/// </summary>
-        /// <param name="Location">Name of the directory to create</param>
-        public static void CreateDirectory(DirectoryReference Location)
+		/// <param name="Location">Name of the directory to create</param>
+		public static void CreateDirectory(DirectoryReference Location)
 		{
 			CreateDirectory(Location.FullName);
 		}
@@ -746,10 +751,10 @@ namespace AutomationTool
 		/// Renames/moves a file.
 		/// If the rename of the file fails, this function throws an Exception.
 		/// </summary>
-        /// <param name="bQuiet">When true, logging is suppressed.</param>
-        /// <param name="OldName">Old name</param>
+		/// <param name="bQuiet">When true, logging is suppressed.</param>
+		/// <param name="OldName">Old name</param>
 		/// <param name="NewName">new name</param>
-        public static void RenameFile(string OldName, string NewName, bool bQuiet = false)
+		public static void RenameFile(string OldName, string NewName, bool bQuiet = false)
 		{
 			var OldNormalized = ConvertSeparators(PathSeparator.Default, OldName);
 			var NewNormalized = ConvertSeparators(PathSeparator.Default, NewName);
@@ -783,11 +788,11 @@ namespace AutomationTool
 		/// <param name="FileName">Filename.</param>
 		/// <returns>True if the file exists, false otherwise.</returns>
 		public static bool FileExists(string FileName)
-        {
+		{
 			// need to remove the quotes before checking to see if it exists
 			var NormalizedFilename = ConvertSeparators(PathSeparator.Default, FileName).Replace("\"", "");
 			return InternalUtils.SafeFileExists(NormalizedFilename);
-        }
+		}
 
 		/// <summary>
 		/// Checks if a file(s) exists.
@@ -807,11 +812,11 @@ namespace AutomationTool
 		/// <param name="FileName">Filename.</param>
 		/// <returns>True if the file exists, false otherwise.</returns>
 		public static bool FileExists(bool bQuiet, string FileName)
-        {
+		{
 			// need to remove the quotes before checking to see if it exists
 			var NormalizedFilename = ConvertSeparators(PathSeparator.Default, FileName).Replace("\"", "");
 			return InternalUtils.SafeFileExists(NormalizedFilename, bQuiet);
-        }
+		}
 
 		/// <summary>
 		/// Checks if a file(s) exists.
@@ -820,10 +825,10 @@ namespace AutomationTool
 		/// <param name="FileName">Filename.</param>
 		/// <returns>True if the file exists, false otherwise.</returns>
 		public static bool FileExists_NoExceptions(bool bQuiet, string FileName)
-        {
-            // Standard version doesn't throw, but keep this function for consistency.
-            return FileExists(bQuiet, FileName);
-        }
+		{
+			// Standard version doesn't throw, but keep this function for consistency.
+			return FileExists(bQuiet, FileName);
+		}
 
 		static Stack<string> WorkingDirectoryStack = new Stack<string>();
 
@@ -1018,7 +1023,7 @@ namespace AutomationTool
 		/// <param name="Hidden">Hidden attribute.</param>
 		/// <param name="Archive">Archive attribute.</param>
 		public static bool SetFileAttributes_NoExceptions(string Filename, bool? ReadOnly = null, bool? Hidden = null, bool? Archive = null)
-		{			
+		{
 			Filename = ConvertSeparators(PathSeparator.Default, Filename);
 			if (!File.Exists(Filename))
 			{
@@ -1088,10 +1093,10 @@ namespace AutomationTool
 		/// <param name="Text">Text to write</param>
 		public static void WriteToFile(string Filename, string Text)
 		{
-            Filename = ConvertSeparators(PathSeparator.Default, Filename);
-            try
+			Filename = ConvertSeparators(PathSeparator.Default, Filename);
+			try
 			{
-        		        File.AppendAllText(Filename, Text + Environment.NewLine);
+				File.AppendAllText(Filename, Text + Environment.NewLine);
 			}
 			catch (Exception Ex)
 			{
@@ -1311,9 +1316,9 @@ namespace AutomationTool
 		/// </summary>
 		/// <param name="Source"></param>
 		/// <param name="Dest"></param>
-        /// <param name="bQuiet">When true, logging is suppressed.</param>
+		/// <param name="bQuiet">When true, logging is suppressed.</param>
 		/// <param name="bRetry">When true, the copy will be attempted up to 5 times before failing.</param>
-        public static void CopyFile(string Source, string Dest, bool bQuiet = false, bool bRetry = false)
+		public static void CopyFile(string Source, string Dest, bool bQuiet = false, bool bRetry = false)
 		{
 			Source = ConvertSeparators(PathSeparator.Default, Source);
 			Dest = ConvertSeparators(PathSeparator.Default, Dest);
@@ -1371,15 +1376,15 @@ namespace AutomationTool
 		/// </summary>
 		/// <param name="Source"></param>
 		/// <param name="Dest"></param>
-        /// <param name="bQuiet">When true, logging is suppressed.</param>
-        /// <returns>True if the operation was successful, false otherwise.</returns>
-        public static bool CopyFile_NoExceptions(string Source, string Dest, bool bQuiet = false)
+		/// <param name="bQuiet">When true, logging is suppressed.</param>
+		/// <returns>True if the operation was successful, false otherwise.</returns>
+		public static bool CopyFile_NoExceptions(string Source, string Dest, bool bQuiet = false)
 		{
 			Source = ConvertSeparators(PathSeparator.Default, Source);
 			Dest = ConvertSeparators(PathSeparator.Default, Dest);
-            if (InternalUtils.SafeFileExists(Dest, true))
+			if (InternalUtils.SafeFileExists(Dest, true))
 			{
-                InternalUtils.SafeDeleteFile(Dest, bQuiet);
+				InternalUtils.SafeDeleteFile(Dest, bQuiet);
 			}
 			else if (!InternalUtils.SafeDirectoryExists(Path.GetDirectoryName(Dest), true))
 			{
@@ -1403,7 +1408,7 @@ namespace AutomationTool
 		/// <returns>True if the files have the same timestamp, false otherwise.</returns>
 		public static bool IsTargetFileUpToDate(FileReference Target, FileReference Source)
 		{
-			if(FileReference.Exists(Target))
+			if (FileReference.Exists(Target))
 			{
 				TimeSpan Diff = FileReference.GetLastWriteTimeUtc(Target) - FileReference.GetLastWriteTimeUtc(Source);
 				if (Math.Abs(Diff.TotalSeconds) < 1)
@@ -1656,18 +1661,18 @@ namespace AutomationTool
 				ReaderSettings.IgnoreComments = true;
 
 				// Get the xml data stream to read from
-				XmlStream = XmlReader.Create( ManifestFile.FullName, ReaderSettings );
+				XmlStream = XmlReader.Create(ManifestFile.FullName, ReaderSettings);
 
 				// Create an object from the xml data
-				Instance = ( BuildManifest )BuildManifestSerializer.Deserialize( XmlStream );
+				Instance = (BuildManifest)BuildManifestSerializer.Deserialize(XmlStream);
 			}
-			catch( Exception Ex )
+			catch (Exception Ex)
 			{
-				Debug.WriteLine( Ex.Message );
+				Debug.WriteLine(Ex.Message);
 			}
 			finally
 			{
-				if( XmlStream != null )
+				if (XmlStream != null)
 				{
 					// Done with the file so close it
 					XmlStream.Close();
@@ -1678,15 +1683,15 @@ namespace AutomationTool
 		}
 
 		// statically create this to avoid exceptions while not leaking to much memory
-		private static readonly XmlSerializer BuildManifestSerializer = XmlSerializer.FromTypes(new[] { typeof( UnrealBuildTool.BuildManifest ) })[0];
+		private static readonly XmlSerializer BuildManifestSerializer = XmlSerializer.FromTypes(new[] { typeof(UnrealBuildTool.BuildManifest) })[0];
 
 		private static void CloneDirectoryRecursiveWorker(string SourcePathBase, string TargetPathBase, List<string> ClonedFiles, bool bIncremental = false)
 		{
 			bool bDirectoryCreated = InternalUtils.SafeCreateDirectory(TargetPathBase);
 			if (!bIncremental && !bDirectoryCreated)
-            {
-                throw new AutomationException("Failed to create directory {0} for copy", TargetPathBase);
-            }
+			{
+				throw new AutomationException("Failed to create directory {0} for copy", TargetPathBase);
+			}
 			else if (bIncremental && !CommandUtils.DirectoryExists_NoExceptions(TargetPathBase))
 			{
 				throw new AutomationException("Target directory {0} does not exist", TargetPathBase);
@@ -1703,14 +1708,14 @@ namespace AutomationTool
 
 				if (!bIncremental || !CommandUtils.FileExists(TargetFilename))
 				{
-				SourceFI.CopyTo(TargetFilename);
+					SourceFI.CopyTo(TargetFilename);
 
-				if (ClonedFiles != null)
-				{
-					ClonedFiles.Add(TargetFilename);
+					if (ClonedFiles != null)
+					{
+						ClonedFiles.Add(TargetFilename);
+					}
 				}
 			}
-			}
 
 			// Recurse into subfolders
 			foreach (DirectoryInfo SourceSubdir in SourceSubdirectories)
@@ -1747,7 +1752,7 @@ namespace AutomationTool
 			CloneDirectoryRecursiveWorker(SourcePath, TargetPath, ClonedFiles, bIncremental: true);
 		}
 
-        /// <summary>
+		/// <summary>
 		/// Copies files using multiple threads
 		/// </summary>
 		/// <param name="SourceDirectory"></param>
@@ -1756,8 +1761,8 @@ namespace AutomationTool
 		public static void ThreadedCopyFiles(string SourceDirectory, string DestDirectory, int MaxThreads = 64)
 		{
 			CreateDirectory(DestDirectory);
-            var SourceFiles = Directory.EnumerateFiles(SourceDirectory, "*", SearchOption.AllDirectories).ToList();
-            var DestFiles = SourceFiles.Select(SourceFile => CommandUtils.MakeRerootedFilePath(SourceFile, SourceDirectory, DestDirectory)).ToList();
+			var SourceFiles = Directory.EnumerateFiles(SourceDirectory, "*", SearchOption.AllDirectories).ToList();
+			var DestFiles = SourceFiles.Select(SourceFile => CommandUtils.MakeRerootedFilePath(SourceFile, SourceDirectory, DestDirectory)).ToList();
 			ThreadedCopyFiles(SourceFiles, DestFiles, MaxThreads);
 		}
 
@@ -1809,20 +1814,20 @@ namespace AutomationTool
 		/// <param name="bRetry"></param>
 		public static void ThreadedCopyFiles(List<string> Source, List<string> Dest, int MaxThreads = 64, bool bQuiet = false, bool bRetry = false)
 		{
-			if(!bQuiet)
+			if (!bQuiet)
 			{
 				Logger.LogInformation("Copying {NumFiles} file(s) using max {MaxThreads} thread(s)", Source.Count, MaxThreads);
 			}
 
-            if (Source.Count != Dest.Count)
+			if (Source.Count != Dest.Count)
 			{
-                throw new AutomationException("Source count ({0}) does not match Dest count ({1})", Source.Count, Dest.Count);
+				throw new AutomationException("Source count ({0}) does not match Dest count ({1})", Source.Count, Dest.Count);
 			}
 			Parallel.ForEach(Source.Zip(Dest, (Src, Dst) => new { SourceFile = Src, DestFile = Dst }), new ParallelOptions { MaxDegreeOfParallelism = MaxThreads }, (Pair) =>
 			{
 				CommandUtils.CopyFile(Pair.SourceFile, Pair.DestFile, true, bRetry);
 			});
-        }
+		}
 
 		/// <summary>
 		/// Copies a set of files from one folder to another
@@ -1835,8 +1840,8 @@ namespace AutomationTool
 		/// <returns>List of filenames copied to the target directory</returns>
 		public static List<string> ThreadedCopyFiles(string SourceDir, string TargetDir, List<string> RelativePaths, int MaxThreads = 64, bool bRetry = false)
 		{
-            var SourceFileNames = RelativePaths.Select(RelativePath => CommandUtils.CombinePaths(SourceDir, RelativePath)).ToList();
-            var TargetFileNames = RelativePaths.Select(RelativePath => CommandUtils.CombinePaths(TargetDir, RelativePath)).ToList();
+			var SourceFileNames = RelativePaths.Select(RelativePath => CommandUtils.CombinePaths(SourceDir, RelativePath)).ToList();
+			var TargetFileNames = RelativePaths.Select(RelativePath => CommandUtils.CombinePaths(TargetDir, RelativePath)).ToList();
 			CommandUtils.ThreadedCopyFiles(SourceFileNames, TargetFileNames, MaxThreads, bRetry: bRetry);
 			return TargetFileNames;
 		}
@@ -1862,12 +1867,12 @@ namespace AutomationTool
 
 		/// <summary>
 		/// Moves files in parallel
-        /// </summary>
+		/// </summary>
 		/// <param name="SourceAndTargetPairs">Pairs of source and target files</param>
 		public static void ParallelMoveFiles(IEnumerable<KeyValuePair<FileReference, FileReference>> SourceAndTargetPairs)
 		{
 			ParallelMoveFiles(SourceAndTargetPairs, false);
-        }
+		}
 
 		/// <summary>
 		/// Moves files in parallel
@@ -2088,15 +2093,15 @@ namespace AutomationTool
 		/// <returns>True if param was found, false otherwise.</returns>
 		public static bool ParseParam(string[] ArgList, string Param)
 		{
-            string ValueParam = Param;
-            if (!ValueParam.EndsWith("="))
-            {
-                ValueParam += "=";
-            }
+			string ValueParam = Param;
+			if (!ValueParam.EndsWith("="))
+			{
+				ValueParam += "=";
+			}
 
-            foreach (string ArgStr in ArgList)
+			foreach (string ArgStr in ArgList)
 			{
-                if (ArgStr.Equals(Param, StringComparison.InvariantCultureIgnoreCase) || ArgStr.StartsWith(ValueParam, StringComparison.InvariantCultureIgnoreCase))
+				if (ArgStr.Equals(Param, StringComparison.InvariantCultureIgnoreCase) || ArgStr.StartsWith(ValueParam, StringComparison.InvariantCultureIgnoreCase))
 				{
 					return true;
 				}
@@ -2192,7 +2197,7 @@ namespace AutomationTool
 		/// <summary>
 		/// Return the main engine directory and any platform extension engine directories
 		/// </summary>
- 		public static DirectoryReference[] GetAllEngineDirectories()
+		public static DirectoryReference[] GetAllEngineDirectories()
 		{
 			List<DirectoryReference> EngineDirectories = new List<DirectoryReference>() { Unreal.EngineDirectory };
 			DirectoryReference EnginePlatformsDirectory = DirectoryReference.Combine(Unreal.EngineDirectory, "Platforms");
@@ -2224,27 +2229,27 @@ namespace AutomationTool
 		/// </summary>
 		public static TelemetryData Telemetry = new TelemetryData();
 
-        /// <summary>
-        /// Cached location of the build root storage because the logic to compute it is a little non-trivial.
-        /// </summary>
-        private static string CachedRootBuildStorageDirectory;
-
-        /// <summary>
-        /// "P:\Builds" or "/Volumes/Builds". Root Folder for all build storage.
-        /// </summary>
-        /// <returns>"P:\Builds" or "/Volumes/Builds" or "/mnt/Builds" unless overridden by -UseLocalBuildStorage from the commandline, where is uses Engine\Saved\LocalBuilds\.</returns>
-        public static string RootBuildStorageDirectory()
-        {
-            if (string.IsNullOrEmpty(CachedRootBuildStorageDirectory))
-            {
-                if (GlobalCommandLine.UseLocalBuildStorage)
-                {
-                    CachedRootBuildStorageDirectory = CombinePaths(CmdEnv.LocalRoot, "Engine", "Saved", "LocalBuilds");
-                    // Must create the directory because much of the system assumes it is already there.
-                    CreateDirectory(CombinePaths(CachedRootBuildStorageDirectory, "UE4"));
-                }
-                else
-                {
+		/// <summary>
+		/// Cached location of the build root storage because the logic to compute it is a little non-trivial.
+		/// </summary>
+		private static string CachedRootBuildStorageDirectory;
+
+		/// <summary>
+		/// "P:\Builds" or "/Volumes/Builds". Root Folder for all build storage.
+		/// </summary>
+		/// <returns>"P:\Builds" or "/Volumes/Builds" or "/mnt/Builds" unless overridden by -UseLocalBuildStorage from the commandline, where is uses Engine\Saved\LocalBuilds\.</returns>
+		public static string RootBuildStorageDirectory()
+		{
+			if (string.IsNullOrEmpty(CachedRootBuildStorageDirectory))
+			{
+				if (GlobalCommandLine.UseLocalBuildStorage)
+				{
+					CachedRootBuildStorageDirectory = CombinePaths(CmdEnv.LocalRoot, "Engine", "Saved", "LocalBuilds");
+					// Must create the directory because much of the system assumes it is already there.
+					CreateDirectory(CombinePaths(CachedRootBuildStorageDirectory, "UE4"));
+				}
+				else
+				{
 					if (RuntimePlatform.IsMac)
 					{
 						CachedRootBuildStorageDirectory = CombinePaths("/Volumes", "Builds");
@@ -2257,35 +2262,35 @@ namespace AutomationTool
 					{
 						CachedRootBuildStorageDirectory = CombinePaths("P:", "Builds");
 					}
-                }
-            }
-            return CachedRootBuildStorageDirectory;
-        }
-
-        public static bool DirectoryExistsAndIsWritable_NoExceptions(string Dir)
-        {
-            if (!DirectoryExists_NoExceptions(Dir))
-            {
+				}
+			}
+			return CachedRootBuildStorageDirectory;
+		}
+
+		public static bool DirectoryExistsAndIsWritable_NoExceptions(string Dir)
+		{
+			if (!DirectoryExists_NoExceptions(Dir))
+			{
 				Logger.LogDebug("Directory {Dir} does not exist", Dir);
 				return false;
 			}
 
-            try
-            {
+			try
+			{
 				string Filename = CombinePaths(Dir, Guid.NewGuid().ToString() + ".Temp.txt");
 				string NativeFilename = ConvertSeparators(PathSeparator.Default, Filename);
-				using(StreamWriter Writer = new StreamWriter(NativeFilename))
+				using (StreamWriter Writer = new StreamWriter(NativeFilename))
 				{
 					Writer.Write("Test");
 				}
-				if(File.Exists(NativeFilename))
+				if (File.Exists(NativeFilename))
 				{
-		            DeleteFile_NoExceptions(Filename, true);
+					DeleteFile_NoExceptions(Filename, true);
 					Logger.LogDebug("Directory {Dir} is writable", Dir);
 					return true;
 				}
 			}
-			catch(Exception)
+			catch (Exception)
 			{
 			}
 
@@ -2294,35 +2299,35 @@ namespace AutomationTool
 		}
 
 		public static void CleanFormalBuilds(string ParentDir, string SearchPattern, int MaximumDaysToKeepTempStorage = 4)
-        {
-            if (!IsBuildMachine || !ParentDir.StartsWith(RootBuildStorageDirectory()))
-            {
-                return;
-            }
-            try
-            {
-                DirectoryInfo DirInfo = new DirectoryInfo(ParentDir);
+		{
+			if (!IsBuildMachine || !ParentDir.StartsWith(RootBuildStorageDirectory()))
+			{
+				return;
+			}
+			try
+			{
+				DirectoryInfo DirInfo = new DirectoryInfo(ParentDir);
 				Logger.LogInformation("Looking for directories to delete in {ParentDir}", ParentDir);
-                foreach (DirectoryInfo ThisDirInfo in DirInfo.EnumerateDirectories(SearchPattern))
-                {
+				foreach (DirectoryInfo ThisDirInfo in DirInfo.EnumerateDirectories(SearchPattern))
+				{
 					double AgeDays = (DateTime.UtcNow - ThisDirInfo.CreationTimeUtc).TotalDays;
 					if (AgeDays > MaximumDaysToKeepTempStorage)
-                    {
+					{
 						Logger.LogInformation("Deleting formal build directory {Directory}, because it is {Age} days old (maximum {MaximumDaysToKeepTempStorage}).", ThisDirInfo.FullName, (int)AgeDays, MaximumDaysToKeepTempStorage);
 						DeleteDirectory_NoExceptions(true, ThisDirInfo.FullName);
 					}
-                    else
-                    {
+					else
+					{
 						Logger.LogDebug("Not deleting formal build directory {Directory}, because it is {Age} days old (maximum {MaximumDaysToKeepTempStorage}).", ThisDirInfo.FullName, (int)AgeDays, MaximumDaysToKeepTempStorage);
-                    }
-                }
-            }
-            catch (Exception Ex)
-            {
+					}
+				}
+			}
+			catch (Exception Ex)
+			{
 				Logger.LogWarning("Unable to clean formal builds from directory: {ParentDir}", ParentDir);
 				Logger.LogWarning(Ex, " Exception was {Exception}", LogUtils.FormatException(Ex));
-            }
-        }
+			}
+		}
 
 
 		/// <summary>
@@ -2331,11 +2336,11 @@ namespace AutomationTool
 		/// <param name="Platform">Specific platform</param>
 		public static string GetGenericPlatformName(UnrealBuildTool.UnrealTargetPlatform Platform)
 		{
-			if(Platform == UnrealTargetPlatform.Win64)
+			if (Platform == UnrealTargetPlatform.Win64)
 			{
 				return "Windows";
 			}
-			else if(Platform == UnrealTargetPlatform.Linux || Platform == UnrealTargetPlatform.LinuxArm64)
+			else if (Platform == UnrealTargetPlatform.Linux || Platform == UnrealTargetPlatform.LinuxArm64)
 			{
 				return "Linux";
 			}
@@ -2433,7 +2438,7 @@ namespace AutomationTool
 			{
 				foreach (ZipArchiveEntry Entry in ZipArchive.Entries)
 				{
-					if(Entry.FullName.EndsWith("/"))
+					if (Entry.FullName.EndsWith("/"))
 					{
 						// ignore directories
 						continue;
@@ -2538,18 +2543,18 @@ namespace AutomationTool
 
 			// Check if it contains any wildcards. If not, we can just add the pattern directly without searching.
 			int WildcardIdx = FileFilter.FindWildcardIndex(Pattern);
-			if(WildcardIdx == -1)
+			if (WildcardIdx == -1)
 			{
 				// Construct a filter which removes all the excluded filetypes
 				FileFilter Filter = new FileFilter(FileFilterType.Include);
-				if(ExcludePatterns != null)
+				if (ExcludePatterns != null)
 				{
 					Filter.AddRules(ExcludePatterns, FileFilterType.Exclude);
 				}
 
 				// Match it against the given file
 				FileReference File = FileReference.Combine(DefaultDir, Pattern);
-				if(Filter.Matches(File.FullName))
+				if (Filter.Matches(File.FullName))
 				{
 					Files.Add(File);
 				}
@@ -2558,18 +2563,18 @@ namespace AutomationTool
 			{
 				// Find the base directory for the search. We construct this in a very deliberate way including the directory separator itself, so matches
 				// against the OS root directory will resolve correctly both on Mac (where / is the filesystem root) and Windows (where / refers to the current drive).
-				int LastDirectoryIdx = Pattern.LastIndexOfAny(new char[]{ Path.DirectorySeparatorChar, Path.AltDirectorySeparatorChar }, WildcardIdx);
+				int LastDirectoryIdx = Pattern.LastIndexOfAny(new char[] { Path.DirectorySeparatorChar, Path.AltDirectorySeparatorChar }, WildcardIdx);
 				DirectoryReference BaseDir = DirectoryReference.Combine(DefaultDir, Pattern.Substring(0, LastDirectoryIdx + 1));
 
 				// Construct the absolute include pattern to match against, re-inserting the resolved base directory to construct a canonical path.
-				string IncludePattern = BaseDir.FullName.TrimEnd(new char[]{ Path.DirectorySeparatorChar, Path.AltDirectorySeparatorChar }) + "/" + Pattern.Substring(LastDirectoryIdx + 1);
+				string IncludePattern = BaseDir.FullName.TrimEnd(new char[] { Path.DirectorySeparatorChar, Path.AltDirectorySeparatorChar }) + "/" + Pattern.Substring(LastDirectoryIdx + 1);
 
 				// Construct a filter and apply it to the directory
-				if(DirectoryReference.Exists(BaseDir))
+				if (DirectoryReference.Exists(BaseDir))
 				{
 					FileFilter Filter = new FileFilter();
 					Filter.AddRule(IncludePattern, FileFilterType.Include);
-					if(ExcludePatterns != null)
+					if (ExcludePatterns != null)
 					{
 						Filter.AddRules(ExcludePatterns, FileFilterType.Exclude);
 					}
@@ -2593,7 +2598,7 @@ namespace AutomationTool
 
 			return String.Format("{0} {1}", AbsSize.ToString("N2"), Units[UnitIndex]);
 		}
-		
+
 		public static void TakeLock(DirectoryReference LockDirectory, TimeSpan Timeout, System.Action Callback)
 		{
 			string LockFilePath = Path.Combine(LockDirectory.FullName, ".lock");
@@ -2779,7 +2784,7 @@ namespace AutomationTool
 		/// <param name="Other">The other telemetry data</param>
 		public void Merge(string Prefix, TelemetryData Other)
 		{
-			foreach(TelemetrySample Sample in Other.Samples)
+			foreach (TelemetrySample Sample in Other.Samples)
 			{
 				Add(Sample.Name, Sample.Value, Sample.Units);
 			}
@@ -2803,7 +2808,7 @@ namespace AutomationTool
 
 			// Check the version is valid
 			int Version;
-			if(!RawObject.TryGetIntegerField("Version", out Version) || Version != CurrentVersion)
+			if (!RawObject.TryGetIntegerField("Version", out Version) || Version != CurrentVersion)
 			{
 				Telemetry = null;
 				return false;
@@ -2838,7 +2843,7 @@ namespace AutomationTool
 				Writer.WriteObjectStart();
 				Writer.WriteValue("Version", CurrentVersion);
 				Writer.WriteArrayStart("Samples");
-				foreach(TelemetrySample Sample in Samples)
+				foreach (TelemetrySample Sample in Samples)
 				{
 					Writer.WriteObjectStart();
 					Writer.WriteValue("Name", Sample.Name);
@@ -2856,81 +2861,81 @@ namespace AutomationTool
 	/// Timer class used for telemetry reporting.
 	/// </summary>
 	public class TelemetryStopwatch : IDisposable
-    {
-        string Name;
-        DateTime StartTime;
-        bool bFinished;
+	{
+		string Name;
+		DateTime StartTime;
+		bool bFinished;
 
 		[StringFormatMethod("Format")]
-        public TelemetryStopwatch(string Format, params object[] Args)
-        {
-            Name = String.Format(Format, Args);
-            StartTime = DateTime.Now;
-        }
-
-        public void Cancel()
-        {
-            bFinished = true;
-        }
-
-        /// <summary>
-        /// Flushes the time to <see cref="CommandEnvironment.CSVFile"/> if we are the build machine and that environment variable is specified.
-        /// Call Finish manually with an alternate name to use that one instead. Useful for dynamically generated names that you can't specify at construction.
-        /// </summary>
-        /// <param name="AlternateName">Used in place of the Name specified during construction.</param>
-        public void Finish(string AlternateName = null)
-        {
-            if(!bFinished)
-            {
-                if (!String.IsNullOrEmpty(AlternateName))
-                {
-                    Name = AlternateName;
-                }
-
-                var OutputStr = $"UAT,{Name},{StartTime},{DateTime.Now}" + Environment.NewLine;
-                Logger.LogDebug("{CSVHeader}", OutputStr);
-                if (CommandUtils.IsBuildMachine && !String.IsNullOrEmpty(CommandUtils.CmdEnv.CSVFile) && CommandUtils.CmdEnv.CSVFile != "nul")
-                {
-                    try
-                    {
-                        File.AppendAllText(CommandUtils.CmdEnv.CSVFile, OutputStr);
-                    }
-                    catch (Exception Ex)
-                    {
-                        Logger.LogWarning("Could not append to csv file ({File}) : {Exception}", CommandUtils.CmdEnv.CSVFile, Ex.ToString());
-                    }
-                }
-            }
-            bFinished = true;
-        }
-
-        public void Dispose()
-        {
-            Finish();
-        }
-    }
-
-    /// <summary>
-    /// Stopwatch that uses DateTime.UtcNow for timing. Not hi-res, but also not subject to short time limitations of System.Diagnostics.Stopwatch.
-    /// </summary>
-    public class DateTimeStopwatch
-    {
-        public static DateTimeStopwatch Start()
-        {
-            return new DateTimeStopwatch();
-        }
-
-        /// <summary>
-        /// Hide public ctor.
-        /// </summary>
-        private DateTimeStopwatch() { }
-
-        readonly DateTime StartTime = DateTime.UtcNow;
-
-        public TimeSpan ElapsedTime { get { return DateTime.UtcNow - StartTime; } }
-    }
-
-    /// <summary>
+		public TelemetryStopwatch(string Format, params object[] Args)
+		{
+			Name = String.Format(Format, Args);
+			StartTime = DateTime.Now;
+		}
+
+		public void Cancel()
+		{
+			bFinished = true;
+		}
+
+		/// <summary>
+		/// Flushes the time to <see cref="CommandEnvironment.CSVFile"/> if we are the build machine and that environment variable is specified.
+		/// Call Finish manually with an alternate name to use that one instead. Useful for dynamically generated names that you can't specify at construction.
+		/// </summary>
+		/// <param name="AlternateName">Used in place of the Name specified during construction.</param>
+		public void Finish(string AlternateName = null)
+		{
+			if (!bFinished)
+			{
+				if (!String.IsNullOrEmpty(AlternateName))
+				{
+					Name = AlternateName;
+				}
+
+				var OutputStr = $"UAT,{Name},{StartTime},{DateTime.Now}" + Environment.NewLine;
+				Logger.LogDebug("{CSVHeader}", OutputStr);
+				if (CommandUtils.IsBuildMachine && !String.IsNullOrEmpty(CommandUtils.CmdEnv.CSVFile) && CommandUtils.CmdEnv.CSVFile != "nul")
+				{
+					try
+					{
+						File.AppendAllText(CommandUtils.CmdEnv.CSVFile, OutputStr);
+					}
+					catch (Exception Ex)
+					{
+						Logger.LogWarning("Could not append to csv file ({File}) : {Exception}", CommandUtils.CmdEnv.CSVFile, Ex.ToString());
+					}
+				}
+			}
+			bFinished = true;
+		}
+
+		public void Dispose()
+		{
+			Finish();
+		}
+	}
+
+	/// <summary>
+	/// Stopwatch that uses DateTime.UtcNow for timing. Not hi-res, but also not subject to short time limitations of System.Diagnostics.Stopwatch.
+	/// </summary>
+	public class DateTimeStopwatch
+	{
+		public static DateTimeStopwatch Start()
+		{
+			return new DateTimeStopwatch();
+		}
+
+		/// <summary>
+		/// Hide public ctor.
+		/// </summary>
+		private DateTimeStopwatch() { }
+
+		readonly DateTime StartTime = DateTime.UtcNow;
+
+		public TimeSpan ElapsedTime { get { return DateTime.UtcNow - StartTime; } }
+	}
+
+	/// <summary>
 	/// Use with "using" syntax to push and pop directories in a convenient, exception-safe way
 	/// </summary>
 	public class PushedDirectory : IDisposable
@@ -2967,39 +2972,39 @@ namespace AutomationTool
 		}
 	}
 
-    /// <summary>
-    /// Helper class to associate a file and its contents
-    /// </summary>
-    public class EMSFileInfo
-    {
-        public string FileName { get; set; }
-        public byte[] Bytes { get; set; }
-    }
+	/// <summary>
+	/// Helper class to associate a file and its contents
+	/// </summary>
+	public class EMSFileInfo
+	{
+		public string FileName { get; set; }
+		public byte[] Bytes { get; set; }
+	}
 
-    /// <summary>
-    /// Wrapper class for the enumerate files JSON response from MCP
-    /// </summary>
-    [DataContract]
-    public sealed class EnumerationResponse
-    {
-        [DataMember(Name = "doNotCache", IsRequired = true)]
-        public Boolean DoNotCache { get; set; }
+	/// <summary>
+	/// Wrapper class for the enumerate files JSON response from MCP
+	/// </summary>
+	[DataContract]
+	public sealed class EnumerationResponse
+	{
+		[DataMember(Name = "doNotCache", IsRequired = true)]
+		public Boolean DoNotCache { get; set; }
 
-        [DataMember(Name = "uniqueFilename", IsRequired = true)]
-        public string UniqueFilename { get; set; }
+		[DataMember(Name = "uniqueFilename", IsRequired = true)]
+		public string UniqueFilename { get; set; }
 
-        [DataMember(Name = "filename", IsRequired = true)]
-        public string Filename { get; set; }
+		[DataMember(Name = "filename", IsRequired = true)]
+		public string Filename { get; set; }
 
-        [DataMember(Name = "hash", IsRequired = true)]
-        public string Hash { get; set; }
+		[DataMember(Name = "hash", IsRequired = true)]
+		public string Hash { get; set; }
 
-        [DataMember(Name = "length", IsRequired = true)]
-        public long Length { get; set; }
+		[DataMember(Name = "length", IsRequired = true)]
+		public long Length { get; set; }
 
-        [DataMember(Name = "uploaded", IsRequired = true)]
-        public string Uploaded { get; set; }
-    }
+		[DataMember(Name = "uploaded", IsRequired = true)]
+		public string Uploaded { get; set; }
+	}
 
 	static class CodeSignWindows
 	{
@@ -3057,18 +3062,18 @@ namespace AutomationTool
 
 			List<Task> SignTaskList = new List<Task>();
 
-			for(int FileIdx = 0; FileIdx < Files.Count; )
+			for (int FileIdx = 0; FileIdx < Files.Count;)
 			{
 				//@TODO: Verbosity choosing
 				//  /v will spew lots of info
 				//  /q does nothing on success and minimal output on failure
 				StringBuilder CommandLine = new StringBuilder();
 				StringBuilder FileList = new StringBuilder();
-				if(SignatureType == SignatureType.SHA1)
+				if (SignatureType == SignatureType.SHA1)
 				{
 					CommandLine.Append($"sign{SpecificStoreArg} /a /n \"{SigningIdentity}\" /v {MultipleSignatureArg} {DescriptionArg}");
 				}
-				else if(SignatureType == SignatureType.SHA256)
+				else if (SignatureType == SignatureType.SHA256)
 				{
 					CommandLine.Append($"sign{SpecificStoreArg} /a /fd sha256 /td sha256 /as /n \"{SigningIdentity}\" {DescriptionArg}");
 				}
@@ -3083,7 +3088,7 @@ namespace AutomationTool
 				int NextFileIdx = FileIdx;
 				int FileCount = 0;
 				int TimeStampLength = SignatureType == SignatureType.SHA1 ? SHA1TimestampArg.Length + LongestSHA1 : SHA256TimestampArg.Length + LongestSHA256;
-				while(NextFileIdx < Files.Count && (SignToolPath.Length + TimeStampLength + CommandLine.Length + FileList.Length + Files[NextFileIdx].FullName.Length < 2080))
+				while (NextFileIdx < Files.Count && (SignToolPath.Length + TimeStampLength + CommandLine.Length + FileList.Length + Files[NextFileIdx].FullName.Length < 2080))
 				{
 					FileList.Append($" \"{Files[NextFileIdx]}\"");
 					NextFileIdx++;
@@ -3163,9 +3168,9 @@ namespace AutomationTool
 				"C:/Program Files (x86)/Windows Kits/10/bin/x86/SignTool.exe"
 			};
 
-			foreach(string PossibleSignToolName in PossibleSignToolNames)
+			foreach (string PossibleSignToolName in PossibleSignToolNames)
 			{
-				if(File.Exists(PossibleSignToolName))
+				if (File.Exists(PossibleSignToolName))
 				{
 					return PossibleSignToolName;
 				}
@@ -3186,12 +3191,12 @@ namespace AutomationTool
 		/// </summary>
 		public static void SignSingleExecutableIfEXEOrDLL(string Filename, bool bIgnoreExtension = false, string Description = null)
 		{
-            if (!OperatingSystem.IsWindows())
-            {
-                Logger.LogDebug("Can't sign '{File}' on non-Windows platform.", Filename);
-                return;
-            }
-            if (!CommandUtils.FileExists(Filename))
+			if (!OperatingSystem.IsWindows())
+			{
+				Logger.LogDebug("Can't sign '{File}' on non-Windows platform.", Filename);
+				return;
+			}
+			if (!CommandUtils.FileExists(Filename))
 			{
 				throw new AutomationException("Can't sign '{0}', file does not exist.", Filename);
 			}
diff --git a/Engine/Source/Programs/AutomationTool/AutomationUtils/HordeUtils.cs b/Engine/Source/Programs/AutomationTool/AutomationUtils/HordeUtils.cs
new file mode 100644
--- /dev/null
+++ b/Engine/Source/Programs/AutomationTool/AutomationUtils/HordeUtils.cs
@@ -0,0 +1,75 @@
+// Copyright Epic Games, Inc. All Rights Reserved.
+
+using EpicGames.Horde;
+using EpicGames.Horde.Secrets;
+using Microsoft.Extensions.DependencyInjection;
+using System.Collections.Concurrent;
+using System.Collections.Generic;
+using System.Net.Http;
+using System.Threading;
+using System.Threading.Tasks;
+
+#nullable enable
+
+namespace AutomationTool
+{
+	/// <summary>
+	/// Helper methods for Horde functionality
+	/// </summary>
+	public partial class CommandUtils
+	{
+		static readonly ConcurrentDictionary<SecretId, IReadOnlyDictionary<string, string>> _cachedSecrets = new ConcurrentDictionary<SecretId, IReadOnlyDictionary<string, string>>();
+
+		/// <summary>
+		/// Reads a secret from Horde
+		/// </summary>
+		/// <param name="secretId">Identifier of the secret to retrieve</param>
+		/// <param name="cancellationToken">Cancellation token for the operation</param>
+		/// <returns>Values configured for the given secret</returns>
+		public static async Task<IReadOnlyDictionary<string, string>> GetHordeSecretAsync(SecretId secretId, CancellationToken cancellationToken = default)
+		{
+			IReadOnlyDictionary<string, string>? secret;
+			if (!_cachedSecrets.TryGetValue(secretId, out secret))
+			{
+				IHordeClient hordeClient = ServiceProvider.GetRequiredService<IHordeClient>();
+
+				using HordeHttpClient hordeHttpClient = hordeClient.CreateHttpClient();
+
+				GetSecretResponse response;
+				try
+				{
+					response = await hordeHttpClient.GetSecretAsync(secretId);
+				}
+				catch (HttpRequestException ex) when (ex.StatusCode == System.Net.HttpStatusCode.NotFound)
+				{
+					throw new AutomationException(ex, $"Secret '{secretId}' was not found on {hordeClient.ServerUrl}");
+				}
+				catch (HttpRequestException ex) when (ex.StatusCode == System.Net.HttpStatusCode.Forbidden)
+				{
+					throw new AutomationException(ex, $"User does not have permissions to read '{secretId}' on {hordeClient.ServerUrl}");
+				}
+
+				secret = response.Data;
+				_cachedSecrets.TryAdd(secretId, secret);
+			}
+			return secret;
+		}
+
+		/// <summary>
+		/// Reads a secret from Horde
+		/// </summary>
+		/// <param name="secretId">Identifier of the secret to retrieve</param>
+		/// <param name="propertyName">Name of the property within the secret</param>
+		/// <param name="cancellationToken">Cancellation token for the operation</param>
+		/// <returns>Values configured for the given secret</returns>
+		public static async Task<string> GetHordeSecretAsync(SecretId secretId, string propertyName, CancellationToken cancellationToken = default)
+		{
+			IReadOnlyDictionary<string, string> secrets = await GetHordeSecretAsync(secretId, cancellationToken);
+			if (!secrets.TryGetValue(propertyName, out string? propertyValue))
+			{
+				throw new AutomationException("Secret '{secretId}' does not contain a value for '{propertyName}'");
+			}
+			return propertyValue;
+		}
+	}
+}
