Horde: Allow artifacts to be declared in macros.

#rnx

diff --git a/Engine/Source/Programs/AutomationTool/BuildGraph/BgSchema.cs b/Engine/Source/Programs/AutomationTool/BuildGraph/BgSchema.cs
--- a/Engine/Source/Programs/AutomationTool/BuildGraph/BgSchema.cs
+++ b/Engine/Source/Programs/AutomationTool/BuildGraph/BgSchema.cs
@@ -898,6 +898,7 @@ namespace AutomationTool
 			macroChoice.Items.Add(CreateSchemaElement("Macro", ScriptSchemaStandardType.Macro));
 			macroChoice.Items.Add(CreateSchemaElement("Agent", ScriptSchemaStandardType.Agent));
 			macroChoice.Items.Add(CreateSchemaElement("Aggregate", ScriptSchemaStandardType.Aggregate));
+			macroChoice.Items.Add(CreateSchemaElement("Artifact", ScriptSchemaStandardType.Artifact));
 			macroChoice.Items.Add(CreateSchemaElement("Report", ScriptSchemaStandardType.Report));
 			macroChoice.Items.Add(CreateSchemaElement("Badge", ScriptSchemaStandardType.Badge));
 			macroChoice.Items.Add(CreateSchemaElement("Notify", ScriptSchemaStandardType.Notify));
