AutomationTool: Flag "RunUBT ERROR:" messages as an ignorable exit code error.

diff --git a/Engine/Source/Programs/AutomationTool/AutomationUtils/Matchers/ExitCodeEventMatcher.cs b/Engine/Source/Programs/AutomationTool/AutomationUtils/Matchers/ExitCodeEventMatcher.cs
--- a/Engine/Source/Programs/AutomationTool/AutomationUtils/Matchers/ExitCodeEventMatcher.cs
+++ b/Engine/Source/Programs/AutomationTool/AutomationUtils/Matchers/ExitCodeEventMatcher.cs
@@ -14,7 +14,9 @@ namespace AutomationUtils.Matchers
 	class ExitCodeEventMatcher : ILogEventMatcher
 	{
 		static readonly Regex s_pattern = new Regex(
-			@"Editor terminated with exit code -?[1-9]|(Error executing.+)(tool returned code)(.+)");
+			@"Editor terminated with exit code -?[1-9]" +
+			@"|(?:Error executing.+)(?:tool returned code)(.+)" +
+			@"|\s*RunU[AB]T ERROR:");
 
 		public LogEventMatch? Match(ILogCursor cursor)
 		{
