BuildGraph: Include block names in errors when artifacts are produced by two things.

diff --git a/Engine/Source/Programs/AutomationTool/BuildGraph/BuildGraph.cs b/Engine/Source/Programs/AutomationTool/BuildGraph/BuildGraph.cs
--- a/Engine/Source/Programs/AutomationTool/BuildGraph/BuildGraph.cs
+++ b/Engine/Source/Programs/AutomationTool/BuildGraph/BuildGraph.cs
@@ -1104,20 +1104,20 @@ namespace AutomationTool
 			}
 
 			// Read all the input storage blocks, keeping track of which block each file came from
-			Dictionary<string, TempStorageFile> inputFiles = new Dictionary<string, TempStorageFile>(FileReference.Comparer);
+			Dictionary<string, (TempStorageFile, TempStorageBlockRef)> inputFiles = new Dictionary<string, (TempStorageFile, TempStorageBlockRef)>(FileReference.Comparer);
 			foreach (KeyValuePair<TempStorageBlockRef, TempStorageBlockManifest> pair in inputManifests)
 			{
 				foreach (TempStorageFile newFile in pair.Value.Files)
 				{
-					TempStorageFile? existingFile;
-					if (inputFiles.TryGetValue(newFile.RelativePath, out existingFile) && !TempStorage.IsDuplicateBuildProduct(newFile.ToFileReference(rootDir)))
+					(TempStorageFile File, TempStorageBlockRef Block) existingItem;
+					if (inputFiles.TryGetValue(newFile.RelativePath, out existingItem) && !TempStorage.IsDuplicateBuildProduct(newFile.ToFileReference(rootDir)))
 					{
-						if (existingFile.LastWriteTimeUtcTicks != newFile.LastWriteTimeUtcTicks)
+						if (existingItem.File.LastWriteTimeUtcTicks != newFile.LastWriteTimeUtcTicks)
 						{
-							Logger.LogError("File '{File}' was produced by multiple nodes", newFile.RelativePath);// {InputStorageBlock} and {CurrentStorageBlock}", file, inputStorageBlock, currentStorageBlock);
+							Logger.LogError("File '{File}' was produced by {InputStorageBlock} and {CurrentStorageBlock}", newFile.RelativePath, existingItem.Block, pair.Key);
 						}
 					}
-					inputFiles[newFile.RelativePath] = newFile;
+					inputFiles[newFile.RelativePath] = (newFile, pair.Key);
 				}
 			}
 
