BuildGraph: Exclude files in the IgnoreModified list for a node from causing errors about mismatched timestamps.

#rnx

diff --git a/Engine/Source/Programs/AutomationTool/BuildGraph/BuildGraph.cs b/Engine/Source/Programs/AutomationTool/BuildGraph/BuildGraph.cs
--- a/Engine/Source/Programs/AutomationTool/BuildGraph/BuildGraph.cs
+++ b/Engine/Source/Programs/AutomationTool/BuildGraph/BuildGraph.cs
@@ -1110,7 +1110,9 @@ namespace AutomationTool
 				foreach (TempStorageFile newFile in pair.Value.Files)
 				{
 					(TempStorageFile File, TempStorageBlockRef Block) existingItem;
-					if (inputFiles.TryGetValue(newFile.RelativePath, out existingItem) && !TempStorage.IsDuplicateBuildProduct(newFile.ToFileReference(rootDir)))
+					if (inputFiles.TryGetValue(newFile.RelativePath, out existingItem)
+						&& !ignoreModifiedFilter.Matches(newFile.ToFileReference(rootDir).FullName)
+						&& !TempStorage.IsDuplicateBuildProduct(newFile.ToFileReference(rootDir)))
 					{
 						if (existingItem.File.LastWriteTimeUtcTicks != newFile.LastWriteTimeUtcTicks)
 						{
