[Gauntlet] Allow users to opt out of screenshot artifacts being converted to jpegs

diff --git a/Engine/Source/Programs/AutomationTool/Gauntlet/Unreal/Base/Gauntlet.UnrealSession.cs b/Engine/Source/Programs/AutomationTool/Gauntlet/Unreal/Base/Gauntlet.UnrealSession.cs
--- a/Engine/Source/Programs/AutomationTool/Gauntlet/Unreal/Base/Gauntlet.UnrealSession.cs
+++ b/Engine/Source/Programs/AutomationTool/Gauntlet/Unreal/Base/Gauntlet.UnrealSession.cs
@@ -134,6 +134,11 @@ namespace Gauntlet
 		/// </summary>
 		public bool DeferredLaunch { get; set; }
 
+		/// <summary>
+		/// Whether this role will compress screenshots produced as an artifact into a jpeg format
+		/// </summary>
+		public bool CompressScreenshots { get; set; }
+
 		/// <summary>
 		/// Is this role Null?
 		/// </summary>
@@ -160,7 +165,7 @@ namespace Gauntlet
 		/// <param name="InConfiguration"></param>
 		/// <param name="InCommandLine"></param>
 		/// <param name="InOptions"></param>
-		public UnrealSessionRole(UnrealTargetRole InType, UnrealTargetPlatform? InPlatform, UnrealTargetConfiguration InConfiguration, string InCommandLine = null, IConfigOption<UnrealAppConfig> InOptions = null)
+		public UnrealSessionRole(UnrealTargetRole InType, UnrealTargetPlatform? InPlatform, UnrealTargetConfiguration InConfiguration, string InCommandLine = null, IConfigOption<UnrealAppConfig> InOptions = null, bool CompressScreenshots = true)
 		{
 			RoleType = InType;
 
@@ -184,8 +189,8 @@ namespace Gauntlet
 				RequiredBuildFlags |= BuildFlags.CanReplaceExecutable;
 			}
 
-			// Enforce build flags for the platform build that support it 
-			IDeviceBuildSupport TargetBuildSupport = Gauntlet.Utils.InterfaceHelpers.FindImplementations<IDeviceBuildSupport>().Where(B => B.CanSupportPlatform(InPlatform)).FirstOrDefault();
+			// Enforce build flags for the platform build that support it
+			IDeviceBuildSupport TargetBuildSupport = InterfaceHelpers.FindImplementations<IDeviceBuildSupport>().Where(B => B.CanSupportPlatform(InPlatform)).FirstOrDefault();
 			if (TargetBuildSupport != null)
 			{
 				if (Globals.Params.ParseParam("bulk") && TargetBuildSupport.CanSupportBuildType(BuildFlags.Bulk))
@@ -216,7 +221,8 @@ namespace Gauntlet
             FilesToCopy = new List<UnrealFileToCopy>();
 			CommandLineParams = new GauntletCommandLine();
 			RoleModifier = ERoleModifier.None;
-        }
+			this.CompressScreenshots = CompressScreenshots;
+		}
 
         /// <summary>
         /// Debugging aid
@@ -1019,6 +1025,13 @@ namespace Gauntlet
 					// Perform the copy
 					try
 					{
+						// Save screenshots and create a gif when not running a server
+						if (!InRunningRole.Role.RoleType.IsServer())
+						{
+							SaveScreenshots(RoleName, SourceDirectory.FullName, DestinationDirectory.FullName, InRunningRole.Role.CompressScreenshots);
+						}
+
+						// Copy remaining artifacts
 						SystemHelpers.CopyDirectory(SourceDirectory.FullName, DestinationDirectory.FullName, SystemHelpers.CopyOptions.Default, TruncateLongPathFilter);
 					}
 					catch(Exception Exception)
@@ -1148,38 +1161,6 @@ namespace Gauntlet
 				}
 			}
 
-			// Convert any screenshots to jpegs and create a gif when not running a server
-			if (!InRunningRole.Role.RoleType.IsServer())
-			{
-				try
-				{
-					DirectoryInfo ScreenshotDirectory = new(Path.Combine(DestinationDirectory.FullName, "Screenshots"));
-					if (ScreenshotDirectory.Exists)
-					{
-						foreach (DirectoryInfo ScreenshotSubdirectory in ScreenshotDirectory.EnumerateDirectories())
-						{
-							if (ScreenshotSubdirectory.GetFiles().Any())
-							{
-								Log.Info("Downsizing and gifying session images at {0}", ScreenshotSubdirectory.FullName);
-
-								// Downsize first so gif-step is quicker and takes less resources.
-								Utils.Image.ConvertImages(ScreenshotSubdirectory.FullName, ScreenshotSubdirectory.FullName, "jpg", true);
-
-								string GifPath = GenerateNotTakenFilePath(Path.Combine(DestinationDirectory.FullName, RoleName + "Test.gif"));
-								if (Utils.Image.SaveImagesAsGif(ScreenshotSubdirectory.FullName, GifPath))
-								{
-									Log.Info("Saved gif to {0}", GifPath);
-								}
-							}
-						}
-					}
-				}
-				catch (Exception Ex)
-				{
-					Log.Info("Failed to downsize and gif-ify images! {0}", Ex.Message);
-				}
-			}
-
 			// TODO REMOVEME- this should go elsewhere, likely a util that can be called or inserted by relevant test nodes.
 			SavePSOs(InContext, InRunningRole, DestinationDirectory.FullName);
 			// END REMOVEME
@@ -1770,6 +1751,70 @@ namespace Gauntlet
 			return new UnrealSessionInstance(RoleInstances.ToArray(), DeferredRolesToInstalls);
 		}
 
+		private void SaveScreenshots(string RoleName, string SourceDirectory, string DestinationDirectory, bool bCompressImages)
+		{
+			try
+			{
+				string ScreenshotPath = PathUtils.FindRelevantPath(BasePath: SourceDirectory, "Screenshots");
+
+				if (string.IsNullOrEmpty(ScreenshotPath) || !Directory.Exists(ScreenshotPath))
+				{
+					Log.Verbose("Could not locate Screenshots subdirectory in artifact path {ArtifactPath}. Skipping GIF creation", SourceDirectory);
+					return;
+				}
+
+				// Image transformation relies on the image being directly on the host PC
+				// This means it can't be on a remote device, or in a network storage.
+				// The screenshots directory will be copied to a temp directory for the transformations, and then moved to the final artifact destination.
+				DirectoryInfo ScreenshotDirectory = new(ScreenshotPath);
+				DirectoryInfo ImageStage = new(Path.Combine(Path.GetTempPath(), "ImageStage", RoleName));
+				if (ImageStage.Exists)
+				{
+					ImageStage.Delete(true);
+				}
+
+				// Copy the images to a the staging directory
+				SystemHelpers.CopyDirectory(ScreenshotPath, ImageStage.FullName, SystemHelpers.CopyOptions.Mirror);
+
+				// Creates an enumerable containing both the root directory and any sub directories
+				IEnumerable<DirectoryInfo> ImageDirectories = ImageStage
+					.EnumerateDirectories()
+					.Concat(Enumerable.Repeat(ImageStage, 1));
+
+				// Compress the images and convert them to a gif
+				foreach (DirectoryInfo ImageDirectory in ImageDirectories)
+				{
+					FileInfo[] Files = ImageDirectory.GetFiles();
+					if (Files.Any())
+					{
+						if (bCompressImages)
+						{
+							Image.ConvertImages(ImageDirectory.FullName, ImageDirectory.FullName, "jpg", true);
+						}
+
+						string GifPath = GenerateNotTakenFilePath(Path.Combine(DestinationDirectory, RoleName + "Test.gif"));
+						if (Image.SaveImagesAsGif(ImageDirectory.FullName, GifPath))
+						{
+							Log.Info("Saved gif to {0}", GifPath);
+						}
+					}
+				}
+
+				// Now, copy the compressed images in the temp staged directory to the desired artifact destination
+				string RelativeScreenshotDirectory = ScreenshotPath.Replace(SourceDirectory, string.Empty).Trim('\\').Trim('/');
+				string ScreenshotDestination = Path.Combine(DestinationDirectory, RelativeScreenshotDirectory);
+				SystemHelpers.CopyDirectory(ImageStage.FullName, ScreenshotDestination, SystemHelpers.CopyOptions.Mirror);
+
+				// Delete the source screenshot directory so they aren't duplicated in the following artifact copy
+				ScreenshotDirectory.Delete(true);
+				ImageStage.Delete(true);
+			}
+			catch (Exception Ex)
+			{
+				Log.Info("Failed to downsize and gif-ify images! {0}", Ex.Message);
+			}
+		}
+
 		private void SavePSOs(UnrealTestContext InContext, UnrealSessionInstance.RoleInstance InRunningRole, string DestSavedDir)
 		{
 			if (InRunningRole.Role.RoleType.IsServer())
diff --git a/Engine/Source/Programs/AutomationTool/Gauntlet/Unreal/Base/Gauntlet.UnrealTestConfiguration.cs b/Engine/Source/Programs/AutomationTool/Gauntlet/Unreal/Base/Gauntlet.UnrealTestConfiguration.cs
--- a/Engine/Source/Programs/AutomationTool/Gauntlet/Unreal/Base/Gauntlet.UnrealTestConfiguration.cs
+++ b/Engine/Source/Programs/AutomationTool/Gauntlet/Unreal/Base/Gauntlet.UnrealTestConfiguration.cs
@@ -393,6 +393,7 @@ namespace Gauntlet
 			InstallOnly = false;
 			DeferredLaunch = false;
 			CommandLineParams = new GauntletCommandLine();
+			CompressScreenshots = true;
 		}
 
 		public ERoleModifier RoleType { get; set; }
@@ -479,10 +480,14 @@ namespace Gauntlet
 		public string MapOverride { get; set; }
 
 		/// <summary>
-		/// Role device configuration 
+		/// Role device configuration
 		/// </summary>
 		public ConfigureDeviceHandler ConfigureDevice;
 
+		/// <summary>
+		/// Whether this role will compress screenshots produced as an artifact into a jpeg format
+		/// </summary>
+		public bool CompressScreenshots;
 	}
 
 	/// <summary>
diff --git a/Engine/Source/Programs/AutomationTool/Gauntlet/Unreal/Base/Gauntlet.UnrealTestNode.cs b/Engine/Source/Programs/AutomationTool/Gauntlet/Unreal/Base/Gauntlet.UnrealTestNode.cs
--- a/Engine/Source/Programs/AutomationTool/Gauntlet/Unreal/Base/Gauntlet.UnrealTestNode.cs
+++ b/Engine/Source/Programs/AutomationTool/Gauntlet/Unreal/Base/Gauntlet.UnrealTestNode.cs
@@ -667,6 +667,7 @@ namespace Gauntlet
 					SessionRole.AdditionalArtifactDirectories = TestRole.AdditionalArtifactDirectories;
 					SessionRole.ConfigureDevice = TestRole.ConfigureDevice;
 					SessionRole.MapOverride = TestRole.MapOverride;
+					SessionRole.CompressScreenshots = TestRole.CompressScreenshots;
 
 					SessionRoles.Add(SessionRole);
 				}
