[Gauntlet] InstallUnrealBuild command - Validate project input

#rnx

diff --git a/Engine/Source/Programs/AutomationTool/Gauntlet/Unreal/Utils/Gauntlet.InstallUnreal.cs b/Engine/Source/Programs/AutomationTool/Gauntlet/Unreal/Utils/Gauntlet.InstallUnreal.cs
--- a/Engine/Source/Programs/AutomationTool/Gauntlet/Unreal/Utils/Gauntlet.InstallUnreal.cs
+++ b/Engine/Source/Programs/AutomationTool/Gauntlet/Unreal/Utils/Gauntlet.InstallUnreal.cs
@@ -30,6 +30,11 @@ namespace Gauntlet
 			string DevicesArg,
 			int ParallelTasks )
 		{
+			if (string.IsNullOrEmpty(ProjectName))
+			{
+				throw new Exception("need -project=<project name>");
+			}
+
 			if (string.IsNullOrEmpty(PlatformParam) || string.IsNullOrEmpty(BuildPath))
 			{
 				throw new Exception("need -platform=<platform> and -path=\"path to build\"");
