Fix a crash bug when no cooked output exists and you attempt to stage a build.

diff --git a/Engine/Source/Programs/AutomationTool/Scripts/CopyBuildToStagingDirectory.Automation.cs b/Engine/Source/Programs/AutomationTool/Scripts/CopyBuildToStagingDirectory.Automation.cs
--- a/Engine/Source/Programs/AutomationTool/Scripts/CopyBuildToStagingDirectory.Automation.cs
+++ b/Engine/Source/Programs/AutomationTool/Scripts/CopyBuildToStagingDirectory.Automation.cs
@@ -3486,8 +3486,8 @@ namespace AutomationScripts
 					CommonAdditionalArgs += " -allowForIndexUnload";
 				}
 			}
-			
-			if (!string.IsNullOrEmpty(SC.PackageStoreData.MarkerFullPath))
+
+			if ((SC.PackageStoreData != null) && !string.IsNullOrEmpty(SC.PackageStoreData.MarkerFullPath))
 			{
 				CommonAdditionalArgs += String.Format(" -ProjectStore={0}",
 					MakePathSafeToUseWithCommandLine(SC.PackageStoreData.MarkerFullPath));
