#BuildGraph Do not automatically set LogSavePackage verbosity to warning

diff --git a/Engine/Source/Programs/AutomationTool/BuildGraph/Tasks/CookTask.cs b/Engine/Source/Programs/AutomationTool/BuildGraph/Tasks/CookTask.cs
--- a/Engine/Source/Programs/AutomationTool/BuildGraph/Tasks/CookTask.cs
+++ b/Engine/Source/Programs/AutomationTool/BuildGraph/Tasks/CookTask.cs
@@ -109,7 +109,7 @@ namespace AutomationTool.Tasks
 				scope.Span.SetTag("project", projectFile == null ? "UE4" : projectFile.GetFileNameWithoutExtension());
 				scope.Span.SetTag("platform", _parameters.Platform);
 				string[] maps = (_parameters.Maps == null) ? null : _parameters.Maps.Split(new char[] { '+' });
-				string arguments = (_parameters.Versioned ? "" : "-Unversioned ") + "-LogCmds=\"LogSavePackage Warning\" " + _parameters.Arguments;
+				string arguments = (_parameters.Versioned ? "" : "-Unversioned ") + _parameters.Arguments;
 				string editorExe = (string.IsNullOrWhiteSpace(_parameters.EditorExe) ? ProjectUtils.GetEditorForProject(projectFile).FullName : _parameters.EditorExe);
 				CommandUtils.CookCommandlet(projectFile, editorExe, maps, null, null, null, _parameters.Platform, arguments);
 			}
