Fix for custom deployment handlers interfering when multiple platforms are specified on the BuildCookRun command line.

#rnx

diff --git a/Engine/Source/Programs/AutomationTool/AutomationUtils/Platform.cs b/Engine/Source/Programs/AutomationTool/AutomationUtils/Platform.cs
--- a/Engine/Source/Programs/AutomationTool/AutomationUtils/Platform.cs
+++ b/Engine/Source/Programs/AutomationTool/AutomationUtils/Platform.cs
@@ -747,6 +747,16 @@ namespace AutomationTool
 			return "";
 		}
 
+		/// <summary>
+		/// Get the custom deployment handler for this platform,
+		/// </summary>
+		/// <param name="Params"></param>
+		/// <param name="SC">Deployment Context</param>
+		public virtual CustomDeploymentHandler GetCustomDeploymentHandler(ProjectParams Params, DeploymentContext SC)
+		{
+			return null;
+		}
+
 		/// <summary>
 		/// True if this platform is supported.
 		/// </summary>
diff --git a/Engine/Source/Programs/AutomationTool/Scripts/CopyBuildToStagingDirectory.Automation.cs b/Engine/Source/Programs/AutomationTool/Scripts/CopyBuildToStagingDirectory.Automation.cs
--- a/Engine/Source/Programs/AutomationTool/Scripts/CopyBuildToStagingDirectory.Automation.cs
+++ b/Engine/Source/Programs/AutomationTool/Scripts/CopyBuildToStagingDirectory.Automation.cs
@@ -5479,6 +5479,8 @@ namespace AutomationScripts
 					);
 				LogDeploymentContext(SC);
 
+				SC.CustomDeployment = SC.StageTargetPlatform?.GetCustomDeploymentHandler(Params, SC);
+
 				if (Params.PreModifyDeploymentContextCallback != null)
 				{
 					Params.PreModifyDeploymentContextCallback(Params, SC);
diff --git a/Engine/Source/Programs/AutomationTool/Win/WinPlatform.Automation.cs b/Engine/Source/Programs/AutomationTool/Win/WinPlatform.Automation.cs
--- a/Engine/Source/Programs/AutomationTool/Win/WinPlatform.Automation.cs
+++ b/Engine/Source/Programs/AutomationTool/Win/WinPlatform.Automation.cs
@@ -47,33 +47,6 @@ public class Win64Platform : Platform
 		return Devices.ToArray();
 	}
 
-	public override void PlatformSetupParams(ref ProjectParams Params)
-	{
-		base.PlatformSetupParams(ref Params);
-
-		// use a custom deployment handler if one is requested
-		Params.PreModifyDeploymentContextCallback = new Action<ProjectParams, DeploymentContext>((ProjectParams Params, DeploymentContext SC) =>
-		{
-			if (SC.CustomDeployment == null)
-			{			
-				string CustomDeploymentName = null;
-
-				ConfigHierarchy EngineIni = ConfigCache.ReadHierarchy(ConfigHierarchyType.Engine, Params.RawProjectPath.Directory, PlatformType, SC.CustomConfig);
-				EngineIni.GetString("/Script/WindowsTargetPlatform.WindowsTargetSettings", "CustomDeployment", out CustomDeploymentName);
-
-				if (string.IsNullOrEmpty(CustomDeploymentName))
-				{
-					CustomDeploymentName = Params.CustomDeploymentHandler;
-				}
-
-				if (!string.IsNullOrEmpty(CustomDeploymentName))
-				{
-					SC.CustomDeployment = CustomDeploymentHandler.Create(CustomDeploymentName, this);
-				}
-			}			
-		});
-	}
-
 	public override void Deploy(ProjectParams Params, DeploymentContext SC)
 	{
 		// We only care about deploying for SteamDeck
@@ -359,6 +332,24 @@ public class Win64Platform : Platform
 		return new List<string> { ".pdb", ".map" };
 	}
 
+	public override CustomDeploymentHandler GetCustomDeploymentHandler(ProjectParams Params, DeploymentContext SC)
+	{
+		ConfigHierarchy EngineIni = ConfigCache.ReadHierarchy(ConfigHierarchyType.Engine, Params.RawProjectPath.Directory, PlatformType, SC.CustomConfig);
+		EngineIni.GetString("/Script/WindowsTargetPlatform.WindowsTargetSettings", "CustomDeployment", out string CustomDeploymentName);
+
+		if (string.IsNullOrEmpty(CustomDeploymentName))
+		{
+			CustomDeploymentName = Params.CustomDeploymentHandler;
+		}
+
+		if (!string.IsNullOrEmpty(CustomDeploymentName))
+		{
+			return CustomDeploymentHandler.Create(CustomDeploymentName, this);
+		}
+
+		return base.GetCustomDeploymentHandler(Params, SC);
+	}
+
 	public override bool SignExecutables(DeploymentContext SC, ProjectParams Params)
 	{
 		// Sign everything we built
