Stage remapped files for Linux

diff --git a/Engine/Source/Programs/AutomationTool/Linux/LinuxPlatform.Automation.cs b/Engine/Source/Programs/AutomationTool/Linux/LinuxPlatform.Automation.cs
--- a/Engine/Source/Programs/AutomationTool/Linux/LinuxPlatform.Automation.cs
+++ b/Engine/Source/Programs/AutomationTool/Linux/LinuxPlatform.Automation.cs
@@ -123,7 +123,8 @@ public abstract class BaseLinuxPlatform : Platform
 						List<StagedFileReference> StagePaths = SC.FilesToStage.NonUFSFiles.Where(x => x.Value == Executable.Path).Select(x => x.Key).ToList();
 						foreach (StagedFileReference StagePath in StagePaths)
 						{
-							StageBootstrapExecutable(SC, BootstrapExeName + Extension, FullExecutablePath, StagePath.Name, BootstrapArguments);
+							StagedFileReference RemappedStagePath = DeploymentContext.ApplyDirectoryRemap(SC, StagePath);
+							StageBootstrapExecutable(SC, BootstrapExeName + Extension, FullExecutablePath, RemappedStagePath.Name, BootstrapArguments);
 						}
 					}
 				}
