Add InputActions module to Gauntlet, this reimplements FKey and EKeys from C++ into C# to be used to define key actions, intended on being deserialized by the RPC Framework and passed into APlayerController:InputKey

diff --git a/Engine/Source/Programs/AutomationTool/Gauntlet/Framework/Utils/Gauntlet.InputActions.cs b/Engine/Source/Programs/AutomationTool/Gauntlet/Framework/Utils/Gauntlet.InputActions.cs
new file mode 100644
--- /dev/null
+++ b/Engine/Source/Programs/AutomationTool/Gauntlet/Framework/Utils/Gauntlet.InputActions.cs
@@ -0,0 +1,606 @@
+﻿using EpicGames.Core;
+using System;
+using UnrealBuildTool;
+
+namespace Gauntlet.Utils
+{
+	public class FKey
+	{
+		public string KeyName { get; set; }
+
+		public FKey(string InName)
+		{
+			this.KeyName = InName;
+		}
+
+		// We might need this in the future, but there's some issues serializing FKeyDetails as it has a cyclic reference to the FKey object itself (In pointer form in C++)
+		// public FKeyDetails KeyDetails
+
+		
+	}
+	
+	// Basically a replica of FKeyDetails from InputCoreTypes, intended to be serialized and then deserialized into its UObject type
+	// We may end up genuinely needing this for Axis-based input (For example MouseWheelAxis or Mouse2D), but it's not needed right now for regular key-presses
+	public struct FKeyDetails
+	{
+		public enum EKeyFlags
+		{
+			GamepadKey					= 1 << 0,
+			Touch						= 1 << 1,
+			MouseButton					= 1 << 2,
+			ModifierKey					= 1 << 3,
+			NotBlueprintBindableKey		= 1 << 4,
+			Axis1D						= 1 << 5,
+			Axis3D						= 1 << 6,
+			UpdateAxisWithoutSamples	= 1 << 7,
+			NotActionBindableKey		= 1 << 8,
+			Deprecated					= 1 << 9,
+			ButtonAxis					= 1 << 10,
+			Axis2D						= 1 << 11,
+			Gesture						= 1 << 12,
+			NoFlags						= 0
+		}
+
+		public enum EInputAxisType
+		{
+			None,
+			Button,
+			Axis1D,
+			Axis2D,
+			Axis3D
+		}
+
+		private FKey Key;
+		public int PairedAxis = 0;
+		public FKey PairedAxisKey;
+		public string MenuCategory;
+
+		public int bIsModifierKey;
+		public int bIsGamepadKey;
+		public int bIsTouch;
+		public int bIsMouseButton;
+		public int bIsBindableInBlueprints;
+		public int bShouldUpdateAxisWithoutSamples;
+		public int bIsbindableToActions;
+		public int bIsDeprecated;
+		public int bIsGesture;
+		public EInputAxisType AxisType;
+
+		public string LongDisplayName;
+		public string ShortDisplayName;
+
+		public FKeyDetails(FKey InKey, string InLongDisplayName, string InShortDisplayName = "", int InKeyFlags = 0)
+		{
+			this.Key = InKey;
+			this.LongDisplayName = InLongDisplayName;
+			this.ShortDisplayName = InShortDisplayName;
+
+			this.bIsGamepadKey = (InKeyFlags & (int)EKeyFlags.GamepadKey) != 0 ? 1 : 0;
+			this.bIsModifierKey = (InKeyFlags & (int)EKeyFlags.ModifierKey) != 0 ? 1 : 0;
+			this.bIsTouch = (InKeyFlags & (int)EKeyFlags.Touch) != 0 ? 1 : 0;
+			this.bIsMouseButton = (InKeyFlags & (int)EKeyFlags.MouseButton) != 0 ? 1 : 0;
+			this.bIsBindableInBlueprints = (~InKeyFlags & (int)EKeyFlags.NotBlueprintBindableKey) != 0 && (~InKeyFlags & (int)EKeyFlags.Deprecated) != 0 ? 1 : 0;
+			this.bShouldUpdateAxisWithoutSamples = (InKeyFlags & (int)EKeyFlags.UpdateAxisWithoutSamples) != 0 ? 1 : 0;
+			this.bIsbindableToActions = (~InKeyFlags & (int)EKeyFlags.NotActionBindableKey) != 0 && (~InKeyFlags & (int)EKeyFlags.Deprecated) != 0 ? 1 : 0;
+			this.bIsDeprecated = (InKeyFlags & (int)EKeyFlags.Deprecated) != 0 ? 1 : 0;
+			this.bIsGesture = (InKeyFlags & (int)EKeyFlags.Gesture) != 0 ? 1 : 0;
+
+			// Set MenuCategory based on above flags
+			if(bIsGamepadKey == 1)
+			{
+				this.MenuCategory = "Gamepad";
+			}
+			else if(bIsMouseButton == 1)
+			{
+				this.MenuCategory = "Mouse";
+			}
+			else
+			{
+				this.MenuCategory = "Keyboard";
+			}
+
+			// Defaulting the Axis stuff to null for now
+			PairedAxisKey = null;
+			AxisType = 0;
+		}
+
+		public FKey GetKey()
+		{
+			return this.Key;
+		}
+
+		public string GetLongDisplayName()
+		{
+			return this.LongDisplayName;
+		}
+
+		public string GetShortDisplayName()
+		{
+			return this.ShortDisplayName;
+		}
+
+		public bool IsGamepadKey()
+		{
+			return this.bIsGamepadKey == 1;
+		}
+
+		public bool IsModifierKey()
+		{
+			return this.bIsModifierKey == 1;
+		}
+
+		public bool IsTouch()
+		{
+			return this.bIsTouch == 1;
+		}
+
+		public bool IsMouseButton()
+		{
+			return this.bIsMouseButton == 1;
+		}
+
+		public bool IsGesture()
+		{
+			return this.bIsGesture == 1;
+		}
+
+		public string GetMenuCategory()
+		{
+			return this.MenuCategory;
+		}
+
+		public int GetPairedAxis()
+		{
+			return this.PairedAxis;
+		}
+
+		public FKey GetPairedAxisKey()
+		{
+			return this.PairedAxisKey;
+		}
+
+		public EInputAxisType GetAxisType()
+		{
+			return this.AxisType;
+		}
+
+	}
+
+	public struct EKeys
+	{
+		public FKey A = new FKey("A");
+		public FKey B = new FKey("B");
+		public FKey C = new FKey("C");
+		public FKey D = new FKey("D");
+		public FKey E = new FKey("E");
+		public FKey F = new FKey("F");
+		public FKey G = new FKey("G");
+		public FKey H = new FKey("H");
+		public FKey I = new FKey("I");
+		public FKey J = new FKey("J");
+		public FKey K = new FKey("K");
+		public FKey L = new FKey("L");
+		public FKey M = new FKey("M");
+		public FKey N = new FKey("N");
+		public FKey O = new FKey("O");
+		public FKey P = new FKey("P");
+		public FKey Q = new FKey("Q");
+		public FKey R = new FKey("R");
+		public FKey S = new FKey("S");
+		public FKey T = new FKey("T");
+		public FKey U = new FKey("U");
+		public FKey V = new FKey("V");
+		public FKey W = new FKey("W");
+		public FKey X = new FKey("X");
+		public FKey Y = new FKey("Y");
+		public FKey Z = new FKey("Z");
+
+		public FKey MouseX = new FKey("MouseX");
+		public FKey MouseY = new FKey("MouseY");
+		public FKey Mouse2D = new FKey("Mouse2D");
+		public FKey MouseScrollUp = new FKey("MouseScrollUp");
+		public FKey MouseScrollDown = new FKey("MouseScrollDown");
+		public FKey MouseWheelAxis = new FKey("MouseWheelAxis");
+		public FKey LeftMouseButton = new FKey("LeftMouseButton");
+		public FKey RightMouseButton = new FKey("RightMouseButton");
+		public FKey MiddleMouseButton = new FKey("MiddleMouseButton");
+		public FKey ThumbMouseButton = new FKey("ThumbMouseButton");
+		public FKey ThumbMouseButton2 = new FKey("ThumbMouseButton2");
+
+		// Separate commands here for Mac
+		public FKey BackSpace = new FKey("Backspace");
+		public FKey Delete = new FKey("Delete");
+
+
+		public FKey Tab = new FKey("Tab");
+		public FKey Enter = new FKey("Enter");
+		public FKey Pause = new FKey("Pause");
+
+		public FKey CapsLock = new FKey("CapsLock");
+		public FKey Escape = new FKey("Escape");
+		public FKey SpaceBar = new FKey("SpaceBar");
+		public FKey PageUp = new FKey("PageUp");
+		public FKey PageDown = new FKey("PageDown");
+		public FKey End = new FKey("End");
+		public FKey Home = new FKey("Home");
+
+		public FKey Left = new FKey("Left");
+		public FKey Up = new FKey("Up");
+		public FKey Right = new FKey("Right");
+		public FKey Down = new FKey("Down");
+
+		public FKey Insert = new FKey("Insert");
+
+		public FKey Zero = new FKey("0");
+		public FKey One = new FKey("1");
+		public FKey Two = new FKey("2");
+		public FKey Three = new FKey("3");
+		public FKey Four = new FKey("4");
+		public FKey Five = new FKey("5");
+		public FKey Six = new FKey("6");
+		public FKey Seven = new FKey("7");
+		public FKey Eight = new FKey("8");
+		public FKey Nine = new FKey("9");
+
+		public FKey NumPadZero = new FKey("NumPadZero");
+		public FKey NumPadOne = new FKey("NumPadOne");
+		public FKey NumPadTwo = new FKey("NumPadTwo");
+		public FKey NumPadThree = new FKey("NumPadThree");
+		public FKey NumPadFour = new FKey("NumPadFour");
+		public FKey NumPadFive = new FKey("NumPadFive");
+		public FKey NumPadSix = new FKey("NumPadSix");
+		public FKey NumPadSeven = new FKey("NumPadSeven");
+		public FKey NumPadEight = new FKey("NumPadEight");
+		public FKey NumPadNine = new FKey("NumPadNine");
+
+		public FKey F1 = new FKey("F1");
+		public FKey F2 = new FKey("F2");
+		public FKey F3 = new FKey("F3");
+		public FKey F4 = new FKey("F4");
+		public FKey F5 = new FKey("F5");
+		public FKey F6 = new FKey("F6");
+		public FKey F7 = new FKey("F7");
+		public FKey F8 = new FKey("F8");
+		public FKey F9 = new FKey("F9");
+		public FKey F10 = new FKey("F10");
+		public FKey F11 = new FKey("F11");
+		public FKey F12 = new FKey("F12");
+
+		public FKey NumLock = new FKey("NumLock");
+
+		public FKey LeftShift = new FKey("LeftShift");
+		public FKey RightShift = new FKey("RightShift");
+		public FKey LeftControl = new FKey("LeftControl");
+		public FKey RightControl = new FKey("RightControl");
+		public FKey LeftAlt = new FKey("LeftAlt");
+		public FKey RightAlt = new FKey("RightAlt");
+		public FKey LeftCommand = new FKey("LeftCommand");
+		public FKey RightCommand = new FKey("RightCommand");
+
+		public FKey Semicolon = new FKey("Semicolon");
+		public FKey EqualsKey = new FKey("Equals");
+
+		public FKey Comma = new FKey("Comma");
+		public FKey Underscore = new FKey("Underscore");
+		public FKey Hyphen = new FKey("Hyphen");
+		public FKey Period = new FKey("Period");
+		public FKey Slash = new FKey("Slash");
+		public FKey Tilde = new FKey("`");
+		public FKey LeftBracket = new FKey("LeftBracket");
+		public FKey Backslash = new FKey("Backslash");
+		public FKey RightBracket = new FKey("RightBracket");
+		public FKey Apostrophe = new FKey("Apostrophe");
+
+		public FKey Ampersand = new FKey("Ampersand");
+		public FKey Asterix = new FKey("Asterix");
+		public FKey Caret = new FKey("Caret");
+		public FKey Colon = new FKey("Colon");
+		public FKey Dollar = new FKey("Dollar");
+		public FKey Exclamation = new FKey("Exclamation");
+		public FKey LeftParantheses = new FKey("LeftParantheses");
+		public FKey RightParantheses = new FKey("RightParantheses");
+		public FKey Quote = new FKey("Quote");
+
+		// In C++ this uses FString::Chr
+		public FKey A_AccentGrave = new FKey($"{(char)224}");
+		public FKey E_AccentGrave = new FKey($"{(char)232}");
+		public FKey E_AccentAigu = new FKey($"{(char)233}");
+		public FKey C_Cedille = new FKey($"{(char)231}");
+		public FKey Section = new FKey($"{(char)167}");
+
+		// Gamepad Axis Stuff
+		// These are still mostly untested and are here for parity until we can confirm how to properly do analog inputs
+		public FKey Gamepad_Left2D = new FKey("Gamepad_Left2D");
+		public FKey Gamepad_LeftX = new FKey("Gamepad_LeftX");
+		public FKey Gamepad_LeftY = new FKey("Gamepad_LeftY");
+		public FKey Gamepad_Right2D = new FKey("Gamepad_Right2D");
+		public FKey Gamepad_RightX = new FKey("Gamepad_RightX");
+		public FKey Gamepad_RightY = new FKey("Gamepad_RightY");
+		public FKey Gamepad_Special_Left_X = new FKey("Gamepad_Special_Left_X");
+		public FKey Gamepad_Special_Left_Y = new FKey("Gamepad_Special_Left_Y");
+		public FKey Gamepad_LeftTrigger = new FKey("Gamepad_LeftTrigger");
+		public FKey Gamepad_RightTrigger = new FKey("Gamepad_RightTrigger");
+		public FKey Gamepad_LeftThumbstick = new FKey("Gamepad_LeftThumbstick");
+		public FKey Gamepad_RightThumbstick = new FKey("Gamepad_RightThumbstick");
+		public FKey Gamepad_LeftTriggerAxis = new FKey("Gamepad_LeftTriggerAxis");
+		public FKey Gamepad_RightTriggerAxis = new FKey("Gamepad_RightTriggerAxis");
+
+		// Gamepad Generic
+		public FKey Gamepad_DPad_Up = new FKey("Gamepad_DPad_Up");
+		public FKey Gamepad_DPad_Down = new FKey("Gamepad_DPad_Down");
+		public FKey Gamepad_DPad_Right = new FKey("Gamepad_DPad_Right");
+		public FKey Gamepad_DPad_Left = new FKey("Gamepad_DPad_Left");
+
+		public FKey Gamepad_LeftStick_Up = new FKey("Gamepad_LeftStick_Up");
+		public FKey Gamepad_LeftStick_Down = new FKey("Gamepad_LeftStick_Down");
+		public FKey Gamepad_LeftStick_Right = new FKey("Gamepad_LeftStick_Right");
+		public FKey Gamepad_LeftStick_Left = new FKey("Gamepad_LeftStick_Left");
+
+		public FKey Gamepad_RightStick_Up = new FKey("Gamepad_RightStick_Up");
+		public FKey Gamepad_RightStick_Down = new FKey("Gamepad_RightStick_Down");
+		public FKey Gamepad_RightStick_Right = new FKey("Gamepad_RightStick_Right");
+		public FKey Gamepad_RightStick_Left = new FKey("Gamepad_RightStick_Left");
+
+		public FKey Gamepad_Special_Left = new FKey("Gamepad_Special_Left");
+		public FKey Gamepad_Special_Right = new FKey("Gamepad_Special_Right");
+		public FKey Gamepad_FaceButton_Bottom = new FKey("Gamepad_FaceButton_Bottom");
+		public FKey Gamepad_FaceButton_Right = new FKey("Gamepad_FaceButton_Right");
+		public FKey Gamepad_FaceButton_Left = new FKey("Gamepad_FaceButton_Left");
+		public FKey Gamepad_FaceButton_Top = new FKey("Gamepad_FaceButton_Top");
+		public FKey Gamepad_LeftShoulder = new FKey("Gamepad_LeftShoulder");
+		public FKey Gamepad_RightShoulder = new FKey("Gamepad_RightShoulder");
+
+		public FKey PS4_Special = new FKey("PS4_Special");
+
+		// Motion Controls
+		public FKey Tilt = new FKey("Tilt");
+		public FKey RotationRate = new FKey("RotationRate");
+		public FKey Gravity = new FKey("Gravity");
+		public FKey Acceleration = new FKey("Acceleration");
+
+		// Gesture Controls
+		public FKey Gesture_Pinch = new FKey("Gesture_Pinch");
+		public FKey Gesture_Flick = new FKey("Gesture_Flick");
+		public FKey Gesture_Rotate = new FKey("Gesture_Rotate");
+
+		// Steam Controls
+		public FKey Steam_Touch_0 = new FKey("Steam_Touch_0");
+		public FKey Steam_Touch_1 = new FKey("Steam_Touch_1");
+		public FKey Steam_Touch_2 = new FKey("Steam_Touch_2");
+		public FKey Steam_Touch_3 = new FKey("Steam_Touch_3");
+		public FKey Steam_Back_Left = new FKey("Steam_Back_Left");
+		public FKey Steam_Back_Right = new FKey("Steam_Back_Right");
+
+		// Xbox One Global Speech commands
+		public FKey Global_Menu = new FKey("Global_Menu");
+		public FKey Global_View = new FKey("Global_View");
+		public FKey Global_Pause = new FKey("Global_Pause");
+		public FKey Global_Play = new FKey("Global_Play");
+		public FKey Global_Back = new FKey("Global_Back");
+
+		// Android-specific Controls
+		public FKey Android_Back = new FKey("Android_Back");
+		public FKey Android_Volume_Up = new FKey("Android_Volume_Up");
+		public FKey Android_Volume_Down = new FKey("Android_Volume_Down");
+		public FKey Android_Menu = new FKey("Android_Menu");
+
+		// Vive Controls
+		public FKey Vive_Left_System_Click = new FKey("Vive_Left_System_Click");
+		public FKey Vive_Left_Grip_Click = new FKey("Vive_Left_Grip_Click");
+		public FKey Vive_Left_Menu_Click = new FKey("Vive_Left_Menu_Click");
+		public FKey Vive_Left_Trigger_Click = new FKey("Vive_Left_Trigger_Click");
+		public FKey Vive_Left_Trigger_Axis = new FKey("Vive_Left_Trigger_Axis");
+		public FKey Vive_Left_Trackpad_2D = new FKey("Vive_Left_Trackpad_2D");
+		public FKey Vive_Left_Trackpad_X = new FKey("Vive_Left_Trackpad_X");
+		public FKey Vive_Left_Trackpad_Y = new FKey("Vive_Left_Trackpad_Y");
+		public FKey Vive_Left_Trackpad_Click = new FKey("Vive_Left_Trackpad_Click");
+		public FKey Vive_Left_Trackpad_Touch = new FKey("Vive_Left_Trackpad_Touch");
+		public FKey Vive_Left_Trackpad_Up = new FKey("Vive_Left_Trackpad_Up");
+		public FKey Vive_Left_Trackpad_Down = new FKey("Vive_Left_Trackpad_Down");
+		public FKey Vive_Left_Trackpad_Left = new FKey("Vive_Left_Trackpad_Left");
+		public FKey Vive_Left_Trackpad_Right = new FKey("Vive_Left_Trackpad_Right");
+		public FKey Vive_Right_System_Click = new FKey("Vive_Right_System_Click");
+		public FKey Vive_Right_Grip_Click = new FKey("Vive_Right_Grip_Click");
+		public FKey Vive_Right_Menu_Click = new FKey("Vive_Right_Menu_Click");
+		public FKey Vive_Right_Trigger_Click = new FKey("Vive_Right_Trigger_Click");
+		public FKey Vive_Right_Trigger_Axis = new FKey("Vive_Right_Trigger_Axis");
+		public FKey Vive_Right_Trackpad_2D = new FKey("Vive_Right_Trackpad_2D");
+		public FKey Vive_Right_Trackpad_X = new FKey("Vive_Right_Trackpad_X");
+		public FKey Vive_Right_Trackpad_Y = new FKey("Vive_Right_Trackpad_Y");
+		public FKey Vive_Right_Trackpad_Click = new FKey("Vive_Right_Trackpad_Click");
+		public FKey Vive_Right_Trackpad_Touch = new FKey("Vive_Right_Trackpad_Touch");
+		public FKey Vive_Right_Trackpad_Up = new FKey("Vive_Right_Trackpad_Up");
+		public FKey Vive_Right_Trackpad_Down = new FKey("Vive_Right_Trackpad_Down");
+		public FKey Vive_Right_Trackpad_Left = new FKey("Vive_Right_Trackpad_Left");
+		public FKey Vive_Right_Trackpad_Right = new FKey("Vive_Right_Trackpad_Right");
+
+		// Mixed Reality Controls
+		public FKey MixedReality_Left_Menu_Click = new FKey("MixedReality_Left_Menu_Click");
+		public FKey MixedReality_Left_Grip_Click = new FKey("MixedReality_Left_Grip_Click");
+		public FKey MixedReality_Left_Trigger_Click = new FKey("MixedReality_Left_Trigger_Click");
+		public FKey MixedReality_Left_Trigger_Axis = new FKey("MixedReality_Left_Trigger_Axis");
+		public FKey MixedReality_Left_Thumbstick_2D = new FKey("MixedReality_Left_Thumbstick_2D");
+		public FKey MixedReality_Left_Thumbstick_X = new FKey("MixedReality_Left_Thumbstick_X");
+		public FKey MixedReality_Left_Thumbstick_Y = new FKey("MixedReality_Left_Thumbstick_Y");
+		public FKey MixedReality_Left_Thumbstick_Click = new FKey("MixedReality_Left_Thumbstick_Click");
+		public FKey MixedReality_Left_Thumbstick_Up = new FKey("MixedReality_Left_Thumbstick_Up");
+		public FKey MixedReality_Left_Thumbstick_Down = new FKey("MixedReality_Left_Thumbstick_Down");
+		public FKey MixedReality_Left_Thumbstick_Left = new FKey("MixedReality_Left_Thumbstick_Left");
+		public FKey MixedReality_Left_Thumbstick_Right = new FKey("MixedReality_Left_Thumbstick_Right");
+		public FKey MixedReality_Left_Trackpad_2D = new FKey("MixedReality_Left_Trackpad_2D");
+		public FKey MixedReality_Left_Trackpad_X = new FKey("MixedReality_Left_Trackpad_X");
+		public FKey MixedReality_Left_Trackpad_Y = new FKey("MixedReality_Left_Trackpad_Y");
+		public FKey MixedReality_Left_Trackpad_Click = new FKey("MixedReality_Left_Trackpad_Click");
+		public FKey MixedReality_Left_Trackpad_Touch = new FKey("MixedReality_Left_Trackpad_Touch");
+		public FKey MixedReality_Left_Trackpad_Up = new FKey("MixedReality_Left_Trackpad_Up");
+		public FKey MixedReality_Left_Trackpad_Down = new FKey("MixedReality_Left_Trackpad_Down");
+		public FKey MixedReality_Left_Trackpad_Left = new FKey("MixedReality_Left_Trackpad_Left");
+		public FKey MixedReality_Left_Trackpad_Right = new FKey("MixedReality_Left_Trackpad_Right");
+		public FKey MixedReality_Right_Menu_Click = new FKey("MixedReality_Right_Menu_Click");
+		public FKey MixedReality_Right_Grip_Click = new FKey("MixedReality_Right_Grip_Click");
+		public FKey MixedReality_Right_Trigger_Click = new FKey("MixedReality_Right_Trigger_Click");
+		public FKey MixedReality_Right_Trigger_Axis = new FKey("MixedReality_Right_Trigger_Axis");
+		public FKey MixedReality_Right_Thumbstick_2D = new FKey("MixedReality_Right_Thumbstick_2D");
+		public FKey MixedReality_Right_Thumbstick_X = new FKey("MixedReality_Right_Thumbstick_X");
+		public FKey MixedReality_Right_Thumbstick_Y = new FKey("MixedReality_Right_Thumbstick_Y");
+		public FKey MixedReality_Right_Thumbstick_Click = new FKey("MixedReality_Right_Thumbstick_Click");
+		public FKey MixedReality_Right_Thumbstick_Up = new FKey("MixedReality_Right_Thumbstick_Up");
+		public FKey MixedReality_Right_Thumbstick_Down = new FKey("MixedReality_Right_Thumbstick_Down");
+		public FKey MixedReality_Right_Thumbstick_Left = new FKey("MixedReality_Right_Thumbstick_Left");
+		public FKey MixedReality_Right_Thumbstick_Right = new FKey("MixedReality_Right_Thumbstick_Right");
+		public FKey MixedReality_Right_Trackpad_2D = new FKey("MixedReality_Right_Trackpad_2D");
+		public FKey MixedReality_Right_Trackpad_X = new FKey("MixedReality_Right_Trackpad_X");
+		public FKey MixedReality_Right_Trackpad_Y = new FKey("MixedReality_Right_Trackpad_Y");
+		public FKey MixedReality_Right_Trackpad_Click = new FKey("MixedReality_Right_Trackpad_Click");
+		public FKey MixedReality_Right_Trackpad_Touch = new FKey("MixedReality_Right_Trackpad_Touch");
+		public FKey MixedReality_Right_Trackpad_Up = new FKey("MixedReality_Right_Trackpad_Up");
+		public FKey MixedReality_Right_Trackpad_Down = new FKey("MixedReality_Right_Trackpad_Down");
+		public FKey MixedReality_Right_Trackpad_Left = new FKey("MixedReality_Right_Trackpad_Left");
+		public FKey MixedReality_Right_Trackpad_Right = new FKey("MixedReality_Right_Trackpad_Right");
+
+		// Oculus Touch controls
+		public FKey OculusTouch_Left_X_Click = new FKey("OculusTouch_Left_X_Click");
+		public FKey OculusTouch_Left_Y_Click = new FKey("OculusTouch_Left_Y_Click");
+		public FKey OculusTouch_Left_X_Touch = new FKey("OculusTouch_Left_X_Touch");
+		public FKey OculusTouch_Left_Y_Touch = new FKey("OculusTouch_Left_Y_Touch");
+		public FKey OculusTouch_Left_Menu_Click = new FKey("OculusTouch_Left_Menu_Click");
+		public FKey OculusTouch_Left_Grip_Click = new FKey("OculusTouch_Left_Grip_Click");
+		public FKey OculusTouch_Left_Grip_Axis = new FKey("OculusTouch_Left_Grip_Axis");
+		public FKey OculusTouch_Left_Trigger_Click = new FKey("OculusTouch_Left_Trigger_Click");
+		public FKey OculusTouch_Left_Trigger_Axis = new FKey("OculusTouch_Left_Trigger_Axis");
+		public FKey OculusTouch_Left_Trigger_Touch = new FKey("OculusTouch_Left_Trigger_Touch");
+		public FKey OculusTouch_Left_Thumbstick_2D = new FKey("OculusTouch_Left_Thumbstick_2D");
+		public FKey OculusTouch_Left_Thumbstick_X = new FKey("OculusTouch_Left_Thumbstick_X");
+		public FKey OculusTouch_Left_Thumbstick_Y = new FKey("OculusTouch_Left_Thumbstick_Y");
+		public FKey OculusTouch_Left_Thumbstick_Click = new FKey("OculusTouch_Left_Thumbstick_Click");
+		public FKey OculusTouch_Left_Thumbstick_Touch = new FKey("OculusTouch_Left_Thumbstick_Touch");
+		public FKey OculusTouch_Left_Thumbstick_Up = new FKey("OculusTouch_Left_Thumbstick_Up");
+		public FKey OculusTouch_Left_Thumbstick_Down = new FKey("OculusTouch_Left_Thumbstick_Down");
+		public FKey OculusTouch_Left_Thumbstick_Left = new FKey("OculusTouch_Left_Thumbstick_Left");
+		public FKey OculusTouch_Left_Thumbstick_Right = new FKey("OculusTouch_Left_Thumbstick_Right");
+		public FKey OculusTouch_Right_A_Click = new FKey("OculusTouch_Right_A_Click");
+		public FKey OculusTouch_Right_B_Click = new FKey("OculusTouch_Right_B_Click");
+		public FKey OculusTouch_Right_A_Touch = new FKey("OculusTouch_Right_A_Touch");
+		public FKey OculusTouch_Right_B_Touch = new FKey("OculusTouch_Right_B_Touch");
+		public FKey OculusTouch_Right_System_Click = new FKey("OculusTouch_Right_System_Click");
+		public FKey OculusTouch_Right_Grip_Click = new FKey("OculusTouch_Right_Grip_Click");
+		public FKey OculusTouch_Right_Grip_Axis = new FKey("OculusTouch_Right_Grip_Axis");
+		public FKey OculusTouch_Right_Trigger_Click = new FKey("OculusTouch_Right_Trigger_Click");
+		public FKey OculusTouch_Right_Trigger_Axis = new FKey("OculusTouch_Right_Trigger_Axis");
+		public FKey OculusTouch_Right_Trigger_Touch = new FKey("OculusTouch_Right_Trigger_Touch");
+		public FKey OculusTouch_Right_Thumbstick_2D = new FKey("OculusTouch_Right_Thumbstick_2D");
+		public FKey OculusTouch_Right_Thumbstick_X = new FKey("OculusTouch_Right_Thumbstick_X");
+		public FKey OculusTouch_Right_Thumbstick_Y = new FKey("OculusTouch_Right_Thumbstick_Y");
+		public FKey OculusTouch_Right_Thumbstick_Click = new FKey("OculusTouch_Right_Thumbstick_Click");
+		public FKey OculusTouch_Right_Thumbstick_Touch = new FKey("OculusTouch_Right_Thumbstick_Touch");
+		public FKey OculusTouch_Right_Thumbstick_Up = new FKey("OculusTouch_Right_Thumbstick_Up");
+		public FKey OculusTouch_Right_Thumbstick_Down = new FKey("OculusTouch_Right_Thumbstick_Down");
+		public FKey OculusTouch_Right_Thumbstick_Left = new FKey("OculusTouch_Right_Thumbstick_Left");
+		public FKey OculusTouch_Right_Thumbstick_Right = new FKey("OculusTouch_Right_Thumbstick_Right");
+
+		// Valve Index Controls
+		public FKey ValveIndex_Left_A_Click = new FKey("ValveIndex_Left_A_Click");
+		public FKey ValveIndex_Left_B_Click = new FKey("ValveIndex_Left_B_Click");
+		public FKey ValveIndex_Left_A_Touch = new FKey("ValveIndex_Left_A_Touch");
+		public FKey ValveIndex_Left_B_Touch = new FKey("ValveIndex_Left_B_Touch");
+		public FKey ValveIndex_Left_System_Click = new FKey("ValveIndex_Left_System_Click");
+		public FKey ValveIndex_Left_System_Touch = new FKey("ValveIndex_Left_System_Touch");
+		public FKey ValveIndex_Left_Grip_Axis = new FKey("ValveIndex_Left_Grip_Axis");
+		public FKey ValveIndex_Left_Grip_Force = new FKey("ValveIndex_Left_Grip_Force");
+		public FKey ValveIndex_Left_Trigger_Click = new FKey("ValveIndex_Left_Trigger_Click");
+		public FKey ValveIndex_Left_Trigger_Axis = new FKey("ValveIndex_Left_Trigger_Axis");
+		public FKey ValveIndex_Left_Trigger_Touch = new FKey("ValveIndex_Left_Trigger_Touch");
+		public FKey ValveIndex_Left_Thumbstick_2D = new FKey("ValveIndex_Left_Thumbstick_2D");
+		public FKey ValveIndex_Left_Thumbstick_X = new FKey("ValveIndex_Left_Thumbstick_X");
+		public FKey ValveIndex_Left_Thumbstick_Y = new FKey("ValveIndex_Left_Thumbstick_Y");
+		public FKey ValveIndex_Left_Thumbstick_Click = new FKey("ValveIndex_Left_Thumbstick_Click");
+		public FKey ValveIndex_Left_Thumbstick_Touch = new FKey("ValveIndex_Left_Thumbstick_Touch");
+		public FKey ValveIndex_Left_Thumbstick_Up = new FKey("ValveIndex_Left_Thumbstick_Up");
+		public FKey ValveIndex_Left_Thumbstick_Down = new FKey("ValveIndex_Left_Thumbstick_Down");
+		public FKey ValveIndex_Left_Thumbstick_Left = new FKey("ValveIndex_Left_Thumbstick_Left");
+		public FKey ValveIndex_Left_Thumbstick_Right = new FKey("ValveIndex_Left_Thumbstick_Right");
+		public FKey ValveIndex_Left_Trackpad_2D = new FKey("ValveIndex_Left_Trackpad_2D");
+		public FKey ValveIndex_Left_Trackpad_X = new FKey("ValveIndex_Left_Trackpad_X");
+		public FKey ValveIndex_Left_Trackpad_Y = new FKey("ValveIndex_Left_Trackpad_Y");
+		public FKey ValveIndex_Left_Trackpad_Force = new FKey("ValveIndex_Left_Trackpad_Force");
+		public FKey ValveIndex_Left_Trackpad_Touch = new FKey("ValveIndex_Left_Trackpad_Touch");
+		public FKey ValveIndex_Left_Trackpad_Up = new FKey("ValveIndex_Left_Trackpad_Up");
+		public FKey ValveIndex_Left_Trackpad_Down = new FKey("ValveIndex_Left_Trackpad_Down");
+		public FKey ValveIndex_Left_Trackpad_Left = new FKey("ValveIndex_Left_Trackpad_Left");
+		public FKey ValveIndex_Left_Trackpad_Right = new FKey("ValveIndex_Left_Trackpad_Right");
+		public FKey ValveIndex_Right_A_Click = new FKey("ValveIndex_Right_A_Click");
+		public FKey ValveIndex_Right_B_Click = new FKey("ValveIndex_Right_B_Click");
+		public FKey ValveIndex_Right_A_Touch = new FKey("ValveIndex_Right_A_Touch");
+		public FKey ValveIndex_Right_B_Touch = new FKey("ValveIndex_Right_B_Touch");
+		public FKey ValveIndex_Right_System_Click = new FKey("ValveIndex_Right_System_Click");
+		public FKey ValveIndex_Right_System_Touch = new FKey("ValveIndex_Right_System_Touch");
+		public FKey ValveIndex_Right_Grip_Axis = new FKey("ValveIndex_Right_Grip_Axis");
+		public FKey ValveIndex_Right_Grip_Force = new FKey("ValveIndex_Right_Grip_Force");
+		public FKey ValveIndex_Right_Trigger_Click = new FKey("ValveIndex_Right_Trigger_Click");
+		public FKey ValveIndex_Right_Trigger_Axis = new FKey("ValveIndex_Right_Trigger_Axis");
+		public FKey ValveIndex_Right_Trigger_Touch = new FKey("ValveIndex_Right_Trigger_Touch");
+		public FKey ValveIndex_Right_Thumbstick_2D = new FKey("ValveIndex_Right_Thumbstick_2D");
+		public FKey ValveIndex_Right_Thumbstick_X = new FKey("ValveIndex_Right_Thumbstick_X");
+		public FKey ValveIndex_Right_Thumbstick_Y = new FKey("ValveIndex_Right_Thumbstick_Y");
+		public FKey ValveIndex_Right_Thumbstick_Click = new FKey("ValveIndex_Right_Thumbstick_Click");
+		public FKey ValveIndex_Right_Thumbstick_Touch = new FKey("ValveIndex_Right_Thumbstick_Touch");
+		public FKey ValveIndex_Right_Thumbstick_Up = new FKey("ValveIndex_Right_Thumbstick_Up");
+		public FKey ValveIndex_Right_Thumbstick_Down = new FKey("ValveIndex_Right_Thumbstick_Down");
+		public FKey ValveIndex_Right_Thumbstick_Left = new FKey("ValveIndex_Right_Thumbstick_Left");
+		public FKey ValveIndex_Right_Thumbstick_Right = new FKey("ValveIndex_Right_Thumbstick_Right");
+		public FKey ValveIndex_Right_Trackpad_2D = new FKey("ValveIndex_Right_Trackpad_2D");
+		public FKey ValveIndex_Right_Trackpad_X = new FKey("ValveIndex_Right_Trackpad_X");
+		public FKey ValveIndex_Right_Trackpad_Y = new FKey("ValveIndex_Right_Trackpad_Y");
+		public FKey ValveIndex_Right_Trackpad_Force = new FKey("ValveIndex_Right_Trackpad_Force");
+		public FKey ValveIndex_Right_Trackpad_Touch = new FKey("ValveIndex_Right_Trackpad_Touch");
+		public FKey ValveIndex_Right_Trackpad_Up = new FKey("ValveIndex_Right_Trackpad_Up");
+		public FKey ValveIndex_Right_Trackpad_Down = new FKey("ValveIndex_Right_Trackpad_Down");
+		public FKey ValveIndex_Right_Trackpad_Left = new FKey("ValveIndex_Right_Trackpad_Left");
+		public FKey ValveIndex_Right_Trackpad_Right = new FKey("ValveIndex_Right_Trackpad_Right");
+
+		/*FKey Virtual_Accept;
+		FKey Virtual_Back;
+
+		FKey Invalid;*/
+		public EKeys()
+		{
+			UnrealTargetPlatform.TryParse(Globals.Params.ParseValue("Platform", "Win64"), out UnrealTargetPlatform Platform);
+			if (Platform == UnrealTargetPlatform.Mac)
+			{
+				BackSpace = new FKey("Delete");
+				Delete = new FKey("ForwardDelete");
+			}
+
+		}
+
+	}
+
+	// Same as EInputEvent from EngineBaseTypes.h
+	public enum EInputEvent
+	{
+		IE_Pressed,
+		IE_Released,
+		IE_Repeat,
+		IE_DoubleClick,
+		IE_Axis,
+		IE_Max
+
+	}
+
+	public class InputAction
+	{
+		public FKey Key;
+		public EInputEvent KeyAction;
+		public double XDelta;
+		public double YDelta;
+
+		public InputAction(FKey Inkey, EInputEvent InInputEvent, double InXDelta = 0, double InYDelta = 0)
+		{
+			this.Key = Inkey;
+			this.KeyAction = InInputEvent;
+			this.XDelta = InXDelta;
+			this.YDelta = InYDelta;
+		}
+
+	}
+}
