Add Copyright notice to Gauntlet.InputActions

diff --git a/Engine/Source/Programs/AutomationTool/Gauntlet/Framework/Utils/Gauntlet.InputActions.cs b/Engine/Source/Programs/AutomationTool/Gauntlet/Framework/Utils/Gauntlet.InputActions.cs
--- a/Engine/Source/Programs/AutomationTool/Gauntlet/Framework/Utils/Gauntlet.InputActions.cs
+++ b/Engine/Source/Programs/AutomationTool/Gauntlet/Framework/Utils/Gauntlet.InputActions.cs
@@ -1,4 +1,6 @@
-﻿using EpicGames.Core;
+﻿// Copyright Epic Games, Inc. All Rights Reserved.
+
+using EpicGames.Core;
 using System;
 using UnrealBuildTool;
 
