Added command line switch '-ApplyIoStoreOnDemand' to UAT BuildCookRun that will force the IoStoreOnDemand module into the target project and force on additional settings to allow IoStoreOnDemand to work.

#rnx

- This is only intended for development and testing IoStoreOnDemand as it allows us to apply it to a project without changing the project files on disk. This should not be used to enable the feature on a project in production.

### Enabling the code module
- When forcing on IoStoreOnDemand, BuildCookRun will pass '-CompileIoStoreOnDemand' to UnrealBuildTool which will set 'bCompileIoStoreOnDemand' to true in the target rules. This in turn will add the 'IoStoreOnDemand' module to the Engine module but only for runtime targets.
- None of this should be used when enabling the feature for a project permanentl, instead just add 'IoStoreOnDemand' module to your project's *.build.cs

### Staging SSL certificates
- IoStoreOnDemand needs access to the engine/project#s SSL certificates at runtime  so we need to stage them, which is done by forcing 'bStageSSLCertificates' to be true during BuildCookRun.
- To do this properly for a project set ini:Engine:[/Script/Engine.NetworkSettings]:n.VerifyPeer=true

### Generate Chunks
- In order for IoStoreOnDemand to work we need to generate chunks (more than one container file) which can be done by forcing ProjectParams.Manifests to true.
- To do this properly for a project set ini:Game:[/Script/UnrealEd.ProjectPackagingSettings]:bGenerateChunks=true

### Pak Rules
- Once chunked builds are enabled IoStoreOnDemand will require at least one PakRule that sets  'bOnDemand' to true. So first we check to see if any rules like this for the existing project and if not we force add one via CreateIoStoreOnDemandPakRules. This rule will include all *.ubulk and *.uptnl files except ones found under a \map\ directory. This is because some level bulkdata files can get too large for IoStoreOnDemand to support so excluding anything under a map directory is a quick and easy way to try and work around that. The end result is that other than map data, pretty much all bulkdata will end up on demand, which is a good stress test for the system but not a great representation on how the system should be used.
- If you are trying to set up rules properly you'd need to consider what data you actually want to stream over the network and what data you want to remain on disk, which would greatly depend on your project set up. One popular choice is to only stream optional (.uptnl) bulkdata and leave the required bulkdata to be found on disk.

### Enabling Persistent Local Storage
- IoStoreOnDemand requires persistent local storage to be set aside for it's caching which may or may not require special set up for a project depending on the target platform. If we detect that the current project has not been set up for persistent local storage we will attempt to enable it for the project during staging.
- To set this up properly you will need to consult the documentation for your target platform and follow the instructions.

diff --git a/Engine/Source/Programs/AutomationTool/AutomationUtils/ProjectParams.cs b/Engine/Source/Programs/AutomationTool/AutomationUtils/ProjectParams.cs
--- a/Engine/Source/Programs/AutomationTool/AutomationUtils/ProjectParams.cs
+++ b/Engine/Source/Programs/AutomationTool/AutomationUtils/ProjectParams.cs
@@ -1061,7 +1061,17 @@ namespace AutomationTool
 			}
 
 			this.Upload = Command.ParseParamValue("upload");
-			
+
+			this.ApplyIoStoreOnDemand = Command.ParseParam("ApplyIoStoreOnDemand");
+			if (this.ApplyIoStoreOnDemand)
+			{
+				// Compile the IoStoreOnDemand module when uploading content for streaming
+				this.AdditionalBuildOptions += " -CompileIoStoreOnDemand";
+
+				// Required to force the project to use chunk manifests
+				this.Manifests = true;
+			}
+
 			this.XcodeBuildOptions = ParseParamValueIfNotSpecified(Command, XcodeBuildOptions, "xcodebuildoptions", null);
 
 			if (ClientConfigsToBuild == null)
@@ -2401,6 +2411,9 @@ namespace AutomationTool
 		[Help("upload", "Arguments for uploading on demand content")]
 		public string Upload { get; set; }
 
+		[Help("applyiostoreondemand", "Forces IoStoreOnDemand to be enabled for the project even if it is not set up for it")]
+		public bool ApplyIoStoreOnDemand { get; }
+
 		[Help("XcodeBuildOptions", "Extra options to pass to xcodebuild")]
 		public string XcodeBuildOptions { get; set; }
 
diff --git a/Engine/Source/Programs/AutomationTool/Scripts/CopyBuildToStagingDirectory.Automation.cs b/Engine/Source/Programs/AutomationTool/Scripts/CopyBuildToStagingDirectory.Automation.cs
--- a/Engine/Source/Programs/AutomationTool/Scripts/CopyBuildToStagingDirectory.Automation.cs
+++ b/Engine/Source/Programs/AutomationTool/Scripts/CopyBuildToStagingDirectory.Automation.cs
@@ -1642,6 +1642,12 @@ namespace AutomationScripts
 						PlatformEngineConfig.GetBool("/Script/Engine.NetworkSettings", "n.VerifyPeer", out bStageSSLCertificates);
 					}
 
+					if (Params.ApplyIoStoreOnDemand)
+					{
+						// The IoStoreOnDemand system requires SSL certificates at runtime so make sure we stage them
+						bStageSSLCertificates = true;
+					}
+
 					if (bStageSSLCertificates)
 					{
 						// Game's SSL certs
@@ -2607,10 +2613,49 @@ namespace AutomationScripts
 					RulesList.Add(PakRules);
 				}
 			}
+
+			if (Params.ApplyIoStoreOnDemand)
+			{
+				// If there are no pak rules for an 'OnDemand' container we will need to create one via CreateIoStoreOnDemandPakRules
+				if (RulesList.FindIndex(x => x.bOnDemand == true) == -1)
+				{
+					RulesList.AddRange(CreateIoStoreOnDemandPakRules());
+				}
+			}
+
 			if (RulesList.Count == 0)
 			{
 				return null;
 			}
+
+			return RulesList;
+		}
+
+		/// <summary>
+		/// Creates a list of pak file rules that will need to be applied in order for IoStoreOnDemand to work by forcing bulkdata
+		/// into a 'OnDemand' container.
+		/// </summary>
+		/// <returns>The list of rules </returns>
+		private static List<PakFileRules> CreateIoStoreOnDemandPakRules()
+		{
+			List<PakFileRules> RulesList = new List<PakFileRules>();
+
+			PakFileRules PakRules = new PakFileRules();
+			PakRules.Name = "BulkDataOnDemand";
+			PakRules.bOnDemand = true;
+			PakRules.Filter = new FileFilter();
+
+			PakRules.Filter.AddRule("*.uptnl");
+			PakRules.Filter.AddRule("*.ubulk");
+
+			// Map files can be too large for IAS to reasonably handle but with CookToZen it would be difficult for us
+			// to filter by file sizeand it is difficult to justify doing that work given that this is a testing/development
+			//feature.
+			// So for now we should exclude anything under /Map/ directories which should solve most problems.
+			PakRules.Filter.AddRule("-.../Map/*");
+
+			RulesList.Add(PakRules);
+
 			return RulesList;
 		}
 
diff --git a/Engine/Source/Programs/UnrealBuildTool/Configuration/ReadOnlyTargetRules.cs b/Engine/Source/Programs/UnrealBuildTool/Configuration/ReadOnlyTargetRules.cs
--- a/Engine/Source/Programs/UnrealBuildTool/Configuration/ReadOnlyTargetRules.cs
+++ b/Engine/Source/Programs/UnrealBuildTool/Configuration/ReadOnlyTargetRules.cs
@@ -170,6 +170,8 @@ namespace UnrealBuildTool
 
 		public bool bUseGameplayDebuggerCore => Inner.bUseGameplayDebuggerCore;
 
+		public bool bCompileIoStoreOnDemand => Inner.bCompileIoStoreOnDemand;
+
 		public bool bUseIris => Inner.bUseIris;
 
 		public bool bTrackRHIResourceInfoForTest => Inner.bTrackRHIResourceInfoForTest;
diff --git a/Engine/Source/Programs/UnrealBuildTool/Configuration/TargetRules.cs b/Engine/Source/Programs/UnrealBuildTool/Configuration/TargetRules.cs
--- a/Engine/Source/Programs/UnrealBuildTool/Configuration/TargetRules.cs
+++ b/Engine/Source/Programs/UnrealBuildTool/Configuration/TargetRules.cs
@@ -1023,6 +1023,13 @@ namespace UnrealBuildTool
 		}
 		GameplayDebuggerOverrideState UseGameplayDebuggerOverride;
 
+		/// <summary>
+		/// Whether to use I/O store on-demand
+		/// </summary>
+		[RequiresUniqueBuildEnvironment]
+		[CommandLine("-CompileIoStoreOnDemand", Value = "true")]
+		public bool bCompileIoStoreOnDemand { get; set; } = false;
+
 		/// <summary>
 		/// Whether to use Iris.
 		/// </summary>
diff --git a/Engine/Source/Runtime/Engine/Engine.Build.cs b/Engine/Source/Runtime/Engine/Engine.Build.cs
--- a/Engine/Source/Runtime/Engine/Engine.Build.cs
+++ b/Engine/Source/Runtime/Engine/Engine.Build.cs
@@ -117,6 +117,11 @@ public class Engine : ModuleRules
 			PublicDependencyModuleNames.Add("ApplicationCore");
 		}
 
+		if (Target.bCompileIoStoreOnDemand && !Target.bBuildEditor)
+		{
+			PrivateDependencyModuleNames.Add("IoStoreOnDemand");
+		}
+
 		PublicIncludePathModuleNames.AddRange(
 			new string[] {
 				"TypedElementFramework",
