Add additional location to check for Android Studio binary

#android
#11740

diff --git a/Engine/Source/Programs/AutomationTool/Android/AndroidPlatform.Automation.cs b/Engine/Source/Programs/AutomationTool/Android/AndroidPlatform.Automation.cs
--- a/Engine/Source/Programs/AutomationTool/Android/AndroidPlatform.Automation.cs
+++ b/Engine/Source/Programs/AutomationTool/Android/AndroidPlatform.Automation.cs
@@ -302,6 +302,8 @@ public class AndroidPlatform : Platform
 		string DefaultAndroidStudioInstallDir = Path.Combine(Environment.GetFolderPath(Environment.SpecialFolder.ProgramFiles), "Android", "Android Studio");
 		string RegValue = Microsoft.Win32.Registry.GetValue("HKEY_LOCAL_MACHINE\\SOFTWARE\\Android Studio", "Path", null) as string;
 		string AndroidStudioInstallDir = RegValue == null ? DefaultAndroidStudioInstallDir : RegValue;
+  		// Some installs, like JetBrains Toolbox, may not place an entry in the registry so try an alternate location
+		AndroidStudioInstallDir = Directory.Exists(AndroidStudioInstallDir) ? AndroidStudioInstallDir : Path.Combine(Environment.GetFolderPath(Environment.SpecialFolder.LocalApplicationData), "Programs", "Android Studio");
 		return Path.Combine(AndroidStudioInstallDir, "bin", "studio64.exe");
 	}
 
