[Gauntlet] Fix UnrealAutomationEntry timestamp property name case for json output

#rnx

diff --git a/Engine/Source/Programs/AutomationTool/Gauntlet/Unreal/Utils/Gauntlet.UnrealTestResults.cs b/Engine/Source/Programs/AutomationTool/Gauntlet/Unreal/Utils/Gauntlet.UnrealTestResults.cs
--- a/Engine/Source/Programs/AutomationTool/Gauntlet/Unreal/Utils/Gauntlet.UnrealTestResults.cs
+++ b/Engine/Source/Programs/AutomationTool/Gauntlet/Unreal/Utils/Gauntlet.UnrealTestResults.cs
@@ -161,7 +161,7 @@ namespace Gauntlet
 		public string Filename { get; set; }
 		[JsonPropertyName("lineNumber")]
 		public int LineNumber { get; set; }
-		[JsonPropertyName("timeStamp")]
+		[JsonPropertyName("timestamp")]
 		public string Timestamp { get; set; }
 
 		public const string InvalidDateTime = "0001.01.01-00.00.00";
