[Gauntet] Modify Gauntlet Device reservations to allow for an explicit device within an UnrealTargetConstraint

diff --git a/Engine/Source/Programs/AutomationTool/AutomationUtils/DeviceReservation.cs b/Engine/Source/Programs/AutomationTool/AutomationUtils/DeviceReservation.cs
--- a/Engine/Source/Programs/AutomationTool/AutomationUtils/DeviceReservation.cs
+++ b/Engine/Source/Programs/AutomationTool/AutomationUtils/DeviceReservation.cs
@@ -11,14 +11,13 @@ using System.Text.Json;
 using System.Text.Json.Serialization;
 using UnrealBuildBase;
 
-
 #pragma warning disable SYSLIB0014
 
 namespace AutomationTool.DeviceReservation
 {
 	/// <summary>
 	/// Co-operatively reserves remote devices for build automation.
-	/// 
+	///
 	/// The constructor blocks until the specified type and number of devices are available.
 	/// The reservation is automatically renewed for the lifetime of this object, and released
 	/// when the object is disposed or garbage collected.
@@ -79,6 +78,28 @@ namespace AutomationTool.DeviceReservation
 			RenewThread.Start();
 		}
 
+		/// <summary>
+		/// Creates a device reservation for the specified type and number of devices.
+		/// Blocks until the devices are available.
+		/// </summary>
+		/// <param name="InReservationBaseUri">Working directory which contains the devices.xml and reservations.xml files. Usually a network share</param>
+		/// <param name="Reservation">Reservation to continually renew</param>
+		public DeviceReservationAutoRenew(string InReservationBaseUri, Reservation Reservation)
+		{
+			ReservationBaseUri = new Uri(InReservationBaseUri);
+			ActiveReservation = Reservation;
+
+			// Resolve the device IPs
+			ReservedDevices = new List<Device>();
+			foreach (string DeviceName in ActiveReservation.DeviceNames)
+			{
+				ReservedDevices.Add(Device.Get(ReservationBaseUri, DeviceName));
+			}
+
+			RenewThread = new Thread(DoAutoRenew);
+			RenewThread.Start();
+		}
+
 		private void DoAutoRenew()
 		{
 			int RetryCurrent = 0;
@@ -262,9 +283,10 @@ namespace AutomationTool.DeviceReservation
 			public string PoolId { get; set; }
 			public string JobId { get; set; }
 			public string StepId { get; set; }
+			public string DeviceName { get; set; }
 		}
 
-		public static Reservation Create(Uri BaseUri, string[] DeviceTypes, TimeSpan Duration, int RetryMax = 5, string PoolID = "")
+		public static Reservation Create(Uri BaseUri, string[] DeviceTypes, TimeSpan Duration, int RetryMax = 5, string PoolID = "", string DeviceName = "")
 		{
 			bool bFirst = true;
 			TimeSpan RetryTime = TimeSpan.FromMinutes(1);
@@ -293,7 +315,8 @@ namespace AutomationTool.DeviceReservation
 						ReservationDetails = ReservationDetails,
 						PoolId = PoolID,
 						JobId = Environment.GetEnvironmentVariable("UE_HORDE_JOBID"),
-						StepId = Environment.GetEnvironmentVariable("UE_HORDE_STEPID")
+						StepId = Environment.GetEnvironmentVariable("UE_HORDE_STEPID"),
+						DeviceName = DeviceName
 					});
 				}
 				catch (WebException WebEx)
diff --git a/Engine/Source/Programs/AutomationTool/Gauntlet/Framework/Devices/Gauntlet.DevicePool.cs b/Engine/Source/Programs/AutomationTool/Gauntlet/Framework/Devices/Gauntlet.DevicePool.cs
--- a/Engine/Source/Programs/AutomationTool/Gauntlet/Framework/Devices/Gauntlet.DevicePool.cs
+++ b/Engine/Source/Programs/AutomationTool/Gauntlet/Framework/Devices/Gauntlet.DevicePool.cs
@@ -67,12 +67,14 @@ namespace Gauntlet
 		public readonly UnrealTargetPlatform? Platform;
 		public readonly EPerfSpec PerfSpec;
 		public readonly string Model;
+		public readonly string DeviceName;
 
-		public UnrealDeviceTargetConstraint(UnrealTargetPlatform? Platform, EPerfSpec PerfSpec = EPerfSpec.Unspecified, string Model = null)
+		public UnrealDeviceTargetConstraint(UnrealTargetPlatform? Platform, EPerfSpec PerfSpec = EPerfSpec.Unspecified, string Model = null, string DeviceName = null)
 		{
 			this.Platform = Platform;
 			this.PerfSpec = PerfSpec;
 			this.Model = Model == null ? string.Empty : Model;
+			this.DeviceName = DeviceName == null ? string.Empty : DeviceName;
 		}
 
 		/// <summary>
@@ -80,7 +82,7 @@ namespace Gauntlet
 		/// </summary>
 		public bool IsIdentity()
 		{
-			return (PerfSpec == EPerfSpec.Unspecified) && (Model == string.Empty);
+			return (PerfSpec == EPerfSpec.Unspecified) && (Model == string.Empty) && (DeviceName == string.Empty);
 		}
 
 
@@ -106,7 +108,8 @@ namespace Gauntlet
 
 			bool ModelMatch = Model == string.Empty ? true : Model.Equals(DeviceDef.Model, StringComparison.InvariantCultureIgnoreCase);
 			bool PerfMatch = (PerfSpec == EPerfSpec.Unspecified) ? true : PerfSpec == DeviceDef.PerfSpec;
-			return ModelMatch && PerfMatch;
+			bool NameMatch = DeviceName == string.Empty? true : DeviceName.Equals(DeviceDef.Name, StringComparison.InvariantCultureIgnoreCase);
+			return ModelMatch && PerfMatch && NameMatch;
 		}
 
 		public bool Equals(UnrealDeviceTargetConstraint Other)
@@ -121,28 +124,29 @@ namespace Gauntlet
 				return true;
 			}
 
-			if (Other.Platform != Platform)
+			return Other.Platform == Platform
+				&& Other.Model.Equals(Model, StringComparison.InvariantCultureIgnoreCase)
+				&& Other.DeviceName.Equals(DeviceName, StringComparison.InvariantCultureIgnoreCase)
+				&& Other.PerfSpec == PerfSpec;
+		}
+
+		public override bool Equals(object Obj)
+		{
+			if (ReferenceEquals(Obj, null))
 			{
-				return false;
+				throw new AutomationException("Comparing null target constraint");
 			}
 
-			if (Other.Model.Equals(Model, StringComparison.InvariantCultureIgnoreCase))
+			if (ReferenceEquals(this, Obj))
 			{
 				return true;
 			}
 
-			return Other.PerfSpec == PerfSpec;
-		}
-
-		public override bool Equals(object Obj)
-		{
-			if (ReferenceEquals(Obj, null))
+			if (Obj.GetType() != typeof(UnrealDeviceTargetConstraint))
 			{
-				throw new AutomationException("Comparing null target constraint");
+				return false;
 			}
 
-			if (ReferenceEquals(this, Obj)) return true;
-			if (Obj.GetType() != typeof(UnrealDeviceTargetConstraint)) return false;
 			return Equals((UnrealDeviceTargetConstraint)Obj);
 		}
 
@@ -163,17 +167,20 @@ namespace Gauntlet
 
 		public override string ToString()
 		{
-			if (PerfSpec == EPerfSpec.Unspecified && Model == string.Empty)
+			string Value = Platform.ToString();
+			if (PerfSpec != EPerfSpec.Unspecified)
 			{
-				return string.Format("{0}", Platform);
+				Value = string.Format("{0}:{1}", Value, PerfSpec.ToString());
 			}
-
-			return string.Format("{0}:{1}", Platform, Model == string.Empty ? PerfSpec.ToString() : Model);
-		}
-
-		public override int GetHashCode()
-		{
-			return ToString().GetHashCode();
+			if(Model != string.Empty)
+			{
+				Value = string.Format("{0}:{1}", Value, Model);
+			}
+			if (DeviceName != string.Empty)
+			{
+				Value = string.Format("{0}:{1}", Value, DeviceName);
+			}
+			return Value;
 		}
 
 		/// <summary>
@@ -184,6 +191,11 @@ namespace Gauntlet
 		{
 			return string.Format("{0}:{1}", Platform, Model == string.Empty ? PerfSpec.ToString() : Model);
 		}
+
+		public override int GetHashCode()
+		{
+			return ToString().GetHashCode();
+		}
 	}
 
 	/// <summary>
@@ -201,150 +213,6 @@ namespace Gauntlet
 		public UnrealTargetPlatform Platform;
 	}
 
-	/// <summary>
-	/// Device reservation states
-	/// </summary>
-	public class DeviceReservationStates
-	{
-		private Dictionary<string, ResevationState> RequiredCollection;
-		private Dictionary<string, int> ProblemCollection;
-
-		public DeviceReservationStates()
-		{
-			ResetRequirements();
-			ProblemCollection = new Dictionary<string, int>();
-		}
-
-		public void ResetRequirements()
-		{
-			RequiredCollection = new Dictionary<string, ResevationState>();
-		}
-
-		public void AddRequired(string DeviceKey, int Count = 1)
-		{
-			if(RequiredCollection.ContainsKey(DeviceKey))
-			{
-				var ReservationState = RequiredCollection[DeviceKey];
-				ReservationState.NumRequired += Count;
-				ReservationState.NumLacking += Count;
-			}
-			else
-			{
-				RequiredCollection.Add(DeviceKey, new ResevationState(Count));
-			}
-		}
-
-		public void MarkReserved(string DeviceKey, int Count = 1)
-		{
-			var ReservationState = FindClosestState(DeviceKey);
-			if (ReservationState != null && ReservationState.NumLacking > 0)
-			{
-				ReservationState.NumLacking -= Count;
-			}
-		}
-
-		public void MarkReserved(Device InDevice, int Count = 1)
-		{
-			MarkReserved(KeyFromDevice(InDevice), Count);
-		}
-
-		public void MarkProblem(string DeviceKey, int Count = 1)
-		{
-			if (DeviceKey.Contains(":"))
-			{
-				DeviceKey = DeviceKey.Substring(0, DeviceKey.IndexOf(":"));
-			}
-
-			if (!ProblemCollection.ContainsKey(DeviceKey))
-			{
-				ProblemCollection[DeviceKey] = 0;
-			}
-
-			ProblemCollection[DeviceKey] += Count;
-		}
-
-		public void MarkProblem(Device InDevice, int Count = 1)
-		{
-			MarkProblem(KeyFromDevice(InDevice), Count);
-		}
-
-		/// <summary>
-		/// Return the list of requirements for reservation
-		/// </summary>
-		/// <returns></returns>
-		public string[] GetRequirements()
-		{
-			return RequiredCollection.SelectMany(P => Enumerable.Repeat(P.Key, P.Value.NumRequired)).ToArray();
-		}
-
-		/// <summary>
-		/// Return the list of devices and their reservation status as an array
-		/// </summary>
-		/// <returns></returns>
-		public string[] GetStatuses()
-		{
-			return RequiredCollection.Select(P => $"{P.Key.Replace($":{EPerfSpec.Unspecified}", "")} " +
-											$"[" +
-												$"lacking:{P.Value.NumLacking}, " +
-												$"reserved:{P.Value.NumRequired - P.Value.NumLacking}" +
-											$"]").ToArray();
-		}
-
-		/// <summary>
-		/// Return the list of devices marked with problem
-		/// </summary>
-		/// <returns></returns>
-		public string[] GetProblems()
-		{
-			return ProblemCollection.Select(P => $"{P.Value} {P.Key}").ToArray();
-		}
-
-		public int RequiredCount()
-		{
-			return RequiredCollection.Sum(C => C.Value.NumRequired);
-		}
-
-		public int ProblemCount()
-		{
-			return ProblemCollection.Sum(C => C.Value);
-		}
-
-		private class ResevationState
-		{
-			public ResevationState(int Count = 1)
-			{
-				NumRequired = Count;
-				NumLacking = Count;
-			}
-
-			public int NumRequired;
-			public int NumLacking;
-		}
-
-		private ResevationState FindClosestState(string DeviceKey)
-		{
-			if (RequiredCollection.ContainsKey(DeviceKey))
-			{
-				return RequiredCollection[DeviceKey];
-			}
-			else if (DeviceKey.Contains(":"))
-			{
-				string DeviceUndefinedKey = DeviceKey.Substring(0, DeviceKey.IndexOf(":") + 1) + EPerfSpec.Unspecified.ToString();
-				if (RequiredCollection.ContainsKey(DeviceUndefinedKey))
-				{
-					return RequiredCollection[DeviceUndefinedKey];
-				}
-			}
-
-			return null;
-		}
-
-		private string KeyFromDevice(Device InDevice)
-		{
-			return InDevice.Type.Replace("-DevKit", "", StringComparison.OrdinalIgnoreCase) + ":" + (InDevice.PerfSpec == string.Empty ? InDevice.Model : InDevice.PerfSpec);
-		}
-	}
-
 	/// <summary>
 	/// Singleton class that's responsible for providing a list of devices and reserving them. Code should call
 	/// EnumerateDevices to build a list of desired devices, which must then be reserved by calling ReserveDevices.
@@ -392,11 +260,6 @@ namespace Gauntlet
 		/// </summary>
 		private HashSet<UnrealTargetPlatform?> UsedPlatforms = new HashSet<UnrealTargetPlatform?>();
 
-		/// <summary>
-		/// Device reservation left to fulfill
-		/// </summary>
-		private DeviceReservationStates ReservationStates;
-
 		/// <summary>
 		/// Device to reservation lookup
 		/// </summary>
@@ -446,7 +309,6 @@ namespace Gauntlet
 		{
 			// create two local devices by default?
 			AddLocalDevices(2);
-			ReservationStates = new DeviceReservationStates();
 
 			if(Instance == null)
 			{
@@ -536,18 +398,6 @@ namespace Gauntlet
 		}
 		#endregion
 
-		/// <summary>
-		/// Get the reservation status as an array
-		/// </summary>
-		/// <returns></returns>
-		public string[] GetReservationStatuses() => ReservationStates != null ? ReservationStates.GetStatuses() : null;
-
-		/// <summary>
-		/// Get the reserved device marked with problems as an array
-		/// </summary>
-		/// <returns></returns>
-		public string[] GetDeviceProblems() => ReservationStates != null ? ReservationStates.GetProblems() : null;
-
 		/// <summary>
 		/// Returns the number of available devices of the provided type. This includes unprovisioned devices but not reserved ones.
 		/// Note: unprovisioned devices are currently only returned when device is not constrained
@@ -879,7 +729,7 @@ namespace Gauntlet
 
 							// try to create a device. This can fail, but if so we'll just end up back here
 							// on the next iteration
-							ITargetDevice NewDevice = CreateAndRegisterDeviceFromDefinition(NextDeviceDef);
+							ITargetDevice NewDevice = CreateAndRegisterDeviceFromDefinition(NextDeviceDef, Constraint);
 
 							MatchingUnprovisionedDevices.Remove(NextDeviceDef);
 							UnprovisionedDevices.Remove(NextDeviceDef);
@@ -939,32 +789,24 @@ namespace Gauntlet
 		/// </summary>
 		public bool ReserveDevicesFromService(string DeviceURL, Dictionary<UnrealDeviceTargetConstraint, int> DeviceTypes)
 		{
-			if (String.IsNullOrEmpty(DeviceURL))
+			if (string.IsNullOrEmpty(DeviceURL) || !Uri.TryCreate(DeviceURL, UriKind.Absolute, out Uri ReservationServerUri))
 			{
 				return false;
 			}
 
-			ReservationStates.ResetRequirements();
+			Dictionary<UnrealTargetPlatform, string> DeviceMap = UnrealTargetPlatform.GetValidPlatformNames()
+				.ToDictionary(Platform => UnrealTargetPlatform.Parse(Platform), Platform => Platform);
 
-			Dictionary<UnrealTargetPlatform, string> DeviceMap = new Dictionary<UnrealTargetPlatform, string>();
+			string PoolID = Globals.DevicePoolId;
 
-			foreach (string Platform in UnrealTargetPlatform.GetValidPlatformNames())
-			{
-				DeviceMap.Add(UnrealTargetPlatform.Parse(Platform), Platform);
-			}
-
-			// convert devices to request list
+			// First pass, validate all of the constraints are valid/possible
 			foreach (KeyValuePair<UnrealDeviceTargetConstraint, int> Entry in DeviceTypes)
 			{
-				if (Entry.Key.Platform == null)
-				{
-					continue;
-				}
 
-				if (!DeviceMap.ContainsKey(Entry.Key.Platform.Value))
+				if (Entry.Key.Platform == null || !DeviceMap.ContainsKey(Entry.Key.Platform.Value))
 				{
-					// if an unsupported device, we can't reserve it
-					Log.Info(KnownLogEvents.Gauntlet_DeviceEvent, "Unable to reserve service device of type: {Type}", Entry.Key);
+					// If an unsupported device, we can't reserve it
+					Log.Info(KnownLogEvents.Gauntlet_DeviceEvent, "Unable to reserve service device of type: {Type}", Entry.Key.Platform);
 					return false;
 				}
 
@@ -978,42 +820,24 @@ namespace Gauntlet
 					}
 				}
 
-				for (int i = 0; i < Entry.Value; i++)
+				if(!string.IsNullOrEmpty(Entry.Key.DeviceName) && Entry.Value > 1)
 				{
-					// @todo: if any additional reservation requirements, encode constraint into json
-					ReservationStates.AddRequired(Entry.Key.FormatWithIdentifier(), Entry.Value);
+					Log.Info(KnownLogEvents.Gauntlet_DeviceEvent, "Attempted to make a reservation for multiple devices using the same device name {DeviceName}. This is not supported.", Entry.Key.DeviceName);
+					return false;
 				}
 			}
 
-			// reserve devices
-			Uri ReservationServerUri;
-			if (Uri.TryCreate(DeviceURL, UriKind.Absolute, out ReservationServerUri))
+			// Second pass, reserve devices
+			foreach (KeyValuePair<UnrealDeviceTargetConstraint, int> Entry in DeviceTypes)
 			{
-				DeviceReservationAutoRenew DeviceReservation = null;
-
-				string PoolID = Globals.DevicePoolId;
-
-				try
-				{
-					DeviceReservation = new DeviceReservationAutoRenew(DeviceURL, 0, PoolID, ReservationStates.GetRequirements());
-				}
-				catch (Exception Ex)
-				{
-					Log.Info(KnownLogEvents.Gauntlet_DeviceEvent, "Unable to make device registration with constraint {Constraint}: {Exception}", string.Join(", ", ReservationStates.GetRequirements()), Ex);
-					return false;
-				}
+				string[] DeviceRequirements = Enumerable.Repeat(Entry.Key.FormatWithIdentifier(), Entry.Value).ToArray();
 
-				if (DeviceReservation != null)
-				{
-					// Update RequiredContraints list with what the reservation satisfied.
-					foreach (var Device in DeviceReservation.Devices)
-					{
-						ReservationStates.MarkReserved(Device);
-					}
-				}
+				Reservation NewReservation = Reservation.Create(ReservationServerUri, DeviceRequirements, TimeSpan.FromMinutes(10), PoolID: PoolID, DeviceName: Entry.Key.DeviceName);
+				DeviceReservationAutoRenew DeviceReservation = new DeviceReservationAutoRenew(DeviceURL, NewReservation);
 
-				if (DeviceReservation == null || DeviceReservation.Devices.Count != ReservationStates.RequiredCount())
+				if (DeviceReservation == null || DeviceReservation.Devices.Count != Entry.Value)
 				{
+					Log.Warning(KnownLogEvents.Gauntlet_DeviceEvent, "Unable to make device registration with constraint {Constraint}", Entry.Key);
 					return false;
 				}
 
@@ -1032,7 +856,7 @@ namespace Gauntlet
 
 				// Add target devices from reservation
 				List<ITargetDevice> ReservedDevices = new List<ITargetDevice>();
-				foreach (var Device in DeviceReservation.Devices)
+				foreach (Device Device in DeviceReservation.Devices)
 				{
 					DeviceDefinition Def = new DeviceDefinition();
 					Def.Address = Device.IPOrHostName;
@@ -1041,21 +865,19 @@ namespace Gauntlet
 					Def.DeviceData = Device.DeviceData;
 					Def.Model = Device.Model;
 
-					EPerfSpec Out = EPerfSpec.Unspecified;
-					if (!String.IsNullOrEmpty(Device.PerfSpec) && !Enum.TryParse<EPerfSpec>(Device.PerfSpec, true, out Out))
+					EPerfSpec PerfSpec = EPerfSpec.Unspecified;
+					if (!string.IsNullOrEmpty(Device.PerfSpec) && !Enum.TryParse(Device.PerfSpec, true, out PerfSpec))
 					{
 						throw new AutomationException("Unable to convert perfspec '{0}' into an EPerfSpec", Device.PerfSpec);
 					}
-					Def.PerfSpec = Out;
+					Def.PerfSpec = PerfSpec;
 
-					ITargetDevice TargetDevice = CreateAndRegisterDeviceFromDefinition(Def);
+					ITargetDevice TargetDevice = CreateAndRegisterDeviceFromDefinition(Def, Entry.Key);
 
 					// If a device from service can't be added, fail reservation and cleanup devices
-					// @todo: device problem reporting, requesting additional devices
 					if (TargetDevice == null)
 					{
 						ReportDeviceError(Device.Name, "CreateDeviceError");
-						ReservationStates.MarkProblem(Device);
 
 						// If some devices from reservation have been created, release them which will also dispose of reservation
 						if (ReservedDevices.Count > 0)
@@ -1069,7 +891,6 @@ namespace Gauntlet
 						}
 
 						Log.Info(KnownLogEvents.Gauntlet_DeviceEvent, "Unable to make device registration: device registration failed for {Platform}:{Name}", Def.Platform, Def.Name);
-						Log.Verbose(KnownLogEvents.Gauntlet_DeviceEvent, " Device Reservation status from {Pool} pool: {Status}", PoolID, string.Join(", ", ReservationStates.GetStatuses()));
 						return false;
 					}
 					else
@@ -1081,35 +902,12 @@ namespace Gauntlet
 					ServiceDeviceInfo[TargetDevice] = Def;
 					ServiceReservations[TargetDevice] = DeviceReservation;
 				}
-
-				Log.Info(KnownLogEvents.Gauntlet_DeviceEvent, "Successfully reserved service devices");
-				Log.Verbose(KnownLogEvents.Gauntlet_DeviceEvent, " Device Reservation status from {Pool} pool: {Status}", PoolID, string.Join(", ", ReservationStates.GetStatuses()));
-				return true;
 			}
 
 			Log.Info(KnownLogEvents.Gauntlet_DeviceEvent, "Unable to reserve service devices:");
-			Log.Verbose(KnownLogEvents.Gauntlet_DeviceEvent, " Device Reservation status from {Pool} pool: {Status}", Globals.DevicePoolId, string.Join(", ", ReservationStates.GetStatuses()));
 			return false;
 		}
 
-		/// <summary>
-		/// Report on Device reservation state
-		/// </summary>
-		public void ReportDeviceReservationState()
-		{
-			string PoolId = string.IsNullOrEmpty(Globals.DevicePoolId) ? "Local" : Globals.DevicePoolId;
-			string[] DeviceReservationStatus = GetReservationStatuses();
-			if (DeviceReservationStatus != null && DeviceReservationStatus.Length > 0)
-			{
-				Log.Warning(KnownLogEvents.Gauntlet_DeviceEvent, "Device Reservation status from {Pool} pool: {Status}", PoolId, string.Join(", ", DeviceReservationStatus));
-			}
-			string[] DeviceProblemStatus = GetDeviceProblems();
-			if (DeviceProblemStatus != null && DeviceProblemStatus.Length > 0)
-			{
-				Log.Warning(KnownLogEvents.Gauntlet_DeviceEvent, "Reserved Devices marked with problems from {Pool} pool: {Status}", PoolId, string.Join(", ", DeviceProblemStatus));
-			}
-		}
-
 		/// <summary>
 		/// Report target device issue to service with given error message
 		/// </summary>
@@ -1282,7 +1080,7 @@ namespace Gauntlet
 		/// </summary>
 		/// <param name="Def"></param>
 		/// <returns></returns>
-		protected ITargetDevice CreateAndRegisterDeviceFromDefinition(DeviceDefinition Def)
+		protected ITargetDevice CreateAndRegisterDeviceFromDefinition(DeviceDefinition Def, UnrealDeviceTargetConstraint Constraint)
 		{
 			ITargetDevice NewDevice = null;
 
@@ -1333,7 +1131,7 @@ namespace Gauntlet
 				{
 					if (NewDevice != null)
 					{
-						RegisterDevice(NewDevice, new UnrealDeviceTargetConstraint(NewDevice.Platform, Def.PerfSpec, Def.Model));
+						RegisterDevice(NewDevice, Constraint);
 					}
 				}
 			}
@@ -1427,9 +1225,9 @@ namespace Gauntlet
 				Device.Dispose();
 			}
 
-			foreach (DeviceReservationAutoRenew reservation in ServiceReservations.Values.Distinct())
+			foreach (DeviceReservationAutoRenew Reservation in ServiceReservations.Values.Distinct())
 			{
-				reservation.Dispose();
+				Reservation.Dispose();
 			}
 
 			ServiceReservations.Clear();
diff --git a/Engine/Source/Programs/AutomationTool/Gauntlet/Framework/Gauntlet.TestExecutor.cs b/Engine/Source/Programs/AutomationTool/Gauntlet/Framework/Gauntlet.TestExecutor.cs
--- a/Engine/Source/Programs/AutomationTool/Gauntlet/Framework/Gauntlet.TestExecutor.cs
+++ b/Engine/Source/Programs/AutomationTool/Gauntlet/Framework/Gauntlet.TestExecutor.cs
@@ -273,7 +273,6 @@ namespace Gauntlet
 									if (TimeWaiting >= Options.Wait)
 									{
 										Log.Error(KnownLogEvents.Gauntlet_DeviceEvent, "Test {TestName} has been waiting to run resource-free for {Time:00} seconds. Removing from wait list", Node, TimeWaiting);
-										DevicePool.Instance.ReportDeviceReservationState();
 										Node.AddTestEvent(new UnrealTestEvent(EventSeverity.Error, "Insufficient devices found", new List<string> {string.Format("Test {0} was unable to find enough devices after trying for {1:00} seconds.", Node, TimeWaiting), "This is not a test-related failure."}));
 										PendingTests[i] = null;
 										NodeInfo.TimeSetupBegan = NodeInfo.TimeSetupEnded = NodeInfo.TimeTestEnded = DateTime.Now;
diff --git a/Engine/Source/Programs/AutomationTool/Gauntlet/Unreal/Base/Gauntlet.UnrealSession.cs b/Engine/Source/Programs/AutomationTool/Gauntlet/Unreal/Base/Gauntlet.UnrealSession.cs
--- a/Engine/Source/Programs/AutomationTool/Gauntlet/Unreal/Base/Gauntlet.UnrealSession.cs
+++ b/Engine/Source/Programs/AutomationTool/Gauntlet/Unreal/Base/Gauntlet.UnrealSession.cs
@@ -808,7 +808,6 @@ namespace Gauntlet
 					}
 
 					// If device reservation fails, the device pool cannot support this launch.
-					DevicePool.Instance.ReportDeviceReservationState();
 					throw new AutomationException("Failed to acquire all devices for launch. See above for details.");
 				}
 
@@ -1291,7 +1290,6 @@ namespace Gauntlet
 					{
 						if (ReservationRetries == 0)
 						{
-							DevicePool.Instance.ReportDeviceReservationState();
 							throw new AutomationException("Unable to acquire all devices for test.");
 						}
 						Log.Info("\nUnable to find enough device(s). Waiting {0} secs (retries left={1})\n", ReservationRetryWait, --ReservationRetries);
@@ -1534,7 +1532,7 @@ namespace Gauntlet
 		/// </summary>
 		private bool DeviceMatchesRoleConstraint(UnrealSessionRole Role, ITargetDevice Device)
 		{
-			bool bRoleMatchesConstraint = DevicePool.Instance.GetConstraint(Device) == Role.Constraint;
+			bool bRoleMatchesConstraint = Role.Constraint.Equals(DevicePool.Instance.GetConstraint(Device));
 
 			return Device.IsConnected
 				&& Device.Platform == Role.Platform
diff --git a/Engine/Source/Programs/AutomationTool/Gauntlet/Unreal/RunUnreal.cs b/Engine/Source/Programs/AutomationTool/Gauntlet/Unreal/RunUnreal.cs
--- a/Engine/Source/Programs/AutomationTool/Gauntlet/Unreal/RunUnreal.cs
+++ b/Engine/Source/Programs/AutomationTool/Gauntlet/Unreal/RunUnreal.cs
@@ -426,7 +426,10 @@ namespace Gauntlet
 				List<string> ModelArgs = CombinedParams.ParseValues("PerfModel", false);
 				string Model = ModelArgs.Count > 0 ? ModelArgs.Last() : string.Empty;
 
-				TestContext.Constraint = new UnrealDeviceTargetConstraint(UnrealPlatform, PerfSpec, Model);
+				List<string> DeviceNameArgs = Globals.Params.ParseValues("Device", false);
+				string DeviceName = DeviceNameArgs.Count > 0 ? DeviceNameArgs.Last() : string.Empty;
+
+				TestContext.Constraint = new UnrealDeviceTargetConstraint(UnrealPlatform, PerfSpec, Model, DeviceName);
 
 				// parse worker job id
 				List<string> WorkerJobIDArgs = CombinedParams.ParseValues("WorkerJobID", false);
