Disable ZenLoader for localization automation while we resolve an actor GC issue

#nrx

diff --git a/Engine/Source/Programs/AutomationTool/Scripts/Localisation.Automation.cs b/Engine/Source/Programs/AutomationTool/Scripts/Localisation.Automation.cs
--- a/Engine/Source/Programs/AutomationTool/Scripts/Localisation.Automation.cs
+++ b/Engine/Source/Programs/AutomationTool/Scripts/Localisation.Automation.cs
@@ -126,7 +126,7 @@ class Localize : BuildCommand
 					if (LocalizationConfigFiles.Count > 0)
 					{
 						string ConcatenatedConfigFiles = String.Join(";", LocalizationConfigFiles);
-						string CommandLine = $"\"{LauncherArgs.GetAbsoluteUEProjectPath()}\" -run=GatherText -config=\"{ConcatenatedConfigFiles}\" {LauncherArgs.EditorArgs}";
+						string CommandLine = $"\"{LauncherArgs.GetAbsoluteUEProjectPath()}\" -run=GatherText -config=\"{ConcatenatedConfigFiles}\" -nozenloader {LauncherArgs.EditorArgs}";
 						Logger.LogInformation("Running localization commandlet for '{Arg0}': {Arguments}", ProjectInfo.ProjectName, CommandLine);
 						LocalizationTask.GatherProcessResults.Add(Run(LauncherArgs.AbsoluteEditorExePath, CommandLine, null,	LauncherArgs.CommandletRunOptions));
 					}
