Fix GatherTextFromAssetsCommandlet to properly handle external actors when keeping alive selected packages.
This happens with Zen loader because of a major difference with the old loader with partial packages; if a package is loaded, then some of its objects are garbage collected but not all, and we reload the same package again, the Zen loader precreates the missing exports in a RF_NeedLoad state, while the old loader would simply just return the initial package with missing exports.
Also backout CL34748964 as this is not required anymore.

diff --git a/Engine/Source/Editor/UnrealEd/Private/Commandlets/GatherTextFromAssetsCommandlet.cpp b/Engine/Source/Editor/UnrealEd/Private/Commandlets/GatherTextFromAssetsCommandlet.cpp
--- a/Engine/Source/Editor/UnrealEd/Private/Commandlets/GatherTextFromAssetsCommandlet.cpp
+++ b/Engine/Source/Editor/UnrealEd/Private/Commandlets/GatherTextFromAssetsCommandlet.cpp
@@ -397,12 +397,13 @@ void UGatherTextFromAssetsCommandlet::PurgeGarbage(const bool bPurgeReferencedPa
 
 				// Keep any requested packages (and their RF_Standalone inners) alive during a call to PurgeGarbage
 				ObjectsToKeepAlive.Add(Package);
-				ForEachObjectWithOuter(Package, [this](UObject* InPackageInner)
+				ForEachObjectWithPackage(Package, [this](UObject* InPackageInner)
 				{
-					if (InPackageInner->HasAnyFlags(RF_Standalone))
+					if (InPackageInner->HasAnyFlags(RF_Standalone | RF_HasExternalPackage))
 					{
 						ObjectsToKeepAlive.Add(InPackageInner);
 					}
+					return true;
 				}, true, RF_NoFlags, EInternalObjectFlags::Garbage);
 			}
 		}
diff --git a/Engine/Source/Programs/AutomationTool/Scripts/Localisation.Automation.cs b/Engine/Source/Programs/AutomationTool/Scripts/Localisation.Automation.cs
--- a/Engine/Source/Programs/AutomationTool/Scripts/Localisation.Automation.cs
+++ b/Engine/Source/Programs/AutomationTool/Scripts/Localisation.Automation.cs
@@ -126,7 +126,7 @@ class Localize : BuildCommand
 					if (LocalizationConfigFiles.Count > 0)
 					{
 						string ConcatenatedConfigFiles = String.Join(";", LocalizationConfigFiles);
-						string CommandLine = $"\"{LauncherArgs.GetAbsoluteUEProjectPath()}\" -run=GatherText -config=\"{ConcatenatedConfigFiles}\" -nozenloader {LauncherArgs.EditorArgs}";
+						string CommandLine = $"\"{LauncherArgs.GetAbsoluteUEProjectPath()}\" -run=GatherText -config=\"{ConcatenatedConfigFiles}\" {LauncherArgs.EditorArgs}";
 						Logger.LogInformation("Running localization commandlet for '{Arg0}': {Arguments}", ProjectInfo.ProjectName, CommandLine);
 						LocalizationTask.GatherProcessResults.Add(Run(LauncherArgs.AbsoluteEditorExePath, CommandLine, null,	LauncherArgs.CommandletRunOptions));
 					}
