LiveLinkHub: Move Automation script/project to standard location, clean up csproj.

diff --git a/Engine/Source/Programs/AutomationTool/LiveLinkHub/LiveLinkHub.Automation.cs b/Engine/Source/Programs/AutomationTool/LiveLinkHub/LiveLinkHub.Automation.cs
new file mode 100644
--- /dev/null
+++ b/Engine/Source/Programs/AutomationTool/LiveLinkHub/LiveLinkHub.Automation.cs
@@ -0,0 +1,64 @@
+// Copyright Epic Games, Inc. All Rights Reserved.
+
+using EpicGames.Core;
+using System.Linq;
+
+
+namespace AutomationTool
+{
+	class MakeLiveLinkHubEditor : MakeCookedEditor
+	{
+		protected override ProjectParams MakeParams(string DLCName, string BasedOnReleaseVersion)
+		{
+			return new ProjectParams(
+				Command: this
+				, RawProjectPath: ProjectFile
+
+				, NoBootstrapExe: true
+				, DLCName: DLCName
+				, BasedOnReleaseVersion: BasedOnReleaseVersion
+				, DedicatedServer: bIsCookedCooker
+				, NoClient: bIsCookedCooker
+				, OptionalContent: true
+				, ClientCookedTargets: new ParamList<string>() // Prevent AutodetectSettings from looking for a game target
+				, EditorTargets: new ParamList<string>("LiveLinkHubEditor")
+			);
+		}
+
+		protected override void ModifyParams(ProjectParams BuildParams)
+		{
+			base.ModifyParams(BuildParams);
+
+			// We don't want the SDK dir / CookerSupportFiles
+			BuildParams.CookerSupportFilesSubdirectory = null;
+		}
+
+		protected override void ModifyDeploymentContext(ProjectParams Params, DeploymentContext SC)
+		{
+			ModifyStageContext Context = CreateContext(Params, SC);
+
+			DefaultModifyDeploymentContext(Params, SC, Context);
+
+			Context.Apply(SC);
+
+			// Copy .target receipt from project bin to engine bin
+			string PlatName = SC.StageTargetPlatform.PlatformType.ToString();
+			SC.FilesToStage.NonUFSFiles.Add(
+				new StagedFileReference($"Engine/Binaries/{PlatName}/{Context.ProjectName}.target"),
+				new FileReference($"Engine/Binaries/{PlatName}/{Context.ProjectName}.target"));
+
+			// Copy program DefaultEngine.ini to Engine/Config
+			StagedFileReference DefaultEngineDest = new StagedFileReference("Engine/Config/DefaultEngine.ini");
+			SC.FilesToStage.UFSFiles[DefaultEngineDest] = new FileReference($"Engine/Source/Programs/{Context.ProjectName}/Config/DefaultEngine.ini");
+
+			// Remove asset registry entry
+			SC.FilesToStage.UFSFiles.Remove(new StagedFileReference($"{Context.ProjectName}/EditorClientAssetRegistry.bin"));
+
+			// Move PDBs to debug (FIXME: should this be necessary? otherwise -nodebuginfo doesn't exclude them)
+			SC.FilesToStage.NonUFSDebugFiles.Union(SC.FilesToStage.NonUFSFiles.Where(x => x.Key.HasExtension(".pdb")));
+
+			// Remove any files from NonUFS that were also added to Debug
+			SC.FilesToStage.NonUFSFiles = SC.FilesToStage.NonUFSFiles.Where(x => !SC.FilesToStage.NonUFSDebugFiles.ContainsKey(x.Key)).ToDictionary(x => x.Key, x => x.Value);
+		}
+	}
+}
diff --git a/Engine/Source/Programs/AutomationTool/LiveLinkHub/LiveLinkHub.Automation.csproj b/Engine/Source/Programs/AutomationTool/LiveLinkHub/LiveLinkHub.Automation.csproj
new file mode 100644
--- /dev/null
+++ b/Engine/Source/Programs/AutomationTool/LiveLinkHub/LiveLinkHub.Automation.csproj
@@ -0,0 +1,25 @@
+<Project Sdk="Microsoft.NET.Sdk">
+  <Import Project="..\..\Shared\UnrealEngine.csproj.props" />
+
+  <PropertyGroup>
+    <TargetFramework>net6.0</TargetFramework>
+    <OutputType>Library</OutputType>
+    <AppendTargetFrameworkToOutputPath>false</AppendTargetFrameworkToOutputPath>
+    <GenerateAssemblyInfo>false</GenerateAssemblyInfo>
+    <GenerateTargetFrameworkAttribute>false</GenerateTargetFrameworkAttribute>
+    <Configurations>Debug;Release;Development</Configurations>
+    <OutputPath>..\..\..\..\Binaries\DotNET\AutomationTool\AutomationScripts\LiveLinkHub</OutputPath>
+  </PropertyGroup>
+
+  <ItemGroup>
+    <Compile Include="..\..\Shared\MetaData.cs">
+      <Link>Properties\MetaData.cs</Link>
+    </Compile>
+  </ItemGroup>
+  <ItemGroup>
+    <ProjectReference Include="..\AutomationUtils\AutomationUtils.Automation.csproj" />
+    <ProjectReference Include="..\CookedEditor\CookedEditor.Automation.csproj" />
+    <ProjectReference Include="..\..\Shared\EpicGames.Core\EpicGames.Core.csproj" />
+    <ProjectReference Include="..\..\UnrealBuildTool\UnrealBuildTool.csproj" />
+  </ItemGroup>
+</Project>
\ No newline at end of file
diff --git a/Engine/Source/Programs/LiveLinkHub/LiveLinkHub.Automation.cs b/Engine/Source/Programs/LiveLinkHub/LiveLinkHub.Automation.cs
deleted file mode 100644
diff --git a/Engine/Source/Programs/LiveLinkHub/LiveLinkHub.Automation.csproj b/Engine/Source/Programs/LiveLinkHub/LiveLinkHub.Automation.csproj
deleted file mode 100644
