PR #12069: Fixed issue when gauntlet test is run on editor project residing in a folder containing spaces

#rnx

diff --git a/Engine/Source/Programs/AutomationTool/Gauntlet/Unreal/Base/Gauntlet.UnrealTestConfiguration.cs b/Engine/Source/Programs/AutomationTool/Gauntlet/Unreal/Base/Gauntlet.UnrealTestConfiguration.cs
--- a/Engine/Source/Programs/AutomationTool/Gauntlet/Unreal/Base/Gauntlet.UnrealTestConfiguration.cs
+++ b/Engine/Source/Programs/AutomationTool/Gauntlet/Unreal/Base/Gauntlet.UnrealTestConfiguration.cs
@@ -288,7 +288,7 @@ namespace Gauntlet
 		/// <returns></returns>
 		public string GenerateFullCommandLine()
 		{
-			string FinalCommandline = string.Format("{0} {1} ", Project, GameMap);
+			string FinalCommandline = string.Format("\"{0}\" \"{1}\" ", Project, GameMap);
 			foreach (string Key in Params.Keys)
 			{
 				string CurrentArgument;
