[Gauntlet] Fix an issue when original error contains curling brace but it is passed as format string.
Curling braces need to be escaped to avoid being taken as Format token.

#rnx

diff --git a/Engine/Source/Programs/AutomationTool/Gauntlet/Unreal/Utils/Gauntlet.UnrealTestResults.cs b/Engine/Source/Programs/AutomationTool/Gauntlet/Unreal/Utils/Gauntlet.UnrealTestResults.cs
--- a/Engine/Source/Programs/AutomationTool/Gauntlet/Unreal/Utils/Gauntlet.UnrealTestResults.cs
+++ b/Engine/Source/Programs/AutomationTool/Gauntlet/Unreal/Utils/Gauntlet.UnrealTestResults.cs
@@ -217,7 +217,7 @@ namespace Gauntlet
 					{
 						Group MatchedGroup = MatchSummary.Groups[0];
 						int Index = MatchedGroup.Index;
-						Format = $"{Message.Substring(0, Index)}{{Summary}}\n{{Callstack}}";
+						Format = $"{MessageTemplate.Escape(Message.Substring(0, Index))}{{Summary}}\n{{Callstack}}";
 						Properties.Add("Summary", MatchedGroup.Value);
 						Properties.Add("Callstack", Message.Substring(Index + MatchedGroup.Length).TrimStart('\n'));
 					}
@@ -233,7 +233,7 @@ namespace Gauntlet
 				if (!string.IsNullOrEmpty(Event.Context))
 				{
 					Properties = new Dictionary<string, object>() { { "Context", Event.Context } };
-					Format = "[{Context}] " + (Format ?? Message);
+					Format = "[{Context}] " + (Format ?? MessageTemplate.Escape(Message));
 				}
 				if (!string.IsNullOrEmpty(Filename))
 				{
@@ -243,7 +243,7 @@ namespace Gauntlet
 					}
 					Properties.Add("SourceFile", Filename);
 					Properties.Add("Line", LineNumber.ToString());
-					Format = (Format ?? Message) + " [{SourceFile}({Line})]";
+					Format = (Format ?? MessageTemplate.Escape(Message)) + " [{SourceFile}({Line})]";
 				}
 			}
 
