Add PostStagingFileCopy hook for CustomDeploymentHandler

#rnx

diff --git a/Engine/Source/Programs/AutomationTool/AutomationUtils/CustomDeploymentHandler.cs b/Engine/Source/Programs/AutomationTool/AutomationUtils/CustomDeploymentHandler.cs
--- a/Engine/Source/Programs/AutomationTool/AutomationUtils/CustomDeploymentHandler.cs
+++ b/Engine/Source/Programs/AutomationTool/AutomationUtils/CustomDeploymentHandler.cs
@@ -26,6 +26,7 @@ public abstract class CustomDeploymentHandler
 	public virtual bool PreGetFilesToDeployOrStage(ProjectParams Params, DeploymentContext SC) => false;
 	public virtual bool GetFilesToStageForDLC(ProjectParams Params, DeploymentContext SC) => false;
 	public virtual bool GetFilesToArchive(ProjectParams Params, DeploymentContext SC) => false;
+	public virtual bool PostStagingFileCopy(ProjectParams Params, DeploymentContext SC) => false;
 
 	// always called
 	public virtual void PostPackage(ProjectParams Params, DeploymentContext SC, int WorkingCL) { }
diff --git a/Engine/Source/Programs/AutomationTool/Scripts/CopyBuildToStagingDirectory.Automation.cs b/Engine/Source/Programs/AutomationTool/Scripts/CopyBuildToStagingDirectory.Automation.cs
--- a/Engine/Source/Programs/AutomationTool/Scripts/CopyBuildToStagingDirectory.Automation.cs
+++ b/Engine/Source/Programs/AutomationTool/Scripts/CopyBuildToStagingDirectory.Automation.cs
@@ -5106,8 +5106,11 @@ namespace AutomationScripts
 
 			CopyUsingStagingManifest(Params, SC);
 
-			var ThisPlatform = SC.StageTargetPlatform;
-			ThisPlatform.PostStagingFileCopy(Params, SC);
+			if (SC.CustomDeployment == null || !SC.CustomDeployment.PostStagingFileCopy(Params, SC))
+			{
+				var ThisPlatform = SC.StageTargetPlatform;
+				ThisPlatform.PostStagingFileCopy(Params, SC);
+			}
 		}
 
 		private static DirectoryReference GetIntermediateCommandlineDir(DeploymentContext SC)
