CanCookedFileBeStaged changed to GetStagingTypeForCookedFile in the UAT staging scripts. Allows the staging code to also specify the staging type for a given file rather than just whether it gets staged or not
Adjusted metallib staging to mark those files staged as Non-UFS

diff --git a/Engine/Source/Programs/AutomationTool/Scripts/CopyBuildToStagingDirectory.Automation.cs b/Engine/Source/Programs/AutomationTool/Scripts/CopyBuildToStagingDirectory.Automation.cs
--- a/Engine/Source/Programs/AutomationTool/Scripts/CopyBuildToStagingDirectory.Automation.cs
+++ b/Engine/Source/Programs/AutomationTool/Scripts/CopyBuildToStagingDirectory.Automation.cs
@@ -988,27 +988,27 @@ namespace AutomationScripts
 			throw new AutomationException(String.Format("A package store manifest or project store is required when staging to IoStore. Expected to find {0} or {1}. Ensure that cooking was successful.", PackageStoreManifestFile.FullName, ProjectStoreFile.FullName));
 		}
 
-		private static bool CanCookedFileBeStaged(FileReference File)
+		private static StagedFileType? GetStagingTypeForCookedFile(FileReference File)
 		{
 			// json files have never been staged
 			if (File.HasExtension(".json"))
 			{
-				return false;
+				return null;
 			}
 
 			// metallib files cannot *currently* be staged as UFS as the Metal API needs to mmap them from files on disk in order to function efficiently
 			if (File.HasExtension(".metallib"))
 			{
-				return false;
+				return StagedFileType.NonUFS;
 			}
 
 			// Cannot stage
 			if (File.HasExtension(".utoc") || File.HasExtension(".ucas") || File.HasExtension(".uondemandtoc"))
 			{
-				return false;
+				return null;
 			}
 
-			return true;
+			return StagedFileType.UFS;
 		}
 
 		public static bool SetUpStagingSourceDirectories(ProjectParams Params, DeploymentContext SC)
@@ -1221,10 +1221,11 @@ namespace AutomationScripts
 					{
 						continue;
 					}
-					
-					if (CanCookedFileBeStaged(CookedFile))
+
+					StagedFileType? CookedFileStageType = GetStagingTypeForCookedFile(CookedFile);
+					if (CookedFileStageType.HasValue)
 					{
-						SC.StageFile(StagedFileType.UFS, CookedFile, new StagedFileReference(CookedFile.MakeRelativeTo(SC.PlatformCookDir)));
+						SC.StageFile(CookedFileStageType.Value, CookedFile, new StagedFileReference(CookedFile.MakeRelativeTo(SC.PlatformCookDir)));
 					}
 				}
 
@@ -1576,9 +1577,10 @@ namespace AutomationScripts
 								continue;
 							}
 
-							if (CanCookedFileBeStaged(CookedFile))
+							StagedFileType? CookedFileStageType = GetStagingTypeForCookedFile(CookedFile);
+							if (CookedFileStageType.HasValue)
 							{
-								SC.StageFile(StagedFileType.UFS, CookedFile, new StagedFileReference(CookedFile.MakeRelativeTo(SC.PlatformCookDir)));
+								SC.StageFile(CookedFileStageType.Value, CookedFile, new StagedFileReference(CookedFile.MakeRelativeTo(SC.PlatformCookDir)));
 							}
 						}
 					}
