[Gauntlet] Refactor parsing error with callstack to avoid copying large chunk of log; instead limit to line.
Add a limit to the number of errors with callstack that can be searched for (arbitrary limit is set to 5 for fatal errors and 10 for ensures).

#rnx

diff --git a/Engine/Source/Programs/AutomationTool/Gauntlet/Unreal/Utils/Gauntlet.UnrealLogParser.cs b/Engine/Source/Programs/AutomationTool/Gauntlet/Unreal/Utils/Gauntlet.UnrealLogParser.cs
--- a/Engine/Source/Programs/AutomationTool/Gauntlet/Unreal/Utils/Gauntlet.UnrealLogParser.cs
+++ b/Engine/Source/Programs/AutomationTool/Gauntlet/Unreal/Utils/Gauntlet.UnrealLogParser.cs
@@ -103,7 +103,6 @@ namespace Gauntlet
 			public string[] Callstack;
 			public bool IsEnsure;
 			public bool IsSanReport;
-			public bool IsPostMortem;
 
 			/// <summary>
 			/// Generate a string that represents a CallstackMessage formatted to be inserted into a log file.
@@ -943,7 +942,7 @@ namespace Gauntlet
 		/// <returns></returns>
 		public IEnumerable<UnrealLog.CallstackMessage> GetEnsures()
 		{
-			IEnumerable<UnrealLog.CallstackMessage> Ensures = ParseTracedErrors(new[] { @"Log.+:\s{0,1}Error:\s{0,1}(Ensure condition failed:.+)" }, false);
+			IEnumerable<UnrealLog.CallstackMessage> Ensures = ParseTracedErrors(new[] { @"Log.+:\s{0,1}Error:\s{0,1}(Ensure condition failed:.+)" }, 10);
 
 			foreach (UnrealLog.CallstackMessage Error in Ensures)
 			{
@@ -959,9 +958,22 @@ namespace Gauntlet
 		/// <returns></returns>
 		public UnrealLog.CallstackMessage GetFatalError()
 		{
-			string[] ErrorMsgMatches = new string[] { @"(Fatal Error:.+)", @"Critical error: =+\s+(?:[\S\s]+?\s*Error: +)?(.+)", @"(Assertion Failed:.+)", @"(Unhandled Exception:.+)", @"(LowLevelFatalError.+)" };
+			var Traces = GetASanErrors();
+			if (Traces.Any())
+			{
+				return Traces.Last();
+			}
 
-			var Traces = ParseTracedErrors(ErrorMsgMatches).Concat(GetASanErrors());
+			string[] ErrorMsgMatches = new string[] { @"(Fatal Error:.+)", @"Critical error: =+\s+(?:[\S\s]+?\s*Error: +)?(.+)", @"(Assertion Failed:.+)", @"(Unhandled Exception:.+)", @"(LowLevelFatalError.+)", @"(Postmortem Cause:.*)" };
+
+			Traces = ParseTracedErrors(ErrorMsgMatches, 5);
+
+			// If we have a post-mortem error, return that one (on some devices the post-mortem info is way more informative).
+			var PostMortemTraces = Traces.Where(T => T.Message.IndexOf("Postmortem Cause:", StringComparison.OrdinalIgnoreCase) > -1);
+			if (PostMortemTraces.Any())
+			{
+				Traces = PostMortemTraces;
+			}
 
 			// Keep the one with the most information.
 			return Traces.Count() > 0 ? Traces.OrderBy(T => T.Callstack.Length).Last() : null;
@@ -993,42 +1005,44 @@ namespace Gauntlet
 				}
 
 				Regex EndPattern = SanitizerEventMatcher.ReportEndPattern;
-				int TraceInitIndex = TraceInitMatch.Index + TraceInitMatch.Length;
 
 				UnrealLog.CallstackMessage NewTrace = new UnrealLog.CallstackMessage();
 				NewTrace.IsSanReport = true;
 				NewTrace.Position = TraceInitMatch.Index;
 				NewTrace.Message = $"{TraceInitMatch.Groups["SanitizerName"].Value}Sanitizer: {TraceInitMatch.Groups["Summary"].Value}";
 
-				// If the regex matches the very end of the string, the substring will get an invalid range.
-				string ErrorContent = Content.Length <= TraceInitIndex
-					? string.Empty : Content.Substring(TraceInitIndex + 1);
-
-				if (!string.IsNullOrEmpty(ErrorContent))
+				List<string> Backtrace = new List<string>();
+				int Cursor = TraceInitMatch.Index + TraceInitMatch.Length;
+				int EOL = Content.IndexOf('\n', Cursor);
+				while (EOL != -1)
 				{
-					Match MsgMatch = EndPattern.Match(ErrorContent);
+					string Line = Content.Substring(Cursor, EOL - Cursor);
+					Match MsgMatch = EndPattern.Match(Line);
 					if (MsgMatch.Success)
 					{
-						string MsgContent = ErrorContent.Substring(0, MsgMatch.Index);
+						break;
+					}
+					else
+					{
 						// Prune the line with time stamp
-						List<string> Backtrace = new List<string>();
-						foreach (string Line in MsgContent.Split("\n"))
+						Match IsTimeStampLine = LineStartsWithTimeStamp.Match(Line);
+						if (!IsTimeStampLine.Success)
 						{
-							Match IsTimeStampLine = LineStartsWithTimeStamp.Match(Line);
-							if (!IsTimeStampLine.Success)
-							{
-								Backtrace.Add(Line);
-							}
+							Backtrace.Add(Line);
 						}
-
-						NewTrace.Callstack = Backtrace.ToArray();
 					}
+					Cursor = EOL + 1;
+					EOL = Content.IndexOf('\n', Cursor);
 				}
 
-				if(NewTrace.Callstack == null)
+				if(Backtrace.Count == 0)
 				{
 					NewTrace.Callstack = new string[] { "Unable to parse callstack from log" };
 				}
+				else
+				{
+					NewTrace.Callstack = Backtrace.ToArray();
+				}
 
 				ASanReports.Add(NewTrace);
 			}
@@ -1126,19 +1140,18 @@ namespace Gauntlet
 		/// Finds all callstack-based errors with the specified pattern
 		/// </summary>
 		/// <param name="Patterns"></param>
-		/// <param name="IncludePostmortem"></param>
+		/// <param name="Limit">Limit the number of errors to parse with trace per pattern. Zero means no limit.</param>
 		/// <returns></returns>
-		protected IEnumerable<UnrealLog.CallstackMessage> ParseTracedErrors(string[] Patterns, bool IncludePostmortem = true)
+		protected IEnumerable<UnrealLog.CallstackMessage> ParseTracedErrors(string[] Patterns, int Limit = 0)
 		{
 			List<UnrealLog.CallstackMessage> Traces = new List<UnrealLog.CallstackMessage>();
 
-			// As well as what was requested, search for a postmortem stack...
-			IEnumerable<string> AllPatterns = IncludePostmortem ? Patterns.Concat(new[] { "(Postmortem Cause:.*)" }) : Patterns;
-
 			// Try and find an error message
-			foreach (string Pattern in AllPatterns)
+			foreach (string Pattern in Patterns)
 			{
 				MatchCollection Matches = Regex.Matches(Content, Pattern, RegexOptions.IgnoreCase);
+				// Maximum number of errors with callstack to consider
+				int RemainingSlots = Limit;
 
 				foreach (Match TraceMatch in Matches)
 				{
@@ -1147,22 +1160,7 @@ namespace Gauntlet
 					NewTrace.Position = TraceMatch.Index;
 					NewTrace.Message = TraceMatch.Groups[1].Value;
 
-					// If the regex matches the very end of the string, the substring will get an invalid range.
-					string ErrorContent = Content.Length <= TraceMatch.Index + TraceMatch.Length
-						? string.Empty : Content.Substring(TraceMatch.Index + TraceMatch.Length + 1);
-
-					Match MsgMatch = Regex.Match(ErrorContent, @".+:\s*Error:\s*(.+)");
-
-					if (MsgMatch.Success)
-					{
-						string MsgString = MsgMatch.Groups[1].ToString();
-						if (!MsgMatch.Groups[0].ToString().Contains("\n") /* avoid a bug where .+ match \n */
-							&& !string.IsNullOrEmpty(MsgString)
-							&& Regex.Match(MsgString, @"0[xX][0-9A-f]{8,16}").Success == false)
-						{
-							NewTrace.Message = NewTrace.Message + "\n" + MsgString;
-						}
-					}
+					int Cursor = TraceMatch.Index + TraceMatch.Length;
 
 					//
 					// Handing callstacks-
@@ -1177,7 +1175,7 @@ namespace Gauntlet
 					//
 					// E.g 0x00000000 UnknownFunction []
 					//
-					// A calstack as part of an ensure, check, or exception will look something like this -
+					// A callstack as part of an ensure, check, or exception will look something like this -
 					// 
 					//
 					//[2017.08.21-03.28.40:667][313]LogWindows:Error: Assertion failed: false [File:D:\Epic\Orion\Release-Next\Engine\Plugins\NotForLicensees\Gauntlet\Source\Gauntlet\Private\GauntletTestControllerErrorTest.cpp] [Line: 29] 
@@ -1206,25 +1204,20 @@ namespace Gauntlet
 					// channel
 					// 
 
-					string SearchContent = ErrorContent;
-
 					int LinesWithoutBacktrace = 0;
 
 					List<string> Backtrace = new List<string>();
 
 					do
 					{
-						int EOL = SearchContent.IndexOf("\n");
+						int EOL = Content.IndexOf("\n", Cursor);
 
 						if (EOL == -1)
 						{
 							break;
 						}
 
-						string Line = SearchContent.Substring(0, EOL);
-
-						// collapse inline function
-						Line = Line.Replace("[Inline Function] ", "[InlineFunction]");
+						string Line = Content.Substring(Cursor, EOL - Cursor);
 
 						// Must have [Callstack] 0x00123456
 						// The module name is optional, must start with whitespace, and continues until next whites[ace
@@ -1261,10 +1254,24 @@ namespace Gauntlet
 						}
 						else
 						{
+							if (Backtrace.Count == 0)
+							{
+								Match MsgMatch = Regex.Match(Line, @".+:\s*Error:\s*(.+)");
+
+								if (MsgMatch.Success)
+								{
+									string MsgString = MsgMatch.Groups[1].ToString().Trim();
+									if (!string.IsNullOrEmpty(MsgString))
+									{
+										NewTrace.Message += "\n" + MsgString;
+									}
+								}
+							}
+
 							LinesWithoutBacktrace++;
 						}
 
-						SearchContent = SearchContent.Substring(EOL + 1);
+						Cursor = EOL + 1;
 
 					} while (LinesWithoutBacktrace < 10);
 
@@ -1276,55 +1283,27 @@ namespace Gauntlet
 					{
 						NewTrace.Callstack = new[] { "Unable to parse callstack from log" };
 					}
-					Traces.Add(NewTrace);
-				}
-			}
-
-			// Now, because platforms sometimes dump asserts to the log and low-level logging, and we might have a post-mortem stack, we
-			// need to prune out redundancies. Basic approach - find errors with the same assert message and keep the one with the
-			// longest callstack. If we have a post-mortem error, overwrite the previous trace with its info (on some devices the post-mortem
-			// info is way more informative).
-
-			List<UnrealLog.CallstackMessage> FilteredTraces = new List<UnrealLog.CallstackMessage>();
-
 
-			for (int i = 0; i < Traces.Count; i++)
-			{
-				var Trace = Traces[i];
-
-				// check the next trace to see if it's a dupe of us
-				if (i + 1 < Traces.Count)
-				{
-					var NextTrace = Traces[i + 1];
-
-					if (Trace.Message.Equals(NextTrace.Message, StringComparison.OrdinalIgnoreCase))
+					// Because platforms sometimes dump asserts to the log and low-level logging, we need to prune out redundancies.
+					// Basic approach: find errors with the same assert message and keep the one with the longest callstack.
+					if (Traces.Count > 0 && NewTrace.Message.Equals(Traces.Last().Message, StringComparison.OrdinalIgnoreCase))
 					{
-						if (Trace.Callstack.Length < NextTrace.Callstack.Length)
+						var PreviousTrace = Traces.Last();
+						if (PreviousTrace.Callstack.Length < NewTrace.Callstack.Length)
 						{
-							Trace.Callstack = NextTrace.Callstack;
-							// skip the next error as we stole its callstack already
-							i++;
+							PreviousTrace.Callstack = NewTrace.Callstack;
 						}
 					}
-				}
+					else
+					{
+						Traces.Add(NewTrace);
+					}
 
-				// check this trace to see if it's postmortem
-				if (Trace.Message.IndexOf("Postmortem Cause:", StringComparison.OrdinalIgnoreCase) != -1)
-				{
-					// we have post-mortem info, which should be much better than the game-generated stuff and will be sorted to first position
-					Trace.IsPostMortem = true;
+					if (Limit > 0 && --RemainingSlots <= 0) break;
 				}
-
-				FilteredTraces.Add(Trace);
-			}
-
-			// If we have a post mortem crash, sort it to the front
-			if (FilteredTraces.FirstOrDefault((Trace) => { return Trace.IsPostMortem; }) != null)
-			{
-				FilteredTraces.Sort((Trace1, Trace2) => { if (!Trace1.IsPostMortem && !Trace2.IsPostMortem) return 0; return Trace1.IsPostMortem ? -1 : 1; });
 			}
 
-			return FilteredTraces;
+			return Traces;
 		}
 
 		/// <summary>
