[Gauntlet] Avoid failure to generate json artifact if log path property is empty

#rnx

diff --git a/Engine/Source/Programs/AutomationTool/Gauntlet/Unreal/Automation/UE.Automation.cs b/Engine/Source/Programs/AutomationTool/Gauntlet/Unreal/Automation/UE.Automation.cs
--- a/Engine/Source/Programs/AutomationTool/Gauntlet/Unreal/Automation/UE.Automation.cs
+++ b/Engine/Source/Programs/AutomationTool/Gauntlet/Unreal/Automation/UE.Automation.cs
@@ -828,6 +828,7 @@ namespace UE
 							TempReport.SetOutputArtifactPath(HordeArtifactPath);
 							foreach (UnrealRoleArtifacts Artifact in SessionArtifacts)
 							{
+								if (string.IsNullOrEmpty(Artifact.LogPath)) continue;
 								string LogName = Path.GetRelativePath(Path.GetFullPath(Context.Options.LogDir), Path.GetFullPath(Artifact.LogPath));
 								if (TempReport.AttachArtifact(Artifact.LogPath, LogName) && Artifact.SessionRole.RoleType == MainRole.Type)
 								{
diff --git a/Engine/Source/Programs/AutomationTool/Gauntlet/Unreal/Base/Gauntlet.UnrealTestNode.cs b/Engine/Source/Programs/AutomationTool/Gauntlet/Unreal/Base/Gauntlet.UnrealTestNode.cs
--- a/Engine/Source/Programs/AutomationTool/Gauntlet/Unreal/Base/Gauntlet.UnrealTestNode.cs
+++ b/Engine/Source/Programs/AutomationTool/Gauntlet/Unreal/Base/Gauntlet.UnrealTestNode.cs
@@ -1514,7 +1514,7 @@ namespace Gauntlet
 				{
 					foreach (UnrealRoleArtifacts Artifact in SessionArtifacts)
 					{
-						if (Artifact.SessionRole.RoleType == MainRole.Type)
+						if (Artifact.SessionRole.RoleType == MainRole.Type && !string.IsNullOrEmpty(Artifact.LogPath))
 						{
 							string LogName = FileUtils.ConvertPathToUri(Path.GetRelativePath(Path.GetFullPath(ArtifactPath), Path.GetFullPath(Artifact.LogPath)));
 							if (HordeTestPassResults is BaseHordeReport Report)
