LiveLinkHub: Fixes for staged editor target related to running in a project context.

We now also avoid overriding OptedInModulePlatforms unless we're actively staging, since it triggers spurious rebuilds when interleaved with UnrealEditor builds.

diff --git a/Engine/Source/Programs/AutomationTool/LiveLinkHub/LiveLinkHub.Automation.cs b/Engine/Source/Programs/AutomationTool/LiveLinkHub/LiveLinkHub.Automation.cs
--- a/Engine/Source/Programs/AutomationTool/LiveLinkHub/LiveLinkHub.Automation.cs
+++ b/Engine/Source/Programs/AutomationTool/LiveLinkHub/LiveLinkHub.Automation.cs
@@ -22,6 +22,7 @@ namespace AutomationTool
 				, OptionalContent: true
 				, ClientCookedTargets: new ParamList<string>() // Prevent AutodetectSettings from looking for a game target
 				, EditorTargets: new ParamList<string>("LiveLinkHubEditor")
+				, UbtArgs: "-SingleModulePlatform"
 			);
 		}
 
@@ -41,15 +42,25 @@ namespace AutomationTool
 
 			Context.Apply(SC);
 
-			// Copy .target receipt from project bin to engine bin
 			string PlatName = SC.StageTargetPlatform.PlatformType.ToString();
+
+			// Copy .target receipt to project and engine bin
+			SC.FilesToStage.NonUFSFiles.Add(
+				new StagedFileReference($"{Context.ProjectName}/Binaries/{PlatName}/{Context.ProjectName}.target"),
+				new FileReference($"Engine/Binaries/{PlatName}/{Context.ProjectName}.target"));
+
 			SC.FilesToStage.NonUFSFiles.Add(
 				new StagedFileReference($"Engine/Binaries/{PlatName}/{Context.ProjectName}.target"),
 				new FileReference($"Engine/Binaries/{PlatName}/{Context.ProjectName}.target"));
 
-			// Copy program DefaultEngine.ini to Engine/Config
-			StagedFileReference DefaultEngineDest = new StagedFileReference("Engine/Config/DefaultEngine.ini");
-			SC.FilesToStage.UFSFiles[DefaultEngineDest] = new FileReference($"Engine/Source/Programs/{Context.ProjectName}/Config/DefaultEngine.ini");
+			// Stage TargetInfo and target script. Necessary to avoid "Running incorrect executable for target (...)"
+			SC.FilesToStage.NonUFSFiles.Add(
+				new StagedFileReference($"{Context.ProjectName}/Source/{Context.ProjectName}.Target.cs"),
+				new FileReference($"Engine/Source/Programs/{Context.ProjectName}/Source/{Context.ProjectName}.Target.cs"));
+
+			SC.FilesToStage.NonUFSFiles.Add(
+				new StagedFileReference($"{Context.ProjectName}/Intermediate/TargetInfo.json"),
+				new FileReference($"Engine/Source/Programs/{Context.ProjectName}/Intermediate/TargetInfo.json"));
 
 			// Remove asset registry entry
 			SC.FilesToStage.UFSFiles.Remove(new StagedFileReference($"{Context.ProjectName}/EditorClientAssetRegistry.bin"));
diff --git a/Engine/Source/Programs/LiveLinkHub/Private/LiveLinkHubRun.cpp b/Engine/Source/Programs/LiveLinkHub/Private/LiveLinkHubRun.cpp
--- a/Engine/Source/Programs/LiveLinkHub/Private/LiveLinkHubRun.cpp
+++ b/Engine/Source/Programs/LiveLinkHub/Private/LiveLinkHubRun.cpp
@@ -4,6 +4,7 @@
 
 #include "Misc/App.h"
 #include "Misc/PackageName.h"
+#include "Misc/Paths.h"
 #include "HAL/FileManager.h"
 #include "HAL/PlatformApplicationMisc.h"
 #include "HAL/PlatformFileManager.h"
@@ -44,7 +45,17 @@ int32 RunLiveLinkHub(const TCHAR* CommandLine)
 #endif
 
 #if !IS_PROGRAM
-	FPaths::SetProjectFilePath(TEXT("../../Source/Programs/LiveLinkHubEditor/LiveLinkHubEditor.uproject"));
+	const TCHAR* const DevelopmentProjectPath = TEXT("../../Source/Programs/LiveLinkHubEditor/LiveLinkHubEditor.uproject");
+	const TCHAR* const StagedProjectPath = TEXT("../../../LiveLinkHubEditor/LiveLinkHubEditor.uproject");
+
+	if (FPaths::FileExists(DevelopmentProjectPath))
+	{
+		FPaths::SetProjectFilePath(DevelopmentProjectPath);
+	}
+	else if (FPaths::FileExists(StagedProjectPath))
+	{
+		FPaths::SetProjectFilePath(StagedProjectPath);
+	}
 #endif
 
 	// Must match the path in LiveLinkHubStaging.cs / StageLiveLinkHub.automation.cs
diff --git a/Engine/Source/Programs/LiveLinkHubEditor/Source/LiveLinkHubEditor.Target.cs b/Engine/Source/Programs/LiveLinkHubEditor/Source/LiveLinkHubEditor.Target.cs
--- a/Engine/Source/Programs/LiveLinkHubEditor/Source/LiveLinkHubEditor.Target.cs
+++ b/Engine/Source/Programs/LiveLinkHubEditor/Source/LiveLinkHubEditor.Target.cs
@@ -8,6 +8,12 @@ using UnrealBuildTool;
 [SupportedPlatforms(UnrealPlatformClass.Desktop)]
 public class LiveLinkHubEditorTarget : TargetRules
 {
+	// Restrict OptedInModulePlatforms to the current Target.Platform.
+	// Used during staging, which otherwise fails in TargetPlatform-related restricted
+	// subdirectories (Engine/Binaries/Win64/{Android,IOS,Linux,LinuxArm64,...}).
+	[CommandLine("-SingleModulePlatform")]
+	public bool bSingleModulePlatform = false;
+
 	// Whether to enable building Capture Manager plugin.
 	[CommandLine("-EnableCaptureManagerPlugin=")]
 	public bool bEnableCaptureManagerPlugin = false;
@@ -51,9 +57,12 @@ public class LiveLinkHubEditorTarget : TargetRules
 			});
 		}
 
-		// FIXME?: Without this, staging fails on several files in TargetPlatform-related
-		// restricted subdirectories (Engine/Binaries/Win64/{Android,IOS,Linux,LinuxArm64...}).
-		OptedInModulePlatforms = new UnrealTargetPlatform[] { Target.Platform };
+		if (bSingleModulePlatform)
+		{
+			// Necessary for staging, but avoided otherwise because it dirties
+			// Definitions.CookedEditor.h and triggers rebuilds (incl. UnrealEditor).
+			OptedInModulePlatforms = new UnrealTargetPlatform[] { Target.Platform };
+		}
 
 		bAllowEnginePluginsEnabledByDefault = false;
 		bBuildAdditionalConsoleApp = false;
