Allow staging to pull chunk manifest info from zenserver directly instead of from files on the filesystem so that zen snapshots (which contain the chunk manifests produced during cook, but not as files on the filesystem) can stage to pak correctly.

diff --git a/Engine/Source/Programs/AutomationTool/AutomationUtils/AutomationUtils.Automation.csproj b/Engine/Source/Programs/AutomationTool/AutomationUtils/AutomationUtils.Automation.csproj
--- a/Engine/Source/Programs/AutomationTool/AutomationUtils/AutomationUtils.Automation.csproj
+++ b/Engine/Source/Programs/AutomationTool/AutomationUtils/AutomationUtils.Automation.csproj
@@ -58,6 +58,7 @@
       <Private>true</Private>
     </ProjectReference>
     <ProjectReference Include="..\..\Shared\EpicGames.Perforce\EpicGames.Perforce.csproj" />
+    <ProjectReference Include="..\..\Shared\EpicGames.ProjectStore\EpicGames.ProjectStore.csproj" PrivateAssets="All" />
     <ProjectReference Include="..\..\UnrealBuildTool\UnrealBuildTool.csproj" PrivateAssets="All"><Private>true</Private></ProjectReference>
   </ItemGroup>
   <ItemGroup>
diff --git a/Engine/Source/Programs/AutomationTool/AutomationUtils/DeploymentContext.cs b/Engine/Source/Programs/AutomationTool/AutomationUtils/DeploymentContext.cs
--- a/Engine/Source/Programs/AutomationTool/AutomationUtils/DeploymentContext.cs
+++ b/Engine/Source/Programs/AutomationTool/AutomationUtils/DeploymentContext.cs
@@ -8,6 +8,7 @@ using AutomationTool;
 using UnrealBuildTool;
 using System.Linq;
 using EpicGames.Core;
+using EpicGames.ProjectStore;
 
 public struct StageTarget
 {
@@ -94,11 +95,18 @@ public class FilesToStage
 	}
 }
 
+public class ZenCookedFile
+{
+	public string Filename;
+	public IoHash ChunkId;
+}
+
 public class PackageStoreData
 {
+	public ZenServerStoreData ZenServerStore { get; set; }
 	public string ManifestFullPath { get; set; }
 	public string MarkerFullPath { get; set; }
-	public IList<string> ZenCookedFiles { get; set; }
+	public IList<ZenCookedFile> ZenCookedFiles { get; set; }
 }
 
 public class DeploymentContext //: ProjectParams
diff --git a/Engine/Source/Programs/AutomationTool/Scripts/CopyBuildToStagingDirectory.Automation.cs b/Engine/Source/Programs/AutomationTool/Scripts/CopyBuildToStagingDirectory.Automation.cs
--- a/Engine/Source/Programs/AutomationTool/Scripts/CopyBuildToStagingDirectory.Automation.cs
+++ b/Engine/Source/Programs/AutomationTool/Scripts/CopyBuildToStagingDirectory.Automation.cs
@@ -255,9 +255,9 @@ namespace AutomationScripts
 		/// </summary>
 		/// <param name="Filename"></param>
 		/// <returns></returns>
-		private static HashSet<string> ReadPakChunkManifest(string Filename)
+		private static HashSet<string> ReadPakChunkManifest(DeploymentContext SC, string Filename)
 		{
-			var ResponseFile = ReadAllLines(Filename);
+			var ResponseFile = ReadAllLinesFilesystemOrPackageStore(SC, Filename);
 			var Result = new HashSet<string>(ResponseFile, StringComparer.InvariantCultureIgnoreCase);
 			return Result;
 		}
@@ -831,9 +831,14 @@ namespace AutomationScripts
 			return Plugins.ToList();
 		}
 
-		private static List<string> ReadZenCookedFilesFromZenServer(string Host, int Port, string ProjectId, string OplogId, ProjectParams Params, DeploymentContext SC, bool bAutoLaunch, string PackageStoreFileArgName, string PackageStoreFileArgValue)
+		private static IoHash ReadChunkId(CbField Field)
 		{
-			List<string> ZenCookedFiles = new List<string>();
+			return Field.AsHash();
+		}
+
+		private static List<ZenCookedFile> ReadZenCookedFilesFromZenServer(string Host, int Port, string ProjectId, string OplogId, ProjectParams Params, DeploymentContext SC, bool bAutoLaunch, string PackageStoreFileArgName, string PackageStoreFileArgValue)
+		{
+			List<ZenCookedFile> ZenCookedFiles = new List<ZenCookedFile>();
 
 			bool bAttemptAutoLaunchOnFailure = bAutoLaunch;
 			HttpResponseMessage HttpGetResult = null;
@@ -879,23 +884,25 @@ namespace AutomationScripts
 				foreach (CbField PackageDataField in EntryField["packagedata"].AsArray())
 				{
 					string RelativeFilename = PackageDataField["filename"].AsString();
-					ZenCookedFiles.Add(FileReference.Combine(SC.PlatformCookDir, RelativeFilename).FullName);
+					ZenCookedFiles.Add(new ZenCookedFile { Filename = FileReference.Combine(SC.PlatformCookDir, RelativeFilename).FullName, ChunkId = ReadChunkId(PackageDataField["data"]) });
 				}
 				foreach (CbField PackageDataField in EntryField["bulkdata"].AsArray())
 				{
 					string RelativeFilename = PackageDataField["filename"].AsString();
-					ZenCookedFiles.Add(FileReference.Combine(SC.PlatformCookDir, RelativeFilename).FullName);
+					ZenCookedFiles.Add(new ZenCookedFile { Filename = FileReference.Combine(SC.PlatformCookDir, RelativeFilename).FullName, ChunkId = ReadChunkId(PackageDataField["data"]) });
 				}
 				if (EntryField["key"].AsString() == "EndCook")
 				{
 					foreach (CbField PackageDataField in EntryField["files"].AsArray())
 					{
 						string RelativeFilename = PackageDataField["clientpath"].AsString();
-						if (RelativeFilename.EndsWith(".bin") || RelativeFilename.EndsWith(".ushaderbytecode"))
+						if (RelativeFilename.EndsWith(".bin") || RelativeFilename.EndsWith(".ushaderbytecode") ||
+							(RelativeFilename.EndsWith(".txt") && RelativeFilename.Contains("/Metadata/ChunkManifest/"))
+							)
 						{
 							RelativeFilename = RelativeFilename.Replace("/{engine}/", "Engine/");
 							RelativeFilename = RelativeFilename.Replace("/{project}/", SC.ShortProjectName+"/");
-							ZenCookedFiles.Add(FileReference.Combine(SC.PlatformCookDir, RelativeFilename).FullName);
+							ZenCookedFiles.Add(new ZenCookedFile { Filename = FileReference.Combine(SC.PlatformCookDir, RelativeFilename).FullName, ChunkId = ReadChunkId(PackageDataField["data"]) });
 						}
 					}
 				}
@@ -906,7 +913,7 @@ namespace AutomationScripts
 						string RelativeFilename = PackageDataField["clientpath"].AsString();
 						RelativeFilename = RelativeFilename.Replace("/{engine}/", "Engine/");
 						RelativeFilename = RelativeFilename.Replace("/{project}/", SC.ShortProjectName+"/");
-						ZenCookedFiles.Add(FileReference.Combine(SC.PlatformCookDir, RelativeFilename).FullName);
+						ZenCookedFiles.Add(new ZenCookedFile { Filename = FileReference.Combine(SC.PlatformCookDir, RelativeFilename).FullName, ChunkId = ReadChunkId(PackageDataField["data"]) });
 					}
 				}
 			}
@@ -940,6 +947,7 @@ namespace AutomationScripts
 				{
 					Logger.LogInformation("Reading oplog from Zen...");
 					ZenServerStoreData ZenServer = ParsedProjectStore.ZenServer;
+					SC.PackageStoreData.ZenServerStore = ZenServer;
 					SC.PackageStoreData.ZenCookedFiles = ReadZenCookedFilesFromZenServer(ZenServer.HostName, ZenServer.HostPort, ZenServer.ProjectId, ZenServer.OplogId, Params, SC, true, "ProjectStore", ProjectStoreFile.FullName);
 					return;
 				}
@@ -959,6 +967,7 @@ namespace AutomationScripts
 				{
 					Logger.LogInformation("Reading oplog from Zen...");
 					string Host = "localhost";
+					bool IsLocalHost = true;
 					int Port = 8558;
 					string ProjectId = ZenServerObject["projectid"].AsString();
 					string OplogId = ZenServerObject["oplogid"].AsString();
@@ -976,10 +985,12 @@ namespace AutomationScripts
 						if (ConnectExistingSettingsObject != CbObject.Empty)
 						{
 							Host = ConnectExistingSettingsObject["HostName"].AsString();
+							IsLocalHost = Host == "localhost" || Host == "127.0.0.1" || Host == "::1";
 							Port = ConnectExistingSettingsObject["Port"].AsInt16();
 						}
 					}
 
+					SC.PackageStoreData.ZenServerStore = new ZenServerStoreData { HostName = Host, HostPort = Port, ProjectId = ProjectId, OplogId = OplogId, IsLocalHost = IsLocalHost };
 					SC.PackageStoreData.ZenCookedFiles = ReadZenCookedFilesFromZenServer(Host, Port, ProjectId, OplogId, Params, SC, bAutoLaunch, "PackageStoreManifest", PackageStoreManifestFile.FullName);
 				}
 				return;
@@ -1188,9 +1199,9 @@ namespace AutomationScripts
 					{
 						CookedFiles.AddRange(DirectoryReference.EnumerateFiles(SC.PlatformCookDir, "*.projectstore", SearchOption.TopDirectoryOnly).ToList());
 					}
-					foreach (string FilePath in SC.PackageStoreData.ZenCookedFiles)
+					foreach (ZenCookedFile CookedFile in SC.PackageStoreData.ZenCookedFiles)
 					{
-						CookedFiles.Add(new FileReference(FilePath));
+						CookedFiles.Add(new FileReference(CookedFile.Filename));
 					}
 				}
 				else
@@ -1548,9 +1559,9 @@ namespace AutomationScripts
 						LoadPackageStoreData(Params, SC);
 						if (SC.PackageStoreData != null && SC.PackageStoreData.ZenCookedFiles != null)
 						{
-							foreach (string FilePath in SC.PackageStoreData.ZenCookedFiles)
+							foreach (ZenCookedFile CookedFile in SC.PackageStoreData.ZenCookedFiles)
 							{
-								CookedFiles.Add(new FileReference(FilePath));
+								CookedFiles.Add(new FileReference(CookedFile.Filename));
 							}
 						}
 
@@ -4458,6 +4469,68 @@ namespace AutomationScripts
 			}
 		}
 
+		private static byte[] GetChunkFromPackageStore(DeploymentContext SC, string Filename)
+		{
+			if (SC.PackageStoreData != null && SC.PackageStoreData.ZenServerStore != null && SC.PackageStoreData.ZenCookedFiles != null)
+			{
+				ZenCookedFile ZenFile = SC.PackageStoreData.ZenCookedFiles.FirstOrDefault(Item => Item.ChunkId != IoHash.Zero && Item.Filename == Filename);
+				if (ZenFile != null)
+				{
+					HttpResponseMessage HttpGetResult = null;
+
+					HttpClient HttpClient = new HttpClient();
+					using var Request = new HttpRequestMessage(HttpMethod.Get, string.Format("http://{0}:{1}/prj/{2}/oplog/{3}/{4}", SC.PackageStoreData.ZenServerStore.HostName, SC.PackageStoreData.ZenServerStore.HostPort, SC.PackageStoreData.ZenServerStore.ProjectId, SC.PackageStoreData.ZenServerStore.OplogId, ZenFile.ChunkId.ToString()));
+					Request.Headers.Add("Accept", "application/octet-stream");
+					try
+					{
+						HttpGetResult = HttpClient.Send(Request);
+					}
+					catch
+					{
+						throw new AutomationException(String.Format("Failed sending chunk request to Zen at {0}:{1}. Ensure that the server is running.", SC.PackageStoreData.ZenServerStore.HostName, SC.PackageStoreData.ZenServerStore.HostPort));
+					}
+
+					if (!HttpGetResult.IsSuccessStatusCode)
+					{
+						throw new AutomationException(String.Format("Failed reading chunk {0}.{1} from Zen. Ensure that cooking was successful.", SC.PackageStoreData.ZenServerStore.ProjectId, SC.PackageStoreData.ZenServerStore.OplogId));
+					}
+					Task<byte[]> ReadChunkTask = HttpGetResult.Content.ReadAsByteArrayAsync();
+					ReadChunkTask.Wait();
+					return ReadChunkTask.Result;
+				}
+			}
+
+			return null;
+		}
+
+		private static string[] ReadAllLinesFilesystemOrPackageStore(DeploymentContext SC, string Filename)
+		{
+			Filename = ConvertSeparators(PathSeparator.Default, Filename);
+			if (FileExists_NoExceptions(Filename))
+			{
+				return InternalUtils.SafeReadAllLines(Filename);
+			}
+			else if (SC.PackageStoreData != null && SC.PackageStoreData.ZenServerStore != null && SC.PackageStoreData.ZenCookedFiles != null)
+			{
+				byte[] ChunkBytes = GetChunkFromPackageStore(SC, Filename);
+				if (ChunkBytes != null)
+				{
+					List<string> Lines = new List<String>();
+					using (var stream = new StreamReader(new MemoryStream(ChunkBytes)))
+					{
+						string Line;
+						while ((Line = stream.ReadLine()) != null)
+						{
+							Lines.Add(Line);
+						}
+					}
+					return Lines.ToArray();
+				}
+			}
+
+			return null;
+		}
+
 		/// <summary>
 		/// Creates pak files using streaming install chunk manifests.
 		/// </summary>
@@ -4498,7 +4571,7 @@ namespace AutomationScripts
 				const int DefaultChunkIndex = 0;
 
 				var ChunkListFilename = GetChunkPakManifestListFilename(Params, SC);
-				List<string> ChunkList = new List<string>(ReadAllLines(ChunkListFilename));
+				List<string> ChunkList = new List<string>(ReadAllLinesFilesystemOrPackageStore(SC, ChunkListFilename));
 				Logger.LogInformation("Reading chunk list file {ChunkListFilename} which contains {Arg1} entries", ChunkListFilename, ChunkList.Count);
 
 				for (int Index = 0; Index < ChunkList.Count; ++Index)
@@ -4535,7 +4608,7 @@ namespace AutomationScripts
 							}
 						}
 					}
-					CD.Manifest = ReadPakChunkManifest(ChunkManifestFilename);
+					CD.Manifest = ReadPakChunkManifest(SC, ChunkManifestFilename);
 					Logger.LogInformation("Reading chunk manifest {ChunkManifestFilename} which contains {Arg1} entries", ChunkManifestFilename, CD.Manifest.Count);
 					ChunkDefinitions.Add(CD);
 				}
@@ -4841,7 +4914,22 @@ namespace AutomationScripts
 			String ChunkLayerFilename = CombinePaths(GetChunkManifestPath(Params, SC), GetChunkPakLayerListName());
 			String OutputChunkLayerFilename = Path.Combine(SC.ProjectRoot.FullName, "Build", SC.FinalCookPlatform, "ChunkLayerInfo", GetChunkPakLayerListName());
 			Directory.CreateDirectory(Path.GetDirectoryName(OutputChunkLayerFilename));
-			File.Copy(ChunkLayerFilename, OutputChunkLayerFilename, true);
+			if (File.Exists(ChunkLayerFilename))
+			{
+				File.Copy(ChunkLayerFilename, OutputChunkLayerFilename, true);
+			}
+			else
+			{
+				byte[] ChunkLayerContents = GetChunkFromPackageStore(SC, ChunkLayerFilename);
+				if (ChunkLayerContents != null)
+				{
+					File.WriteAllBytes(OutputChunkLayerFilename, ChunkLayerContents);
+				}
+				else
+				{
+					throw new AutomationException("Missing or invalid chunk layer file ({0})", ChunkLayerFilename);
+				}
+			}
 		}
 
 		private static void GeneratePrimaryChunkManifest(string Dir, string Version, string PlatformStr)
@@ -4883,7 +4971,20 @@ namespace AutomationScripts
 
 		private static bool DoesChunkPakManifestExist(ProjectParams Params, DeploymentContext SC)
 		{
-			return FileExists_NoExceptions(GetChunkPakManifestListFilename(Params, SC));
+			if (FileExists_NoExceptions(GetChunkPakManifestListFilename(Params, SC)))
+			{
+				return true;
+			}
+
+			if ((SC.PackageStoreData != null) && (SC.PackageStoreData.ZenCookedFiles != null))
+			{
+				if (SC.PackageStoreData.ZenCookedFiles.Any(Item => Item.ChunkId != IoHash.Zero && Item.Filename == GetChunkPakManifestListFilename(Params, SC)))
+				{
+					return true;
+				}
+			}
+
+			return false;
 		}
 
 		private static string GetChunkPakManifestListFilename(ProjectParams Params, DeploymentContext SC)
