AutomationTool: Add KnownLogEvent code for code signing

#rnx

diff --git a/Engine/Source/Programs/AutomationTool/AutomationUtils/CommandUtils.cs b/Engine/Source/Programs/AutomationTool/AutomationUtils/CommandUtils.cs
--- a/Engine/Source/Programs/AutomationTool/AutomationUtils/CommandUtils.cs
+++ b/Engine/Source/Programs/AutomationTool/AutomationUtils/CommandUtils.cs
@@ -3121,7 +3121,7 @@ namespace AutomationTool
 						{
 							if (Result.ExitCode == 2)
 							{
-								Logger.LogWarning("Signtool returned a warning.");
+								Logger.LogWarning(KnownLogEvents.AutomationTool_CodeSign, "Signtool returned a warning.");
 							}
 							// Success!
 							break;
@@ -3193,7 +3193,7 @@ namespace AutomationTool
 		{
 			if (!OperatingSystem.IsWindows())
 			{
-				Logger.LogDebug("Can't sign '{File}' on non-Windows platform.", Filename);
+				Logger.LogDebug(KnownLogEvents.AutomationTool_CodeSign, "Can't sign '{File}' on non-Windows platform.", Filename);
 				return;
 			}
 			if (!CommandUtils.FileExists(Filename))
@@ -3222,7 +3222,7 @@ namespace AutomationTool
 			}
 			if (!IsExecutable)
 			{
-				Logger.LogDebug("Won't sign '{File}', not an executable.", TargetFileInfo.FullName);
+				Logger.LogDebug(KnownLogEvents.AutomationTool_CodeSign, "Won't sign '{File}', not an executable.", TargetFileInfo.FullName);
 				return;
 			}
 
@@ -3299,7 +3299,7 @@ namespace AutomationTool
 			{
 				if (!bIsDirectory)
 				{
-					Logger.LogDebug("Won't sign '{File}', not an executable.", InPath);
+					Logger.LogDebug(KnownLogEvents.AutomationTool_CodeSign, "Won't sign '{File}', not an executable.", InPath);
 				}
 				return;
 			}
@@ -3362,7 +3362,7 @@ namespace AutomationTool
 		{
 			if (!Command.ParseParam("NoSign"))
 			{
-				Logger.LogInformation("Signing up to {FileCount} files...", Files.Count());
+				Logger.LogInformation(KnownLogEvents.AutomationTool_CodeSign, "Signing up to {FileCount} files...", Files.Count());
 				UnrealBuildTool.UnrealTargetPlatform TargetPlatform = UnrealBuildTool.BuildHostPlatform.Current.Platform;
 				if (TargetPlatform == UnrealBuildTool.UnrealTargetPlatform.Mac)
 				{
@@ -3383,7 +3383,7 @@ namespace AutomationTool
 			}
 			else
 			{
-				Logger.LogDebug("Skipping signing {FileCount} files due to -nosign.", Files.Count());
+				Logger.LogDebug(KnownLogEvents.AutomationTool_CodeSign, "Skipping signing {FileCount} files due to -nosign.", Files.Count());
 			}
 		}
 
@@ -3391,7 +3391,7 @@ namespace AutomationTool
 		{
 			if (!OperatingSystem.IsWindows())
 			{
-				Logger.LogDebug("Can't sign on non-Windows platform.");
+				Logger.LogDebug(KnownLogEvents.AutomationTool_CodeSign, "Can't sign on non-Windows platform.");
 				return;
 			}
 			List<FileReference> FinalFiles = new List<FileReference>();
diff --git a/Engine/Source/Programs/Shared/EpicGames.Core/KnownLogEvents.cs b/Engine/Source/Programs/Shared/EpicGames.Core/KnownLogEvents.cs
--- a/Engine/Source/Programs/Shared/EpicGames.Core/KnownLogEvents.cs
+++ b/Engine/Source/Programs/Shared/EpicGames.Core/KnownLogEvents.cs
@@ -140,6 +140,11 @@ namespace EpicGames.Core
 		/// </summary>
 		public static EventId AutomationTool_BuildGraphScript { get; } = new EventId(407);
 
+		/// <summary>
+		/// UAT: CodeSign message
+		/// </summary>
+		public static EventId AutomationTool_CodeSign { get; } = new EventId(408);
+
 		/// <summary>
 		/// MSBuild: Generic error
 		/// </summary>
