AutomationTool: Add KnownLogEvent code for code signing

#rnx

diff --git a/Engine/Source/Programs/AutomationTool/AutomationUtils/CommandUtils.cs b/Engine/Source/Programs/AutomationTool/AutomationUtils/CommandUtils.cs
--- a/Engine/Source/Programs/AutomationTool/AutomationUtils/CommandUtils.cs
+++ b/Engine/Source/Programs/AutomationTool/AutomationUtils/CommandUtils.cs
@@ -3141,7 +3141,7 @@ namespace AutomationTool
 
 			// If running in parallel, Limit to between 4 and 16 concurrent. Otherwise 1 instance.
 			int MaxParallelism = RunInParallel ? Math.Max(16, Math.Min(Environment.ProcessorCount, 4)) : 1;
-			Logger.LogInformation("Running {Count}, {Concurrent} max concurrent, signtool instances", SignTaskList.Count, MaxParallelism);
+			Logger.LogInformation(KnownLogEvents.AutomationTool_CodeSign, "Running {Count}, {Concurrent} max concurrent, signtool instances", SignTaskList.Count, MaxParallelism);
 			Parallel.ForEach(SignTaskList, new ParallelOptions() { MaxDegreeOfParallelism = MaxParallelism }, x => { x.Start(); x.Wait(); });
 		}
 
