Use new MinCompression stuff in BPT on iOS platform.

diff --git a/Engine/Source/Programs/AutomationTool/AutomationUtils/MCPPublic.cs b/Engine/Source/Programs/AutomationTool/AutomationUtils/MCPPublic.cs
--- a/Engine/Source/Programs/AutomationTool/AutomationUtils/MCPPublic.cs
+++ b/Engine/Source/Programs/AutomationTool/AutomationUtils/MCPPublic.cs
@@ -804,6 +804,10 @@ namespace EpicGames.MCP.Automation
 			/// Contains a list of custom float arguments to be embedded in the generated manifest file.
 			/// </summary>
 			public List<KeyValuePair<string, float>> CustomFloatArgs;
+			/// <summary>
+			/// Specifies additional arguments to be passed to BPT
+			/// </summary>
+			public string AdditionalArguments;
 		}
 
 
