Change pak chunk manifests to be relative paths, not absolute.  This allows them to be re-used as part of snapshots when the generating machine's root path will differ from the consuming machine's root path.

diff --git a/Engine/Source/Editor/UnrealEd/Private/Commandlets/AssetRegistryGenerator.cpp b/Engine/Source/Editor/UnrealEd/Private/Commandlets/AssetRegistryGenerator.cpp
--- a/Engine/Source/Editor/UnrealEd/Private/Commandlets/AssetRegistryGenerator.cpp
+++ b/Engine/Source/Editor/UnrealEd/Private/Commandlets/AssetRegistryGenerator.cpp
@@ -750,6 +750,7 @@ bool FAssetRegistryGenerator::GenerateStreamingInstallManifest(int64 InOverrideC
 					}
 				}
 				
+				FPaths::MakePathRelativeTo(PakListLine, *FPaths::RootDir());
 				PakListLine.ReplaceInline(TEXT("/"), TEXT("\\"));
 				PakListLine += TEXT("\r\n");
 				PakListFile->Serialize(TCHAR_TO_ANSI(*PakListLine), PakListLine.Len());
diff --git a/Engine/Source/Programs/AutomationTool/Scripts/CopyBuildToStagingDirectory.Automation.cs b/Engine/Source/Programs/AutomationTool/Scripts/CopyBuildToStagingDirectory.Automation.cs
--- a/Engine/Source/Programs/AutomationTool/Scripts/CopyBuildToStagingDirectory.Automation.cs
+++ b/Engine/Source/Programs/AutomationTool/Scripts/CopyBuildToStagingDirectory.Automation.cs
@@ -257,7 +257,7 @@ namespace AutomationScripts
 		/// <returns></returns>
 		private static HashSet<string> ReadPakChunkManifest(DeploymentContext SC, string Filename)
 		{
-			var ResponseFile = ReadAllLinesFilesystemOrPackageStore(SC, Filename);
+			var ResponseFile = ReadAllLinesFilesystemOrPackageStore(SC, Filename).Select(Line => FileReference.Combine(SC.LocalRoot, Line).FullName);
 			var Result = new HashSet<string>(ResponseFile, StringComparer.InvariantCultureIgnoreCase);
 			return Result;
 		}
