[Gauntlet] Avoid parsing the log several times when interpreting results from UE automated tests by allowing UnrealLogParser to take UnrealLog as input constructor

#rnx

diff --git a/Engine/Source/Programs/AutomationTool/Gauntlet/Unreal/Automation/UE.Automation.cs b/Engine/Source/Programs/AutomationTool/Gauntlet/Unreal/Automation/UE.Automation.cs
--- a/Engine/Source/Programs/AutomationTool/Gauntlet/Unreal/Automation/UE.Automation.cs
+++ b/Engine/Source/Programs/AutomationTool/Gauntlet/Unreal/Automation/UE.Automation.cs
@@ -637,7 +637,7 @@ namespace UE
 				{
 					// Parse automaton info from the log then
 					TestPassResults = new UnrealAutomatedTestPassResults();
-					AutomationLogParser LogParser = new AutomationLogParser(InLog.FullLogContent);
+					AutomationLogParser LogParser = new AutomationLogParser(InLog);
 					IEnumerable<UnrealAutomatedTestResult> LogTestResults = LogParser.GetResults();
 					if (LogTestResults.Any())
 					{
@@ -1230,7 +1230,7 @@ namespace UE
 			{
 				if (Role.Artifacts.SessionRole.RoleType.IsEditor())
 				{
-					AutomationLogParser Parser = new AutomationLogParser(Role.LogSummary.FullLogContent);
+					AutomationLogParser Parser = new AutomationLogParser(Role.LogSummary);
 					AllErrors.AddRange(
 						Parser.GetResults().Where(R => R.HasFailed)
 							.SelectMany(R => R.Entries
@@ -1261,7 +1261,7 @@ namespace UE
 			{
 				if (Role.Artifacts.SessionRole.RoleType == UnrealTargetRole.Editor)
 				{
-					AutomationLogParser Parser = new AutomationLogParser(Role.LogSummary.FullLogContent);
+					AutomationLogParser Parser = new AutomationLogParser(Role.LogSummary);
 					AllWarnings.AddRange(
 						Parser.GetResults()
 							.SelectMany(R => R.Entries
diff --git a/Engine/Source/Programs/AutomationTool/Gauntlet/Unreal/Engine/UnrealTests.cs b/Engine/Source/Programs/AutomationTool/Gauntlet/Unreal/Engine/UnrealTests.cs
--- a/Engine/Source/Programs/AutomationTool/Gauntlet/Unreal/Engine/UnrealTests.cs
+++ b/Engine/Source/Programs/AutomationTool/Gauntlet/Unreal/Engine/UnrealTests.cs
@@ -75,7 +75,7 @@ namespace Gauntlet
 				{
 					if (Role.Artifacts.SessionRole.RoleType == UnrealTargetRole.Editor)
 					{
-						AutomationLogParser Parser = new AutomationLogParser(Role.LogSummary.FullLogContent);
+						AutomationLogParser Parser = new AutomationLogParser(Role.LogSummary);
 						AllErrors.AddRange(
 							Parser.GetResults().Where(R => R.HasFailed)
 							.SelectMany(R => R.Entries
@@ -221,7 +221,7 @@ namespace Gauntlet
 					// if no fatal errors, check test results
 					if (InLog.FatalError == null)
 					{
-						AutomationLogParser Parser = new AutomationLogParser(InLog.FullLogContent);
+						AutomationLogParser Parser = new AutomationLogParser(InLog);
 
 						IEnumerable<UnrealAutomatedTestResult> TotalTests = Parser.GetResults();
 						IEnumerable<UnrealAutomatedTestResult> FailedTests = TotalTests.Where(R => R.HasFailed);
@@ -259,7 +259,7 @@ namespace Gauntlet
 
 				if (EditorRole != null)
 				{
-					AutomationLogParser Parser = new AutomationLogParser(EditorRole.LogSummary.FullLogContent);
+					AutomationLogParser Parser = new AutomationLogParser(EditorRole.LogSummary);
 
 					IEnumerable<UnrealAutomatedTestResult> AllTests = Parser.GetResults();
 					IEnumerable<UnrealAutomatedTestResult> FailedTests = AllTests.Where(R => R.IsComplete && R.HasFailed);
diff --git a/Engine/Source/Programs/AutomationTool/Gauntlet/Unreal/Utils/Gauntlet.AutomationLogParser.cs b/Engine/Source/Programs/AutomationTool/Gauntlet/Unreal/Utils/Gauntlet.AutomationLogParser.cs
--- a/Engine/Source/Programs/AutomationTool/Gauntlet/Unreal/Utils/Gauntlet.AutomationLogParser.cs
+++ b/Engine/Source/Programs/AutomationTool/Gauntlet/Unreal/Utils/Gauntlet.AutomationLogParser.cs
@@ -73,7 +73,16 @@ namespace Gauntlet
 				: this(new UnrealLogParser(InContents))
 		{
 		}
-		
+
+		/// <summary>
+		/// Constructor that takes UnrealLog instance
+		/// </summary>
+		/// <param name="InLog"></param>
+		public AutomationLogParser(UnrealLog InLog)
+				: this(new UnrealLogParser(InLog))
+		{
+		}
+
 		/// <summary>
 		/// Returns all results found in our construction content.
 		/// </summary>
diff --git a/Engine/Source/Programs/AutomationTool/Gauntlet/Unreal/Utils/Gauntlet.UnrealLogParser.cs b/Engine/Source/Programs/AutomationTool/Gauntlet/Unreal/Utils/Gauntlet.UnrealLogParser.cs
--- a/Engine/Source/Programs/AutomationTool/Gauntlet/Unreal/Utils/Gauntlet.UnrealLogParser.cs
+++ b/Engine/Source/Programs/AutomationTool/Gauntlet/Unreal/Utils/Gauntlet.UnrealLogParser.cs
@@ -675,6 +675,16 @@ namespace Gauntlet
 			return ContentBuilder.ToString();
 		}
 
+		/// <summary>
+		/// Constructor that takes a UnrealLog instance
+		/// </summary>
+		/// <param name="Log"></param>
+		public UnrealLogParser(UnrealLog Log)
+		{
+			Content = Log.FullLogContent;
+			LogEntries = Log.LogEntries;
+		}
+
 		public UnrealLog GetSummary()
 		{
 			if (Summary == null)
@@ -958,15 +968,9 @@ namespace Gauntlet
 		/// <returns></returns>
 		public UnrealLog.CallstackMessage GetFatalError()
 		{
-			var Traces = GetASanErrors();
-			if (Traces.Any())
-			{
-				return Traces.Last();
-			}
-
 			string[] ErrorMsgMatches = new string[] { @"(Fatal Error:.+)", @"Critical error: =+\s+(?:[\S\s]+?\s*Error: +)?(.+)", @"(Assertion Failed:.+)", @"(Unhandled Exception:.+)", @"(LowLevelFatalError.+)", @"(Postmortem Cause:.*)" };
 
-			Traces = ParseTracedErrors(ErrorMsgMatches, 5);
+			var Traces = ParseTracedErrors(ErrorMsgMatches, 5).Concat(GetASanErrors());
 
 			// If we have a post-mortem error, return that one (on some devices the post-mortem info is way more informative).
 			var PostMortemTraces = Traces.Where(T => T.Message.IndexOf("Postmortem Cause:", StringComparison.OrdinalIgnoreCase) > -1);
