[Automation] Prevent server installation nodes from clearing client build installations

diff --git a/Engine/Source/Programs/AutomationTool/Gauntlet/Framework/Devices/Gauntlet.DevicePool.cs b/Engine/Source/Programs/AutomationTool/Gauntlet/Framework/Devices/Gauntlet.DevicePool.cs
--- a/Engine/Source/Programs/AutomationTool/Gauntlet/Framework/Devices/Gauntlet.DevicePool.cs
+++ b/Engine/Source/Programs/AutomationTool/Gauntlet/Framework/Devices/Gauntlet.DevicePool.cs
@@ -1263,9 +1263,10 @@ namespace Gauntlet
 			// When using device reservation blocks, we don't want to fully clean the cache and lose previously installed builds.
 			// If bRetainBuilds evaluates to true, it means we are in the second step or beyond in a device reservation block.
 			// In this case we'll just delete the left over UserDir which should already have been emptied by UnrealSession.
+			bool bRetainCache = Globals.Params.ParseParam("RetainCache");
 			bool bRetainBuilds = SkipInstall && !FullClean;
 
-			if(bRetainBuilds)
+			if(bRetainBuilds || bRetainCache)
 			{
 				Log.Info("Retaining build cache for device reservation block");
 
