Horde: Remove references to UnrealBuildBase from BgScriptReader.

#rnx

diff --git a/Engine/Source/Programs/AutomationTool/BuildGraph/BgScriptReader.cs b/Engine/Source/Programs/AutomationTool/BuildGraph/BgScriptReader.cs
--- a/Engine/Source/Programs/AutomationTool/BuildGraph/BgScriptReader.cs
+++ b/Engine/Source/Programs/AutomationTool/BuildGraph/BgScriptReader.cs
@@ -13,7 +13,6 @@ using System.Xml.Schema;
 using EpicGames.BuildGraph;
 using EpicGames.Core;
 using Microsoft.Extensions.Logging;
-using UnrealBuildBase;
 
 #nullable enable
 
@@ -304,6 +303,11 @@ namespace AutomationTool
 	/// </summary>
 	public class BgScriptReader
 	{
+		/// <summary>
+		/// Root directory to resolve relative paths against
+		/// </summary>
+		readonly DirectoryReference _rootDir;
+
 		/// <summary>
 		/// List of property name to value lookups. Modifications to properties are scoped to nodes and agents. EnterScope() pushes an empty dictionary onto the end of this list, and LeaveScope() removes one. 
 		/// ExpandProperties() searches from last to first lookup when trying to resolve a property name, and takes the first it finds.
@@ -357,13 +361,16 @@ namespace AutomationTool
 		/// <summary>
 		/// Private constructor. Use ScriptReader.TryRead() to read a script file.
 		/// </summary>
+		/// <param name="rootDir">Root directory to resolve relative paths against</param>
 		/// <param name="defaultProperties">Default properties available to the script</param>
 		/// <param name="arguments">Arguments passed in to the graph on the command line</param>
 		/// <param name="singleNodeName">If a single node will be processed, the name of that node.</param>
 		/// <param name="schema">Schema for the script</param>
 		/// <param name="logger">Logger for diagnostic messages</param>
-		protected BgScriptReader(IDictionary<string, string?> defaultProperties, IReadOnlyDictionary<string, string> arguments, string? singleNodeName, BgScriptSchema? schema, ILogger logger)
+		protected BgScriptReader(DirectoryReference rootDir, IDictionary<string, string?> defaultProperties, IReadOnlyDictionary<string, string> arguments, string? singleNodeName, BgScriptSchema? schema, ILogger logger)
 		{
+			_rootDir = rootDir;
+
 			Schema = schema;
 			Logger = logger;
 
@@ -382,16 +389,17 @@ namespace AutomationTool
 		/// Try to read a script file from the given file.
 		/// </summary>
 		/// <param name="file">File to read from</param>
+		/// <param name="rootDir">Root directory to resolve files to</param>
 		/// <param name="arguments">Arguments passed in to the graph on the command line</param>
 		/// <param name="defaultProperties">Default properties available to the script</param>
 		/// <param name="schema">Schema for the script</param>
 		/// <param name="logger">Logger for output messages</param>
 		/// <param name="singleNodeName">If a single node will be processed, the name of that node.</param>
 		/// <returns>True if the graph was read, false if there were errors</returns>
-		public static async Task<BgGraphDef?> ReadAsync(FileReference file, Dictionary<string, string> arguments, Dictionary<string, string?> defaultProperties, BgScriptSchema? schema, ILogger logger, string? singleNodeName = null)
+		public static async Task<BgGraphDef?> ReadAsync(FileReference file, DirectoryReference rootDir, Dictionary<string, string> arguments, Dictionary<string, string?> defaultProperties, BgScriptSchema? schema, ILogger logger, string? singleNodeName = null)
 		{
 			// Read the file and build the graph
-			BgScriptReader reader = new BgScriptReader(defaultProperties, arguments, singleNodeName, schema, logger);
+			BgScriptReader reader = new BgScriptReader(rootDir, defaultProperties, arguments, singleNodeName, schema, logger);
 			if (!await reader.TryReadAsync(file) || reader.NumErrors > 0)
 			{
 				return null;
@@ -662,7 +670,7 @@ namespace AutomationTool
 		{
 			if (await EvaluateConditionAsync(element))
 			{
-				string basePath = element.Location.File.MakeRelativeTo(Unreal.RootDirectory).Replace(Path.DirectorySeparatorChar, '/');
+				string basePath = element.Location.File.MakeRelativeTo(_rootDir).Replace(Path.DirectorySeparatorChar, '/');
 
 				HashSet<FileReference> files = new HashSet<FileReference>();
 				foreach (string script in ReadListAttribute(element, "Script"))
@@ -670,7 +678,7 @@ namespace AutomationTool
 					string includePath = CombinePaths(basePath, script);
 					if (Regex.IsMatch(includePath, @"\*|\?|\.\.\."))
 					{
-						files.UnionWith(FindMatchingFiles(includePath));
+						files.UnionWith(FindMatchingFiles(_rootDir, includePath));
 					}
 					else
 					{
@@ -727,14 +735,12 @@ namespace AutomationTool
 		/// <summary>
 		/// Find files matching a pattern
 		/// </summary>
-		/// <param name="pattern"></param>
-		/// <returns></returns>
-		static IEnumerable<FileReference> FindMatchingFiles(string pattern)
+		static IEnumerable<FileReference> FindMatchingFiles(DirectoryReference rootDir, string pattern)
 		{
 			FileFilter filter = new FileFilter();
 			filter.AddRule(pattern, FileFilterType.Include);
 
-			return filter.ApplyToDirectory(Unreal.RootDirectory, true);
+			return filter.ApplyToDirectory(rootDir, true);
 		}
 
 		/// <summary>
@@ -2304,9 +2310,9 @@ namespace AutomationTool
 		}
 
 		/// <inheritdoc/>
-		public static object GetNativePath(string path)
+		public static object GetNativePath(DirectoryReference rootDir, string path)
 		{
-			return FileReference.Combine(Unreal.RootDirectory, path).FullName;
+			return FileReference.Combine(rootDir, path).FullName;
 		}
 	}
 }
\ No newline at end of file
diff --git a/Engine/Source/Programs/AutomationTool/BuildGraph/BuildGraph.cs b/Engine/Source/Programs/AutomationTool/BuildGraph/BuildGraph.cs
--- a/Engine/Source/Programs/AutomationTool/BuildGraph/BuildGraph.cs
+++ b/Engine/Source/Programs/AutomationTool/BuildGraph/BuildGraph.cs
@@ -500,7 +500,7 @@ namespace AutomationTool
 				FileReference fullScriptFile = FileReference.Combine(Unreal.RootDirectory, scriptFileName);
 
 				// Read the script from disk
-				graph = await BgScriptReader.ReadAsync(fullScriptFile, arguments, defaultProperties, schema, Logger, singleNodeName);
+				graph = await BgScriptReader.ReadAsync(fullScriptFile, Unreal.RootDirectory, arguments, defaultProperties, schema, Logger, singleNodeName);
 				if (graph == null)
 				{
 					return ExitCode.Error_Unknown;
