Horde: Replace a few BuildGraph exceptions with error messages indicating the file and line number.

#rnx

diff --git a/Engine/Source/Programs/AutomationTool/BuildGraph/BgScriptReader.cs b/Engine/Source/Programs/AutomationTool/BuildGraph/BgScriptReader.cs
--- a/Engine/Source/Programs/AutomationTool/BuildGraph/BgScriptReader.cs
+++ b/Engine/Source/Programs/AutomationTool/BuildGraph/BgScriptReader.cs
@@ -63,7 +63,7 @@ namespace AutomationTool
 		/// <summary>
 		/// Set to true if the reader encounters an error
 		/// </summary>
-		bool _bHasErrors;
+		bool _hasErrors;
 
 		/// <summary>
 		/// Logger for validation errors
@@ -120,16 +120,16 @@ namespace AutomationTool
 				}
 				catch (XmlException ex)
 				{
-					if (!document._bHasErrors)
+					if (!document._hasErrors)
 					{
 						BgScriptLocation location = new BgScriptLocation(file, ex.LineNumber);
 						logger.LogScriptError(location, "{Message}", ex.Message);
-						document._bHasErrors = true;
+						document._hasErrors = true;
 					}
 				}
 
 				// If we hit any errors while parsing
-				if (document._bHasErrors)
+				if (document._hasErrors)
 				{
 					outDocument = null;
 					return false;
@@ -171,7 +171,7 @@ namespace AutomationTool
 			else
 			{
 				Logger.LogScriptError(location, "{Message}", args.Message);
-				_bHasErrors = true;
+				_hasErrors = true;
 			}
 		}
 	}
@@ -675,8 +675,12 @@ namespace AutomationTool
 				HashSet<FileReference> files = new HashSet<FileReference>();
 				foreach (string script in ReadListAttribute(element, "Script"))
 				{
-					string includePath = CombinePaths(basePath, script);
-					if (Regex.IsMatch(includePath, @"\*|\?|\.\.\."))
+					string? includePath = CombinePaths(basePath, script);
+					if (includePath == null)
+					{
+						LogError(element, $"Path '{script}' cannot be combined with '{basePath}'");
+					}
+					else if (Regex.IsMatch(includePath, @"\*|\?|\.\.\."))
 					{
 						files.UnionWith(FindMatchingFiles(_rootDir, includePath));
 					}
@@ -697,7 +701,7 @@ namespace AutomationTool
 		/// <summary>
 		/// Combine two paths without validating the result
 		/// </summary>
-		static string CombinePaths(string basePath, string nextPath)
+		static string? CombinePaths(string basePath, string nextPath)
 		{
 			if (Path.IsPathRooted(nextPath))
 			{
@@ -721,7 +725,7 @@ namespace AutomationTool
 					}
 					else
 					{
-						throw new Exception($"Path '{nextPath}' cannot be combined with '{basePath}'");
+						return null;
 					}
 				}
 				else
@@ -757,7 +761,8 @@ namespace AutomationTool
 					// Make sure we're at global scope
 					if (ScopedProperties.Count > 1)
 					{
-						throw new Exception("Incorrect scope depth for reading option settings");
+						LogError(element, "Incorrect scope depth for reading option settings");
+						return;
 					}
 
 					// Check if the property already exists. If it does, we don't need to register it as an option.
@@ -971,25 +976,33 @@ namespace AutomationTool
 					case "Replace":
 						if (arguments.Length != 2)
 						{
-							throw new AutomationException($"String operation 'Replace' requires exactly 2 arguments.");
+							LogError(element, $"String operation 'Replace' requires exactly 2 arguments.");
+							return;
 						}
+
 						operationResult = input.Replace(arguments[0], arguments[1], StringComparison.Ordinal);
 						break;
 					case "SplitFirst":
 						if (arguments.Length != 1)
 						{
-							throw new AutomationException($"String operation 'SplitFirst' requires exactly 1 argument.");
+							LogError(element, $"String operation 'SplitFirst' requires exactly 1 argument.");
+							return;
 						}
+
 						operationResult = input.Split(arguments[0]).First();
 						break;
 					case "SplitLast":
 						if (arguments.Length != 1)
 						{
-							throw new AutomationException($"String operation 'SplitLast' requires exactly 1 argument.");
+							LogError(element, $"String operation 'SplitLast' requires exactly 1 argument.");
+							return;
 						}
+	
 						operationResult = input.Split(arguments[0]).Last();
 						break;
-					default: throw new AutomationException($"String operation '{method}' not available.");
+					default:
+						LogError(element, $"String operation '{method}' not available.");
+						return;
 				}
 				SetPropertyValue(element, output, operationResult);
 			}
