BuildGraph: Remove some references to UnrealBuildBase.

#rnx

diff --git a/Engine/Source/Programs/AutomationTool/BuildGraph/BgCondition.cs b/Engine/Source/Programs/AutomationTool/BuildGraph/BgCondition.cs
--- a/Engine/Source/Programs/AutomationTool/BuildGraph/BgCondition.cs
+++ b/Engine/Source/Programs/AutomationTool/BuildGraph/BgCondition.cs
@@ -6,7 +6,6 @@ using System.IO;
 using System.Linq;
 using System.Threading.Tasks;
 using EpicGames.Core;
-using UnrealBuildBase;
 
 #nullable enable
 
@@ -27,6 +26,12 @@ namespace AutomationTool
 		}
 	}
 
+	/// <summary>
+	/// Context for evaluating BuildGraph conditions
+	/// </summary>
+	/// <param name="RootDir">Root directory for resolving paths</param>
+	public record class BgConditionContext(DirectoryReference RootDir);
+
 	/// <summary>
 	/// Class to evaluate condition expressions in build scripts, following this grammar:
 	/// 
@@ -87,24 +92,26 @@ namespace AutomationTool
 		/// Evaluates the given string as a condition. Throws a ConditionException on a type or syntax error.
 		/// </summary>
 		/// <param name="text"></param>
+		/// <param name="context">Context for evaluating the condition</param>
 		/// <returns>The result of evaluating the condition</returns>
-		public static ValueTask<bool> Evaluate(string text)
+		public static ValueTask<bool> EvaluateAsync(string text, BgConditionContext context)
 		{
-			return new BgCondition(text).EvaluateAsync();
+			return new BgCondition(text).EvaluateAsync(context);
 		}
 
 		/// <summary>
 		/// Evaluates the given string as a condition. Throws a ConditionException on a type or syntax error.
 		/// </summary>
+		/// <param name="context">Context for evaluating the condition</param>
 		/// <returns>The result of evaluating the condition</returns>
-		async ValueTask<bool> EvaluateAsync()
+		async ValueTask<bool> EvaluateAsync(BgConditionContext context)
 		{
 			bool result = true;
 			if (_tokens.Count > 1)
 			{
 				_idx = 0;
 
-				string value = await EvaluateOrAsync();
+				string value = await EvaluateOrAsync(context);
 				if (_tokens[_idx] != EndToken)
 				{
 					throw new BgConditionException("Garbage after expression: {0}", String.Join("", _tokens.Skip(_idx)));
@@ -118,17 +125,18 @@ namespace AutomationTool
 		/// <summary>
 		/// Evaluates an "or-expression" production.
 		/// </summary>
+		/// <param name="context">Context for evaluating the expression</param>
 		/// <returns>A scalar representing the result of evaluating the expression.</returns>
-		async ValueTask<string> EvaluateOrAsync()
+		async ValueTask<string> EvaluateOrAsync(BgConditionContext context)
 		{
 			// <Condition> Or <Condition> Or...
-			string result = await EvaluateAndAsync();
+			string result = await EvaluateAndAsync(context);
 			while (String.Equals(_tokens[_idx], "Or", StringComparison.OrdinalIgnoreCase))
 			{
 				// Evaluate this condition. We use a binary OR here, because we want to parse everything rather than short-circuit it.
 				_idx++;
 				string lhs = result;
-				string rhs = await EvaluateAndAsync();
+				string rhs = await EvaluateAndAsync(context);
 				result = (CoerceToBool(lhs) | CoerceToBool(rhs)) ? "true" : "false";
 			}
 			return result;
@@ -137,17 +145,18 @@ namespace AutomationTool
 		/// <summary>
 		/// Evaluates an "and-expression" production.
 		/// </summary>
+		/// <param name="context">Context for evaluating the expression</param>
 		/// <returns>A scalar representing the result of evaluating the expression.</returns>
-		async ValueTask<string> EvaluateAndAsync()
+		async ValueTask<string> EvaluateAndAsync(BgConditionContext context)
 		{
 			// <Condition> And <Condition> And...
-			string result = await EvaluateComparisonAsync();
+			string result = await EvaluateComparisonAsync(context);
 			while (String.Equals(_tokens[_idx], "And", StringComparison.OrdinalIgnoreCase))
 			{
 				// Evaluate this condition. We use a binary AND here, because we want to parse everything rather than short-circuit it.
 				_idx++;
 				string lhs = result;
-				string rhs = await EvaluateComparisonAsync();
+				string rhs = await EvaluateComparisonAsync(context);
 				result = (CoerceToBool(lhs) & CoerceToBool(rhs)) ? "true" : "false";
 			}
 			return result;
@@ -156,8 +165,9 @@ namespace AutomationTool
 		/// <summary>
 		/// Evaluates a "comparison" production.
 		/// </summary>
+		/// <param name="context">Context for evaluating the expression</param>
 		/// <returns>The result of evaluating the expression</returns>
-		async ValueTask<string> EvaluateComparisonAsync()
+		async ValueTask<string> EvaluateComparisonAsync(BgConditionContext context)
 		{
 			// scalar
 			// scalar == scalar
@@ -167,13 +177,13 @@ namespace AutomationTool
 			// scalar > scalar
 			// scalar >= scalar
 
-			string result = await EvaluateScalarAsync();
+			string result = await EvaluateScalarAsync(context);
 			if (_tokens[_idx] == "==")
 			{
 				// Compare two scalars for equality
 				_idx++;
 				string lhs = result;
-				string rhs = await EvaluateScalarAsync();
+				string rhs = await EvaluateScalarAsync(context);
 				result = String.Equals(lhs, rhs, StringComparison.OrdinalIgnoreCase) ? "true" : "false";
 			}
 			else if (_tokens[_idx] == "!=")
@@ -181,7 +191,7 @@ namespace AutomationTool
 				// Compare two scalars for inequality
 				_idx++;
 				string lhs = result;
-				string rhs = await EvaluateScalarAsync();
+				string rhs = await EvaluateScalarAsync(context);
 				result = String.Equals(lhs, rhs, StringComparison.OrdinalIgnoreCase) ? "false" : "true";
 			}
 			else if (_tokens[_idx] == "<")
@@ -189,7 +199,7 @@ namespace AutomationTool
 				// Compares whether the first integer is less than the second
 				_idx++;
 				int lhs = CoerceToInteger(result);
-				int rhs = CoerceToInteger(await EvaluateScalarAsync());
+				int rhs = CoerceToInteger(await EvaluateScalarAsync(context));
 				result = (lhs < rhs) ? "true" : "false";
 			}
 			else if (_tokens[_idx] == "<=")
@@ -197,7 +207,7 @@ namespace AutomationTool
 				// Compares whether the first integer is less than the second
 				_idx++;
 				int lhs = CoerceToInteger(result);
-				int rhs = CoerceToInteger(await EvaluateScalarAsync());
+				int rhs = CoerceToInteger(await EvaluateScalarAsync(context));
 				result = (lhs <= rhs) ? "true" : "false";
 			}
 			else if (_tokens[_idx] == ">")
@@ -205,7 +215,7 @@ namespace AutomationTool
 				// Compares whether the first integer is less than the second
 				_idx++;
 				int lhs = CoerceToInteger(result);
-				int rhs = CoerceToInteger(await EvaluateScalarAsync());
+				int rhs = CoerceToInteger(await EvaluateScalarAsync(context));
 				result = (lhs > rhs) ? "true" : "false";
 			}
 			else if (_tokens[_idx] == ">=")
@@ -213,7 +223,7 @@ namespace AutomationTool
 				// Compares whether the first integer is less than the second
 				_idx++;
 				int lhs = CoerceToInteger(result);
-				int rhs = CoerceToInteger(await EvaluateScalarAsync());
+				int rhs = CoerceToInteger(await EvaluateScalarAsync(context));
 				result = (lhs >= rhs) ? "true" : "false";
 			}
 			return result;
@@ -270,15 +280,16 @@ namespace AutomationTool
 		/// <summary>
 		/// Evaluates a "scalar" production.
 		/// </summary>
+		/// <param name="context">Context for evaluating the expression</param>
 		/// <returns>The result of evaluating the expression</returns>
-		async ValueTask<string> EvaluateScalarAsync()
+		async ValueTask<string> EvaluateScalarAsync(BgConditionContext context)
 		{
 			string result;
 			if (_tokens[_idx] == "(")
 			{
 				// Subexpression
 				_idx++;
-				result = await EvaluateOrAsync();
+				result = await EvaluateOrAsync(context);
 				if (_tokens[_idx] != ")")
 				{
 					throw new BgConditionException("Expected ')'");
@@ -289,21 +300,21 @@ namespace AutomationTool
 			{
 				// Logical not
 				_idx++;
-				string rhs = await EvaluateScalarAsync();
+				string rhs = await EvaluateScalarAsync(context);
 				result = CoerceToBool(rhs) ? "false" : "true";
 			}
 			else if (String.Equals(_tokens[_idx], "Exists", StringComparison.OrdinalIgnoreCase) && _tokens[_idx + 1] == "(")
 			{
 				// Check whether file or directory exists. Evaluate the argument as a subexpression.
 				_idx++;
-				string argument = await EvaluateScalarAsync();
-				result = Exists(argument) ? "true" : "false";
+				string argument = await EvaluateScalarAsync(context);
+				result = Exists(argument, context) ? "true" : "false";
 			}
 			else if (String.Equals(_tokens[_idx], "HasTrailingSlash", StringComparison.OrdinalIgnoreCase) && _tokens[_idx + 1] == "(")
 			{
 				// Check whether the given string ends with a slash
 				_idx++;
-				string argument = await EvaluateScalarAsync();
+				string argument = await EvaluateScalarAsync(context);
 				result = (argument.Length > 0 && (argument[^1] == Path.DirectorySeparatorChar || argument[^1] == Path.AltDirectorySeparatorChar)) ? "true" : "false";
 			}
 			else if (String.Equals(_tokens[_idx], "Contains", StringComparison.OrdinalIgnoreCase) && _tokens[_idx + 1] == "(")
@@ -357,13 +368,11 @@ namespace AutomationTool
 		/// <summary>
 		/// Determine if a path exists
 		/// </summary>
-		/// <param name="path"></param>
-		/// <returns></returns>
-		static bool Exists(string path)
+		static bool Exists(string path, BgConditionContext context)
 		{
 			try
 			{
-				return FileReference.Exists(FileReference.Combine(Unreal.RootDirectory, path)) || DirectoryReference.Exists(DirectoryReference.Combine(Unreal.RootDirectory, path));
+				return FileReference.Exists(FileReference.Combine(context.RootDir, path)) || DirectoryReference.Exists(DirectoryReference.Combine(context.RootDir, path));
 			}
 			catch
 			{
@@ -534,7 +543,7 @@ namespace AutomationTool
 		/// <param name="expectedResult">The expected result</param>
 		static async Task TestConditionAsync(string condition, bool expectedResult)
 		{
-			bool result = await new BgCondition(condition).EvaluateAsync();
+			bool result = await new BgCondition(condition).EvaluateAsync(new BgConditionContext(DirectoryReference.GetCurrentDirectory()));
 			Console.WriteLine("{0}: {1} = {2}", (result == expectedResult) ? "PASS" : "FAIL", condition, result);
 		}
 	}
diff --git a/Engine/Source/Programs/AutomationTool/BuildGraph/BgNodeExecutor.cs b/Engine/Source/Programs/AutomationTool/BuildGraph/BgNodeExecutor.cs
--- a/Engine/Source/Programs/AutomationTool/BuildGraph/BgNodeExecutor.cs
+++ b/Engine/Source/Programs/AutomationTool/BuildGraph/BgNodeExecutor.cs
@@ -163,12 +163,12 @@ namespace AutomationTool
 			Node = node;
 		}
 
-		public bool Bind(Dictionary<string, ScriptTaskBinding> nameToTask, Dictionary<string, BgNodeOutput> tagNameToNodeOutput, ILogger logger)
+		public async ValueTask<bool> BindAsync(Dictionary<string, ScriptTaskBinding> nameToTask, Dictionary<string, BgNodeOutput> tagNameToNodeOutput, ILogger logger)
 		{
 			bool result = true;
 			foreach (BgTask taskInfo in Node.Tasks)
 			{
-				BgTaskImpl? boundTask = BindTask(taskInfo, nameToTask, tagNameToNodeOutput, logger);
+				BgTaskImpl? boundTask = await BindTaskAsync(taskInfo, nameToTask, tagNameToNodeOutput, logger);
 				if (boundTask == null)
 				{
 					result = false;
@@ -181,7 +181,7 @@ namespace AutomationTool
 			return result;
 		}
 
-		BgTaskImpl? BindTask(BgTask taskInfo, Dictionary<string, ScriptTaskBinding> nameToTask, IReadOnlyDictionary<string, BgNodeOutput> tagNameToNodeOutput, ILogger logger)
+		async ValueTask<BgTaskImpl?> BindTaskAsync(BgTask taskInfo, Dictionary<string, ScriptTaskBinding> nameToTask, IReadOnlyDictionary<string, BgNodeOutput> tagNameToNodeOutput, ILogger logger)
 		{
 			// Get the reflection info for this element
 			ScriptTaskBinding? task;
@@ -202,6 +202,9 @@ namespace AutomationTool
 				}
 			}
 
+			// Create a context for evaluating conditions
+			BgConditionContext conditionContext = new BgConditionContext(Unreal.RootDirectory);
+
 			// Read all the attributes into a parameters object for this task
 			object parametersObject = Activator.CreateInstance(task.ParametersClass)!;
 			foreach ((string name, string value) in taskInfo.Arguments)
@@ -220,7 +223,7 @@ namespace AutomationTool
 					if (parameter.CollectionType == null)
 					{
 						// Parse it and assign it to the parameters object
-						object? fieldValue = ParseValue(value, parameter.ValueType);
+						object? fieldValue = await ParseValueAsync(value, parameter.ValueType, conditionContext);
 						if (fieldValue != null)
 						{
 							parameter.SetValue(parametersObject, fieldValue);
@@ -244,7 +247,7 @@ namespace AutomationTool
 						List<string> valueStrings = BgTaskImpl.SplitDelimitedList(value);
 						foreach (string valueString in valueStrings)
 						{
-							object? elementValue = ParseValue(valueString, parameter.ValueType);
+							object? elementValue = await ParseValueAsync(valueString, parameter.ValueType, conditionContext);
 							if (elementValue != null)
 							{
 								parameter.CollectionType.InvokeMember("Add", BindingFlags.InvokeMethod | BindingFlags.Instance | BindingFlags.Public, null, collectionValue, new object[] { elementValue });
@@ -304,8 +307,9 @@ namespace AutomationTool
 		/// </summary>
 		/// <param name="valueText">The text to parse</param>
 		/// <param name="valueType">Type of the value to parse</param>
+		/// <param name="context">Context for evaluating boolean expressions</param>
 		/// <returns>Value that was parsed</returns>
-		static object? ParseValue(string valueText, Type valueType)
+		static async ValueTask<object?> ParseValueAsync(string valueText, Type valueType, BgConditionContext context)
 		{
 			// Parse it and assign it to the parameters object
 			if (valueType.IsEnum)
@@ -314,7 +318,7 @@ namespace AutomationTool
 			}
 			else if (valueType == typeof(Boolean))
 			{
-				return BgCondition.Evaluate(valueText).AsTask().Result;
+				return await BgCondition.EvaluateAsync(valueText, context);
 			}
 			else if (valueType == typeof(FileReference))
 			{
diff --git a/Engine/Source/Programs/AutomationTool/BuildGraph/BgScriptReader.cs b/Engine/Source/Programs/AutomationTool/BuildGraph/BgScriptReader.cs
--- a/Engine/Source/Programs/AutomationTool/BuildGraph/BgScriptReader.cs
+++ b/Engine/Source/Programs/AutomationTool/BuildGraph/BgScriptReader.cs
@@ -921,7 +921,7 @@ namespace AutomationTool
 						string input = ReadAttribute(element, "Input");
 						Match match = regexValue.Match(input);
 
-						bool optional = await BgCondition.Evaluate(ReadAttribute(element, "Optional"));
+						bool optional = await BgCondition.EvaluateAsync(ReadAttribute(element, "Optional"), new BgConditionContext(_rootDir));
 						if (!match.Success)
 						{
 							if (!optional)
@@ -2273,7 +2273,7 @@ namespace AutomationTool
 			try
 			{
 				string text = ExpandProperties(element, element.GetAttribute("If"));
-				return await BgCondition.Evaluate(text);
+				return await BgCondition.EvaluateAsync(text, new BgConditionContext(_rootDir));
 			}
 			catch (BgConditionException ex)
 			{
diff --git a/Engine/Source/Programs/AutomationTool/BuildGraph/BuildGraph.cs b/Engine/Source/Programs/AutomationTool/BuildGraph/BuildGraph.cs
--- a/Engine/Source/Programs/AutomationTool/BuildGraph/BuildGraph.cs
+++ b/Engine/Source/Programs/AutomationTool/BuildGraph/BuildGraph.cs
@@ -727,14 +727,14 @@ namespace AutomationTool
 			Dictionary<BgNodeDef, BgNodeExecutor> nodeToExecutor = new Dictionary<BgNodeDef, BgNodeExecutor>();
 			if (skipValidation && singleNode != null)
 			{
-				if (!BindNodes(singleNode, nameToTask, graph.TagNameToNodeOutput, nodeToExecutor))
+				if (!await BindNodesAsync(singleNode, nameToTask, graph.TagNameToNodeOutput, nodeToExecutor))
 				{
 					return ExitCode.Error_Unknown;
 				}
 			}
 			else
 			{
-				if (!BindNodes(graph, nameToTask, nodeToExecutor))
+				if (!await BindNodesAsync(graph, nameToTask, nodeToExecutor))
 				{
 					return ExitCode.Error_Unknown;
 				}
@@ -761,26 +761,26 @@ namespace AutomationTool
 			return ExitCode.Success;
 		}
 
-		static bool BindNodes(BgGraphDef graph, Dictionary<string, ScriptTaskBinding> nameToTask, Dictionary<BgNodeDef, BgNodeExecutor> nodeToExecutor)
+		static async ValueTask<bool> BindNodesAsync(BgGraphDef graph, Dictionary<string, ScriptTaskBinding> nameToTask, Dictionary<BgNodeDef, BgNodeExecutor> nodeToExecutor)
 		{
 			bool result = true;
 			foreach (BgAgentDef agent in graph.Agents)
 			{
 				foreach (BgNodeDef node in agent.Nodes)
 				{
-					result &= BindNodes(node, nameToTask, graph.TagNameToNodeOutput, nodeToExecutor);
+					result &= await BindNodesAsync(node, nameToTask, graph.TagNameToNodeOutput, nodeToExecutor);
 				}
 			}
 			return result;
 		}
 
-		static bool BindNodes(BgNodeDef node, Dictionary<string, ScriptTaskBinding> nameToTask, Dictionary<string, BgNodeOutput> tagNameToNodeOutput, Dictionary<BgNodeDef, BgNodeExecutor> nodeToExecutor)
+		static async ValueTask<bool> BindNodesAsync(BgNodeDef node, Dictionary<string, ScriptTaskBinding> nameToTask, Dictionary<string, BgNodeOutput> tagNameToNodeOutput, Dictionary<BgNodeDef, BgNodeExecutor> nodeToExecutor)
 		{
 			if (node is BgScriptNode scriptNode)
 			{
 				BgScriptNodeExecutor executor = new BgScriptNodeExecutor(scriptNode);
 				nodeToExecutor[node] = executor;
-				return executor.Bind(nameToTask, tagNameToNodeOutput, Logger);
+				return await executor.BindAsync(nameToTask, tagNameToNodeOutput, Logger);
 			}
 			else if (node is BgNodeExpressionDef bytecodeNode)
 			{
