[Gauntlet] point Demos directory to user directory instead of build directory, since Replay feature search in user directory for assets.

#rnx

diff --git a/Engine/Source/Programs/AutomationTool/Gauntlet/Framework/Base/Gauntlet.TargetDeviceDesktopCommon.cs b/Engine/Source/Programs/AutomationTool/Gauntlet/Framework/Base/Gauntlet.TargetDeviceDesktopCommon.cs
--- a/Engine/Source/Programs/AutomationTool/Gauntlet/Framework/Base/Gauntlet.TargetDeviceDesktopCommon.cs
+++ b/Engine/Source/Programs/AutomationTool/Gauntlet/Framework/Base/Gauntlet.TargetDeviceDesktopCommon.cs
@@ -215,11 +215,12 @@ namespace Gauntlet
 			LocalDirectoryMappings.Add(EIntendedBaseCopyDirectory.Build, Path.Combine(BaseDirectory, "Build"));
 			LocalDirectoryMappings.Add(EIntendedBaseCopyDirectory.Binaries, Path.Combine(BaseDirectory, "Binaries"));
 			LocalDirectoryMappings.Add(EIntendedBaseCopyDirectory.Config, Path.Combine(BaseDirectory, "Saved", "Config"));
-            LocalDirectoryMappings.Add(EIntendedBaseCopyDirectory.Content, Path.Combine(BaseDirectory, "Content"));
-            LocalDirectoryMappings.Add(EIntendedBaseCopyDirectory.Demos, Path.Combine(BaseDirectory, "Saved", "Demos"));
+			LocalDirectoryMappings.Add(EIntendedBaseCopyDirectory.Content, Path.Combine(BaseDirectory, "Content"));
 			LocalDirectoryMappings.Add(EIntendedBaseCopyDirectory.PersistentDownloadDir, Path.Combine(BaseDirectory, "Saved", "PersistentDownloadDir"));
 			LocalDirectoryMappings.Add(EIntendedBaseCopyDirectory.Profiling, Path.Combine(BaseDirectory, "Saved", "Profiling"));
-            LocalDirectoryMappings.Add(EIntendedBaseCopyDirectory.Saved, Path.Combine(BaseDirectory, "Saved"));
+			LocalDirectoryMappings.Add(EIntendedBaseCopyDirectory.Saved, Path.Combine(BaseDirectory, "Saved"));
+			// Folders that are located in User dir instead of build dir
+			LocalDirectoryMappings.Add(EIntendedBaseCopyDirectory.Demos, Path.Combine(UserDir, "Saved", "Demos"));
 		}
 
 		public abstract IAppInstance Run(IAppInstall Install);
