[Gauntlet] Highlight when crash dump is generated in test report.
Generate a hang analysis of the callstack.

#rnx

diff --git a/Engine/Source/Programs/AutomationTool/Gauntlet/Platform/Windows/Gauntlet.TargetDeviceWindows.cs b/Engine/Source/Programs/AutomationTool/Gauntlet/Platform/Windows/Gauntlet.TargetDeviceWindows.cs
--- a/Engine/Source/Programs/AutomationTool/Gauntlet/Platform/Windows/Gauntlet.TargetDeviceWindows.cs
+++ b/Engine/Source/Programs/AutomationTool/Gauntlet/Platform/Windows/Gauntlet.TargetDeviceWindows.cs
@@ -9,6 +9,9 @@ using static AutomationTool.ProcessResult;
 using EpicGames.Core;
 using System.ComponentModel;
 using System.Runtime.InteropServices;
+using System.Text.RegularExpressions;
+using System.Collections.Generic;
+using System.Linq;
 
 namespace Gauntlet
 {
@@ -384,9 +387,62 @@ namespace Gauntlet
 			}
 			else
 			{
-				Log.Info("Wrote minidump to {FileName}", DumpName);
+				Log.Warning(KnownLogEvents.Gauntlet, "Wrote minidump to {FileName}", DumpName);
+
+				if (CommandUtils.IsBuildMachine)
+				{
+					// Produce a stack analysis on buildmachine.
+					FileReference CDB = GetCDBExe();
+					if (CDB != null)
+					{
+						string Args = $"-z \"{DumpName}\" -c \"!analyze -v -hang; q\"";
+						string Output = UnrealBuildTool.Utils.RunLocalProcessAndReturnStdOut(CDB.FullName, Args);
+						List<string> Stack = new List<string>();
+						Match M = Regex.Match(Output, @"STACK_TEXT:.*", RegexOptions.IgnoreCase);
+						if (M.Success)
+						{
+							Output = Output.Substring(M.Index + M.Length + 1);
+							foreach (string Line in Output.Split('\n').Select(L => L.Trim()))
+							{
+								if (string.IsNullOrEmpty(Line)) break;
+								// A callstack line looks like this from cdb:
+								//   000000ef`8877ad20 00007ff8`71c18cde     : 00000523`3b3cd580 00007722`b1765a01 00007722`b1765a01 00000216`6e3bfc00 : UnrealEditor_Core!LowLevelTasks::Private::FWaitingQueue::PrepareWait+0x89a
+								// We are removing the first 115 characters from the each line as they are not going to add meaningful information to the report. 
+								Stack.Add(Line.Substring(Line.Length > 115 ? 115 : 0));
+							}
+						}
+						if (Stack.Any())
+						{
+							Log.Warning(KnownLogEvents.Gauntlet, "Stack analysis produced with 'cdb {Args}':\n{Callstack}", Args, string.Join('\n', Stack));
+						}
+						else
+						{
+							Log.Warning("Could not find stack analysis from cdb command.");
+						}
+					}
+				}
 			}
 		}
+
+		private static FileReference GetCDBExe()
+		{
+			// Trying to look for auto sdk latest WindowsKits debugger tools
+			DirectoryReference HostAutoSdkDir = null;
+			if (UEBuildPlatformSDK.TryGetHostPlatformAutoSDKDir(out HostAutoSdkDir))
+			{
+				DirectoryReference WindowsKitsDebuggersDirAutoSdk = DirectoryReference.Combine(HostAutoSdkDir, "Win64", "Windows Kits", "Debuggers");
+				if (DirectoryReference.Exists(WindowsKitsDebuggersDirAutoSdk))
+				{
+					FileReference CDBExe64 = FileReference.Combine(WindowsKitsDebuggersDirAutoSdk, "x64", "cdb.exe");
+					if (FileReference.Exists(CDBExe64))
+					{
+						return CDBExe64;
+					}
+				}
+			}
+
+			return null;
+		}
 	}
 
 	public class Win64DeviceFactory : IDeviceFactory
