Check AllowSpew before putting process outputs into stderr

When launching a process via UAT, you can control whether it logs its output, or whether it is just put into the output. This flag used to be respected for both stdout and stderr output, but it got lost during a merge and subsequent integration for stderr. This CL reinstates it in this codepath: Not doing this is super annoying because the stderr and stdout can otherwise overlap in really annoying ways (stderr can just come in from a background thread and interfere with whatever you write into stdout).

#nojira

#rnx

diff --git a/Engine/Source/Programs/AutomationTool/AutomationUtils/ProcessUtils.cs b/Engine/Source/Programs/AutomationTool/AutomationUtils/ProcessUtils.cs
--- a/Engine/Source/Programs/AutomationTool/AutomationUtils/ProcessUtils.cs
+++ b/Engine/Source/Programs/AutomationTool/AutomationUtils/ProcessUtils.cs
@@ -358,18 +358,22 @@ namespace AutomationTool
 		{
 			if (e.Data != null)
 			{
-				if (SpewFilterCallback != null)
+				if (AllowSpew)
 				{
-					string FilteredSpew = SpewFilterCallback(e.Data);
-					if (FilteredSpew != null)
+					if (SpewFilterCallback != null)
 					{
-						LogOutput(SpewVerbosity, FilteredSpew);
+						string FilteredSpew = SpewFilterCallback(e.Data);
+						if (FilteredSpew != null)
+						{
+							LogOutput(SpewVerbosity, FilteredSpew);
+						}
+					}
+					else
+					{
+						LogOutput(SpewVerbosity, e.Data);
 					}
 				}
-				else
-				{
-					LogOutput(SpewVerbosity, e.Data);
-				}
+
 				if(ProcessOutput != null)
 				{
 					lock (ProcSyncObject)
