[Gauntlet] Corrects error in Assigning devices to roles that would assign roles to a device that was already assigned
#tests local run

diff --git a/Engine/Source/Programs/AutomationTool/Gauntlet/Unreal/Base/Gauntlet.UnrealSession.cs b/Engine/Source/Programs/AutomationTool/Gauntlet/Unreal/Base/Gauntlet.UnrealSession.cs
--- a/Engine/Source/Programs/AutomationTool/Gauntlet/Unreal/Base/Gauntlet.UnrealSession.cs
+++ b/Engine/Source/Programs/AutomationTool/Gauntlet/Unreal/Base/Gauntlet.UnrealSession.cs
@@ -1547,6 +1547,7 @@ namespace Gauntlet
 		{
 			// Order by constraint. This ensures roles with constraints have their devices selected first.
 			IEnumerable<UnrealSessionRole> RolesSortedByConstraint = SessionRoles.OrderBy(R => R.Constraint.IsIdentity() ? 1 : 0);
+			IEnumerable<ITargetDevice> ReservedDevices = UnrealDeviceReservation.ReservedDevices;
 
 			foreach (UnrealSessionRole Role in RolesSortedByConstraint)
 			{
@@ -1562,7 +1563,7 @@ namespace Gauntlet
 					{
 						try
 						{
-							DeviceToAssign = UnrealDeviceReservation.ReservedDevices.Where(Device => DeviceMatchesRoleConstraint(Role, Device)).First();
+							DeviceToAssign = ReservedDevices.Where(Device => DeviceMatchesRoleConstraint(Role, Device)).First();
 							IDeviceUsageReporter.RecordStart(DeviceToAssign.Name, DeviceToAssign.Platform, IDeviceUsageReporter.EventType.Device, IDeviceUsageReporter.EventState.Success);
 						}
 						catch (Exception Ex)
@@ -1574,6 +1575,7 @@ namespace Gauntlet
 					}
 
 					RolesToDevices.Add(Role, DeviceToAssign);
+					ReservedDevices = ReservedDevices.Except(Enumerable.Repeat(DeviceToAssign, 1));
 				}
 			}
 
