Only call p4 client for clients where the root path is a match for the file in question
#rnx

diff --git a/Engine/Source/Programs/AutomationTool/AutomationUtils/P4Utils.cs b/Engine/Source/Programs/AutomationTool/AutomationUtils/P4Utils.cs
--- a/Engine/Source/Programs/AutomationTool/AutomationUtils/P4Utils.cs
+++ b/Engine/Source/Programs/AutomationTool/AutomationUtils/P4Utils.cs
@@ -3626,37 +3626,33 @@ namespace AutomationTool
 				P4Command += " -S " + AllowedStream;
 			}
 
-			var P4Result = P4(P4Command, AllowSpew: false, WithClient: false);
-			if (P4Result.ExitCode != 0)
+			string Output;
+			if (!P4Output(out Output, "-ztag -F \"Client %client% Root %Root%\"", P4Command, AllowSpew: false, WithClient: false))
 			{
 				throw new AutomationException("p4 clients -u {0} failed.", UserName);
 			}
 
 			// Parse output.
-			var Lines = P4Result.Output.Split(new string[] { Environment.NewLine }, StringSplitOptions.RemoveEmptyEntries);
+			Regex OutputSplitter = new Regex(@"Client (?<client>.+) Root (?<root>.*)");
+			var Lines = Output.Split(new string[] { Environment.NewLine }, StringSplitOptions.RemoveEmptyEntries);
 			foreach (string Line in Lines)
 			{
-				var Tokens = Line.Split(new char[] { ' ' }, StringSplitOptions.RemoveEmptyEntries);
-				P4ClientInfo Info = null;
-
-				// Retrieve the client name and info.
-				for (int TokenIndex = 0; TokenIndex < Tokens.Length; ++TokenIndex)
-				{
-					if (Tokens[TokenIndex] == "Client")
-					{
-						var ClientName = Tokens[++TokenIndex];
-						Info = GetClientInfoInternal(ClientName);
-						break;
-					}
-				}
+				Match RegexMatch = OutputSplitter.Match(Line);
+				string ClientName = RegexMatch.Groups["client"].Value;
+				string RootPath = RegexMatch.Groups["root"].Value;
 
-				if (Info == null || String.IsNullOrEmpty(Info.Name) || String.IsNullOrEmpty(Info.RootPath))
+				if (String.IsNullOrEmpty(ClientName) || String.IsNullOrEmpty(RootPath))
 				{
 					throw new AutomationException("Failed to retrieve p4 client info for user {0}. Unable to set up local environment", UserName);
 				}
-				
-				if (IsValidClientForFile(Info, PathUnderClientRoot))
+
+				if (IsValidRootPathForFile(RootPath, PathUnderClientRoot))
 				{
+					P4ClientInfo Info = GetClientInfoInternal(ClientName);
+					if (Info == null)
+					{
+						throw new AutomationException("Failed to retrieve p4 client info for user {0}. Unable to set up local environment", UserName);
+					}
 					ClientList.Add(Info);
 				}
 			}
@@ -3664,12 +3660,17 @@ namespace AutomationTool
 		}
 
 		public bool IsValidClientForFile(P4ClientInfo Info, string PathUnderClientRoot)
+		{
+			return IsValidRootPathForFile(Info.RootPath, PathUnderClientRoot);
+		}
+
+		public bool IsValidRootPathForFile(string RootPath, string PathUnderClientRoot)
 		{
 			// Filter the client out if the specified path is not under the client root
 			bool bAddClient = true;
-			if (!String.IsNullOrEmpty(PathUnderClientRoot) && !String.IsNullOrEmpty(Info.RootPath))
+			if (!String.IsNullOrEmpty(PathUnderClientRoot) && !String.IsNullOrEmpty(RootPath))
 			{
-				var ClientRootPathWithSlash = Info.RootPath;
+				var ClientRootPathWithSlash = RootPath;
 				if (!ClientRootPathWithSlash.EndsWith("\\") && !ClientRootPathWithSlash.EndsWith("/"))
 				{
 					ClientRootPathWithSlash = CommandUtils.ConvertSeparators(PathSeparator.Default, ClientRootPathWithSlash + "/");
