[Automation] Adds full automation support to ReplayRun.py for iOS when running with a mac host

diff --git a/Engine/Source/Programs/AutomationTool/Gauntlet/Framework/Utils/Gauntlet.ReportGenUtils.cs b/Engine/Source/Programs/AutomationTool/Gauntlet/Framework/Utils/Gauntlet.ReportGenUtils.cs
--- a/Engine/Source/Programs/AutomationTool/Gauntlet/Framework/Utils/Gauntlet.ReportGenUtils.cs
+++ b/Engine/Source/Programs/AutomationTool/Gauntlet/Framework/Utils/Gauntlet.ReportGenUtils.cs
@@ -6,15 +6,28 @@ using System.Linq;
 using System.Diagnostics;
 using System.IO;
 using AutomationTool;
+using UnrealBuildTool;
 
 namespace Gauntlet
 {
 	public static class ReportGenUtils
 	{
 		private static string PythonExecutable = null;
-		private const string BasePythonLocation = @"Engine\Binaries\ThirdParty\Python3\Win64\python.exe";
+		private static string BasePythonLocation;
 		public const string CsvBinExt = ".csv.bin";
 
+		static ReportGenUtils()
+		{
+			if (BuildHostPlatform.Current.Platform == UnrealTargetPlatform.Win64)
+			{
+				BasePythonLocation = @"Engine\Binaries\ThirdParty\Python3\Win64\python.exe";
+			}
+			else if (BuildHostPlatform.Current.Platform == UnrealTargetPlatform.Mac)
+			{
+				BasePythonLocation = @"Engine/Binaries/ThirdParty/Python3/Mac/bin/python3";
+			}
+		}
+
 		public static int RunPerfReportTool(string Args, string ReportXmlBaseDir)
 		{
 			string PerfReportToolPath = GetCsvToolPath("PerfReportTool");
