[Gauntlet] Redo 35342435 as manual merge

diff --git a/Engine/Source/Programs/AutomationTool/Gauntlet/Framework/Utils/Gauntlet.ReportGenUtils.cs b/Engine/Source/Programs/AutomationTool/Gauntlet/Framework/Utils/Gauntlet.ReportGenUtils.cs
--- a/Engine/Source/Programs/AutomationTool/Gauntlet/Framework/Utils/Gauntlet.ReportGenUtils.cs
+++ b/Engine/Source/Programs/AutomationTool/Gauntlet/Framework/Utils/Gauntlet.ReportGenUtils.cs
@@ -14,18 +14,35 @@ namespace Gauntlet
 	{
 		private static string PythonExecutable = null;
 		private static string BasePythonLocation;
+		private static string DotnetPathLocation;
 		public const string CsvBinExt = ".csv.bin";
 
 		static ReportGenUtils()
 		{
-			if (BuildHostPlatform.Current.Platform == UnrealTargetPlatform.Win64)
+			char PathDelimter;
+			if (BuildHostPlatform.Current.Platform.IsInGroup("Windows"))
 			{
+				PathDelimter = ';';
+				DotnetPathLocation = "dotnet.exe";
 				BasePythonLocation = @"Engine\Binaries\ThirdParty\Python3\Win64\python.exe";
 			}
-			else if (BuildHostPlatform.Current.Platform == UnrealTargetPlatform.Mac)
+			else
 			{
+				PathDelimter = ':';
+				DotnetPathLocation = "dotnet";
 				BasePythonLocation = @"Engine/Binaries/ThirdParty/Python3/Mac/bin/python3";
 			}
+
+			DotnetPathLocation = Path.Combine(Environment.GetEnvironmentVariable("PATH")
+				.Split(PathDelimter)
+				.Select(Executable => Executable.TrimEnd('\\').TrimEnd('/'))
+				.Where(Executable => Executable.Contains("dotnet"))
+				.FirstOrDefault(), DotnetPathLocation);
+
+			if (string.IsNullOrEmpty(DotnetPathLocation))
+			{
+				Log.Verbose("Could not locate dotnet directory from PATH! Running ReportGenUtils commands with .net dlls will throw exceptions.");
+			}
 		}
 
 		public static int RunPerfReportTool(string Args, string ReportXmlBaseDir)
@@ -35,7 +52,6 @@ namespace Gauntlet
 			return ReportGenUtils.RunCommandlineTool(PerfReportToolPath, AllArgs);
 		}
 
-
 		public static void DirCopy(string SourceDir, string TargetDir, bool bSkipExisting = false, string SearchPattern = "*.*")
 		{
 			Log.Info(string.Format("Copying Directory {0} to {1}", SourceDir, TargetDir));
@@ -61,7 +77,8 @@ namespace Gauntlet
 
 		public static string GetCsvToolPath(string ToolName)
 		{
-			return Path.Combine(CommandUtils.CmdEnv.LocalRoot, "Engine", "Binaries", "DotNet", "CsvTools", ToolName + ".exe");
+			string Extension = BuildHostPlatform.Current.Platform.IsInGroup("Windows") ? ".exe" : ".dll";
+			return Path.Combine(CommandUtils.CmdEnv.LocalRoot, "Engine", "Binaries", "DotNet", "CsvTools", ToolName) + Extension;
 		}
 
 		public static int RunPythonScript(string ScriptPath, string Args)
@@ -108,11 +125,25 @@ namespace Gauntlet
 		public static Process RunCommandlineToolAsync(string ToolPath, string Args, bool bRedirectStdOut = false)
 		{
 			Process ToolProcess = new Process();
-			ToolProcess.StartInfo.FileName = ToolPath;
-			ToolProcess.StartInfo.Arguments = Args;
+			if (ToolPath.EndsWith(".dll"))
+			{
+				if (string.IsNullOrEmpty(DotnetPathLocation))
+				{
+					throw new AutomationException("Could not locate a dotnet executable in your PATH. Ensure it has been added to your environment variables.");
+				}
+				ToolProcess.StartInfo.FileName = DotnetPathLocation;
+				ToolProcess.StartInfo.Arguments = ToolPath + ' ' + Args;
+			}
+			else
+			{
+				ToolProcess.StartInfo.FileName = ToolPath;
+				ToolProcess.StartInfo.Arguments = Args;
+			}
+
 			ToolProcess.StartInfo.RedirectStandardOutput = bRedirectStdOut;
 			ToolProcess.StartInfo.UseShellExecute = false;
-			Log.Info("Running " + ToolPath + " with args:\n" + ToolProcess.StartInfo.Arguments);
+
+			Log.Info("Running " + ToolProcess.StartInfo.FileName + " with args:\n" + ToolProcess.StartInfo.Arguments);
 			ToolProcess.Start();
 			return ToolProcess;
 		}
@@ -396,4 +427,4 @@ namespace Gauntlet
 			}
 		}
 	}
-}
+}
\ No newline at end of file
