[Gauntlet] Avoid formatting error when surfacing log from child process

#rnx

diff --git a/Engine/Source/Programs/AutomationTool/Gauntlet/Unreal/Base/Gauntlet.UnrealTestNode.cs b/Engine/Source/Programs/AutomationTool/Gauntlet/Unreal/Base/Gauntlet.UnrealTestNode.cs
--- a/Engine/Source/Programs/AutomationTool/Gauntlet/Unreal/Base/Gauntlet.UnrealTestNode.cs
+++ b/Engine/Source/Programs/AutomationTool/Gauntlet/Unreal/Base/Gauntlet.UnrealTestNode.cs
@@ -1034,7 +1034,7 @@ namespace Gauntlet
 		}
 
 		/// <summary>
-		/// This class is here to provide compatiblity 
+		/// This class is here to provide compatibility 
 		/// </summary>
 		/// <param name="WasCancelled"></param>
 		protected virtual void StopTest(bool WasCancelled)
@@ -1043,7 +1043,7 @@ namespace Gauntlet
 		}
 
 		/// <summary>
-		/// Called when a test has completed. By default saves artifacts and calles CreateReport
+		/// Called when a test has completed. By default saves artifacts and calls CreateReport
 		/// </summary>
 		/// <param name="InReason"></param>
 		/// <returns></returns>
@@ -1262,7 +1262,7 @@ namespace Gauntlet
 
 		/// <summary>
 		/// Optional function that is called on test completion and gives an opportunity to create a report. The returned class will later be passed to
-		/// SubmitToDashboard if submisson of results is enabled
+		/// SubmitToDashboard if submission of results is enabled
 		/// </summary>
 		/// <param name="Result">Test result</param>
 		/// <param name="Context">Context that describes the environment of the test</param>
@@ -1871,8 +1871,7 @@ namespace Gauntlet
 				return;
 			}
 
-			List<string> LogCategories = new List<string>();
-			LogCategories.Add("Gauntlet");
+			List<string> LogCategories = new List<string>() { "Gauntlet" };
 			{
 				// get the categories used to monitor process (this needs rethought).
 				IEnumerable<string> HeartbeatCategories = GetHeartbeatLogCategories().Union(GetCachedConfiguration().LogCategoriesForEvents);
@@ -1887,7 +1886,7 @@ namespace Gauntlet
 					string StdOut;
 					try
 					{
-						// Convert the ProcessOutput StringBuilder into a regular string. This can run OOM if the output is sufficently large...
+						// Convert the ProcessOutput StringBuilder into a regular string. This can run OOM if the output is sufficiently large...
 						StdOut = App.StdOut;
 					}
 					catch (OutOfMemoryException OOMEx)
@@ -1908,7 +1907,7 @@ namespace Gauntlet
 
 					foreach (string TestLine in Parser.GetLogFromShortNameChannels(LogCategories))
 					{
-						Log.Info(string.Format("{0}: {1}", AppPrefix, TestLine));
+						Log.Info("{App}: {Message}", AppPrefix, TestLine);
 
 						if (bUpdateHeartbeatTime)
 						{
