[LowLevelTest] removing dependency to UnrealLogParser

#rnx

diff --git a/Engine/Source/Programs/AutomationTool/LowLevelTests/Utility/LowLevelTests.LogParser.cs b/Engine/Source/Programs/AutomationTool/LowLevelTests/Utility/LowLevelTests.LogParser.cs
--- a/Engine/Source/Programs/AutomationTool/LowLevelTests/Utility/LowLevelTests.LogParser.cs
+++ b/Engine/Source/Programs/AutomationTool/LowLevelTests/Utility/LowLevelTests.LogParser.cs
@@ -7,7 +7,7 @@ using System.Text.RegularExpressions;
 
 namespace LowLevelTests
 {
-	public class LowLevelTestsLogParser : UnrealLogParser
+	public class LowLevelTestsLogParser
 	{
 		public class CatchTestResults
 		{
@@ -66,9 +66,11 @@ namespace LowLevelTests
 			internal const string CatchThrowException = "Catch will terminate because it needed to throw an exception.";
 		}
 
+		private string Content;
+
 		public LowLevelTestsLogParser(string InContent)
-			: base(InContent)
 		{
+			Content = InContent;
 		}
 
 		public CatchTestResults GetCatchTestResults()
