Removing the restriction to force tvOS and iOS into serial compiles

The issue needing the restriction no longer exists, each target has a unique .app folder with modern xcode tooling

#rnx

diff --git a/Engine/Source/Programs/AutomationTool/AutomationUtils/UnrealBuild.cs b/Engine/Source/Programs/AutomationTool/AutomationUtils/UnrealBuild.cs
--- a/Engine/Source/Programs/AutomationTool/AutomationUtils/UnrealBuild.cs
+++ b/Engine/Source/Programs/AutomationTool/AutomationUtils/UnrealBuild.cs
@@ -591,36 +591,13 @@ namespace AutomationTool
 				}
 			}
 
-			List<BuildTarget> Targets = new List<BuildTarget>(InAgenda.Targets);
-
-			// Temporary hack: iOS & tvOS configs need to build separately
-			if (Targets.Any(x => x.Platform == UnrealTargetPlatform.IOS || x.Platform == UnrealTargetPlatform.TVOS))
-			{
-				List<string> TargetNames = Targets.Select(x => x.TargetName).Distinct().ToList();
-				List<UnrealTargetConfiguration> Configs = Targets.Select(x => x.Config).Distinct().ToList();
-				foreach (string TargetName in TargetNames)
-				{
-					foreach (UnrealTargetConfiguration Config in Configs)
-					{
-						List<BuildTarget> ConfigTargets = Targets.Where(x => (x.Platform == UnrealTargetPlatform.IOS || x.Platform == UnrealTargetPlatform.TVOS) && x.TargetName == TargetName && x.Config == Config).ToList();
-						if (ConfigTargets.Count > 0)
-						{
-							// Build all the targets
-							BuildWithUBT(ConfigTargets, InTargetToManifest, bDisableXGE, InAllCores, InSkipBuild);
-						}
-					}
-				}
-				Targets.RemoveAll(x => x.Platform == UnrealTargetPlatform.IOS || x.Platform == UnrealTargetPlatform.TVOS);
-			}
-			// End hack
-
-			if (Targets.Count == 0)
+			if (InAgenda.Targets.Count == 0)
 			{
 				return;
 			}
 
 			// Build all the targets
-			BuildWithUBT(Targets, InTargetToManifest, bDisableXGE, InAllCores, InSkipBuild);
+			BuildWithUBT(InAgenda.Targets, InTargetToManifest, bDisableXGE, InAllCores, InSkipBuild);
 		}
 
 		/// <summary>
