Fix for debugging gauntlet tests initializing cached variables too early.

diff --git a/Engine/Source/Programs/AutomationTool/Gauntlet/Unreal/Base/Gauntlet.UnrealTestNode.cs b/Engine/Source/Programs/AutomationTool/Gauntlet/Unreal/Base/Gauntlet.UnrealTestNode.cs
--- a/Engine/Source/Programs/AutomationTool/Gauntlet/Unreal/Base/Gauntlet.UnrealTestNode.cs
+++ b/Engine/Source/Programs/AutomationTool/Gauntlet/Unreal/Base/Gauntlet.UnrealTestNode.cs
@@ -423,7 +423,7 @@ namespace Gauntlet
 			}
 			// Using CachedConfiguration so GetConfiguration does not get looped over. If it has not been cached
 			// yet, will call test's GetConfiguration method.
-			var Config = GetCachedConfiguration();
+			var Config = CachedConfig;
 			return (Config is UnrealTestConfiguration) ? Context.GetRoleContext(Config.GetMainRequiredRole().Type).ToString() : Context.ToString();
 		}
 
