Made frequent RPC framework messages Verbose level instead of Info

#rnx

diff --git a/Engine/Source/Programs/AutomationTool/Gauntlet/Framework/RpcFramework/Gauntlet.RpcExecutor.cs b/Engine/Source/Programs/AutomationTool/Gauntlet/Framework/RpcFramework/Gauntlet.RpcExecutor.cs
--- a/Engine/Source/Programs/AutomationTool/Gauntlet/Framework/RpcFramework/Gauntlet.RpcExecutor.cs
+++ b/Engine/Source/Programs/AutomationTool/Gauntlet/Framework/RpcFramework/Gauntlet.RpcExecutor.cs
@@ -196,7 +196,7 @@ namespace Gauntlet
 				{
 					IncomingMessageQueue.Enqueue(IncomingMessage);
 					CallResponse.ResponseBody = String.Format("New message enqueued! Source: {0} \n Category: {1} \n Payload: {2} \n", IncomingMessage.SenderId, IncomingMessage.Category, IncomingMessage.Payload);
-					Log.Info(CallResponse.ResponseBody);
+					Log.Verbose(CallResponse.ResponseBody);
 				}
 			}
 			else
diff --git a/Engine/Source/Programs/AutomationTool/Gauntlet/Framework/RpcFramework/Gauntlet.RpcTargetRegistry.cs b/Engine/Source/Programs/AutomationTool/Gauntlet/Framework/RpcFramework/Gauntlet.RpcTargetRegistry.cs
--- a/Engine/Source/Programs/AutomationTool/Gauntlet/Framework/RpcFramework/Gauntlet.RpcTargetRegistry.cs
+++ b/Engine/Source/Programs/AutomationTool/Gauntlet/Framework/RpcFramework/Gauntlet.RpcTargetRegistry.cs
@@ -74,7 +74,7 @@ namespace Gauntlet
 					ResponseReader.Close();
 				}
 				
-				Log.Info($"{TargetName} received new list of {AvailableRPCs.Count} available RPCs.");
+				Log.Verbose($"{TargetName} received new list of {AvailableRPCs.Count} available RPCs.");
 			}
 			catch (Exception e)
 			{
