Add first version of iteration time tests on Horde

This is a first stab at regularly measuring projects on Horde, with all the imperfections that come with it.

In this CL:
 * updates to the iteration time tests themselves. Everything is timed now, which helps with debugging, and we are no longer spammig the log with heartbeat events.
 * new RPC for flushing analytics events. Gauntlet will otherwise shut down all processes without giving them a chance to finish their work
 * there is code for running arbitrary projects in this setup, but no projects are set yet. The projects often fail.
 * we do however run tests on clean cooks for an empty project as a baseline.

#rnx

diff --git a/Engine/Source/Programs/AutomationTool/Gauntlet/Framework/RpcFramework/Gauntlet.RpcExecutor.cs b/Engine/Source/Programs/AutomationTool/Gauntlet/Framework/RpcFramework/Gauntlet.RpcExecutor.cs
--- a/Engine/Source/Programs/AutomationTool/Gauntlet/Framework/RpcFramework/Gauntlet.RpcExecutor.cs
+++ b/Engine/Source/Programs/AutomationTool/Gauntlet/Framework/RpcFramework/Gauntlet.RpcExecutor.cs
@@ -750,7 +750,7 @@ namespace Gauntlet
 			catch (Exception e)
 			{
 				Log.SuspendECErrorParsing();
-				Log.Info(String.Format("CallRpc failed with following exception: {0}", e.Message));
+				Log.Info($"CallRpc failed with the following exception: {e}");
 				Log.ResumeECErrorParsing();
 				return null;
 			}
