IterativeCompare test script: cooks incrementally and non-incrementally, and compares the results with DiffCookCommandlet.
#rnx

diff --git a/Engine/Build/Graph/Tests/IterativeCompare.xml b/Engine/Build/Graph/Tests/IterativeCompare.xml
new file mode 100644
--- /dev/null
+++ b/Engine/Build/Graph/Tests/IterativeCompare.xml
@@ -0,0 +1,78 @@
+<?xml version='1.0' ?>
+<BuildGraph xmlns="http://www.epicgames.com/BuildGraph" xmlns:xsi="http://www.w3.org/2001/XMLSchema-instance" xsi:schemaLocation="http://www.epicgames.com/BuildGraph ../Schema.xsd" >
+
+	<!-- Project Options -->
+	<Option Name="ProjectFile" DefaultValue="_Unknown_" Restrict=".+" Description="The projects to cook"/>
+	<Option Name="ProjectName" DefaultValue="_Unknown_" Restrict=".+" Description="The projects to cook"/>
+	<Option Name="ProjectDir" DefaultValue="_Unknown_" Restrict=".+" Description="The projects to cook"/>
+	<Option Name="EditorTarget" DefaultValue="$(ProjectName)Editor" Description="Compile name of the editor target"/>
+
+	<!-- Note we can not make the default='', because the default is used for some validation. -->
+	<Error Message="No project file specified" If="'$(ProjectFile)' == '_Unknown_'"/>
+	<Error Message="No project name specified" If="'$(ProjectName)' == '_Unknown_'"/>
+	<Error Message="No project dir specified" If="'$(ProjectDir)' == '_Unknown_'"/>
+
+	<!-- Platform Options -->
+	<Option Name="TargetPlatform" DefaultValue="Windows" Restrict=".+" Description="The platform to cook"/>
+
+	<!-- Agent Options -->
+	<Option Name="AgentType" DefaultValue="" Description="The platform to cook"/>
+	<Property Name="AgentType" Value="IterativeCompareAgent_$(ProjectName)_$(TargetPlatform)" If="'$(AgentType)' == ''"/>
+
+	<Option Name="WithUBA" Restrict="true|false" DefaultValue="true" Description="Whether to use UBA compile or not"/>
+	<Property Name="UBACompileArgs" Value=""/>
+	<Property Name="UBACompileArgs" Value="-UBA -ActionStallReportTime=300 -UBALinkRemote -UBAPrintSummary" If="$(WithUBA)"/>
+	
+	<!-- Agent=IterativeCompareAgent variables -->
+	<Property Name="FullPathProjectDir" Value="$(RootDir)/$(ProjectDir)"/>
+	<Property Name="EditorPlatform" Value="Win64"/>
+	<Property Name="CommonCookArgs" Value="-targetplatform=$(TargetPlatform) -zenstore -cookincremental"/>
+	<Property Name="IterativeCompareDescriptiveName" Value="IterativeCompare $(EditorPlatform) $(ProjectName)"/>
+	<Property Name="PlatformCookedDir" Value="$(FullPathProjectDir)/Saved/Cooked/$(TargetPlatform)"/>
+	<Property Name="CompareDir" Value="$(FullPathProjectDir)/Intermediate/IterativeCompare/$(TargetPlatform)"/>
+	<Property Name="NewLine" Value="&#10;"/>
+
+	<Agent Name="IterativeCompareAgent" Type="$(AgentType)">
+		<Node Name="$(IterativeCompareDescriptiveName)">
+			<!-- Compile. -->
+			<Log Message="$(NewLine)========== IterativeCompare.xml: Compile =========="/>
+			<Compile Target="ShaderCompileWorker" Platform="$(EditorPlatform)" Project="$(ProjectFile)" Configuration="Development" Arguments="$(UBACompileArgs)"/>
+			<Compile Target="ZenLaunch" Platform="$(EditorPlatform)" Configuration="Development" Arguments="$(UBACompileArgs)"/>
+			<Compile Target="$(EditorTarget)" Platform="$(EditorPlatform)" Project="$(ProjectFile)" Configuration="Development" Arguments="$(UBACompileArgs)"/>
+
+			<!-- Execute the incremental cook (we don't pass -forcerecook). -->
+			<Log Message="$(NewLine)========== IterativeCompare.xml: Incremental Cook =========="/>
+			<Commandlet Name="Cook" Project="$(ProjectFile)" Arguments="$(CommonCookArgs)"/>
+
+			<!-- Write the files from incremental cook's zenstore out into the Saved/Cooked/Platform dir. -->
+			<Log Message="$(NewLine)========== IterativeCompare.xml: ZenMirror Incremental Cook =========="/>
+			<Log Message="ZenMirror Project=$(ProjectFile) Platform=$(TargetPlatform)"/>
+			<ZenMirror Project="$(ProjectFile)" Platform="$(TargetPlatform)"/>
+
+			<!-- Move all cook artifacts to the compare dir. -->
+			<Log Message="$(NewLine)========== IterativeCompare.xml: Move =========="/>
+			<Log Message="Delete Directories=$(CompareDir)"/>
+			<Delete Directories="$(CompareDir)"/>
+			<Log Message="Move From=$(PlatformCookedDir)/*.* To=$(CompareDir)"/>
+			<Move From="$(PlatformCookedDir)/*.*" To="$(CompareDir)" ErrorIfNotFound="true"/>
+			<Log Message="Delete Directories=$(PlatformCookedDir)"/>
+			<Delete Directories="$(PlatformCookedDir)"/>
+
+			<!-- Execute the forcerecook cook. -->
+			<Log Message="$(NewLine)========== IterativeCompare.xml: Non-Incremental Cook =========="/>
+			<Commandlet Name="Cook" Project="$(ProjectFile)" Arguments="$(CommonCookArgs) -forcerecook"/>
+
+			<!-- Write the files from forcerecook's zenstore out into the Saved/Cooked/Platform dir. -->
+			<Log Message="$(NewLine)========== IterativeCompare.xml: ZenMirror ForceRecook Cook =========="/>
+			<Log Message="ZenMirror Project=$(ProjectFile) Platform=$(TargetPlatform)"/>
+			<ZenMirror Project="$(ProjectFile)" Platform="$(TargetPlatform)"/>
+
+			<!-- Compare the incremental cook with the non-incremental cook. -->
+			<Log Message="$(NewLine)========== IterativeCompare.xml: DiffCook =========="/>
+			<Commandlet Name="DiffCook" Project="$(ProjectFile)" Arguments="-base=$(PlatformCookedDir) -comp=$(CompareDir) -targetplatform=$(TargetPlatform) -addedverbosity=error -removedverbosity=error -modifiedverbosity=warning"/>
+		</Node>
+	</Agent>
+
+	<Aggregate Name="IterativeCompare" Requires="$(IterativeCompareDescriptiveName)"/>
+
+</BuildGraph>
\ No newline at end of file
diff --git a/Engine/Build/Graph/Tests/IterativeValidate.xml b/Engine/Build/Graph/Tests/IterativeValidate.xml
--- a/Engine/Build/Graph/Tests/IterativeValidate.xml
+++ b/Engine/Build/Graph/Tests/IterativeValidate.xml
@@ -24,12 +24,12 @@
 	<!-- Iterative Options -->
 	<Option Name="IterativeMode" DefaultValue="incremental" Restrict="iterative|incremental" Description="Whether to iteratively cook using Incremental or using the legacy iterative"/>
 	<Option Name="ValidateMode" DefaultValue="onephase" Restrict="onephase|twophase" Description="Whether to validate results in onephase (which reports indeterminism as an iterativefalseskip error) or two phase (which detects indeterminism and downgrades it to a display message)"/>
-	<Property Name="IterativeArgs" Value="-zenstore -ini:Editor:[CookSettings]:HybridIterativeEnabled=true" If="'$(IterativeMode)' == 'incremental'"/>
+	<Property Name="IterativeArgs" Value="-zenstore -cookincremental" If="'$(IterativeMode)' == 'incremental'"/>
 	<Property Name="IterativeArgs" Value="-iterate -ini:Editor:[CookSettings]:HybridIterativeEnabled=false -ini:Editor:[CookSettings]:IterativeUseClassFilters=true" If="'$(IterativeMode)' == 'iterative'"/>
 
 	<!-- Agent=IterativeValidate variables -->
 	<Property Name="EditorPlatform" Value="Win64"/>
-	<Property Name="CommonCookArgs" Value="-targetplatform=$(TargetPlatform) $(IterativeArgs) -skiponlyeditoronly=true"/>
+	<Property Name="CommonCookArgs" Value="-targetplatform=$(TargetPlatform) $(IterativeArgs)"/>
 
 	<Property Name="CompileNodeName" Value="Compile $(EditorPlatform) $(ProjectName)"/>
 	<Property Name="CookIterativeValidateSinglePhaseName" Value="Cook IterativeValidateSinglePhase $(ProjectName) $(TargetPlatform)"/>
diff --git a/Engine/Source/Editor/UnrealEd/Classes/Commandlets/DiffCookCommandlet.h b/Engine/Source/Editor/UnrealEd/Classes/Commandlets/DiffCookCommandlet.h
--- a/Engine/Source/Editor/UnrealEd/Classes/Commandlets/DiffCookCommandlet.h
+++ b/Engine/Source/Editor/UnrealEd/Classes/Commandlets/DiffCookCommandlet.h
@@ -9,6 +9,7 @@
 #include "Containers/Map.h"
 #include "Containers/StringView.h"
 #include "Containers/UnrealString.h"
+#include "Logging/LogVerbosity.h"
 #include "Misc/PackagePath.h"
 #include "Misc/StringBuilder.h"
 #include "Templates/UniquePtr.h"
@@ -20,6 +21,7 @@ class ITargetPlatform;
 class UCookOnTheFlyServer;
 namespace UE::DiffCook::DiffDepotUtils { struct FDepotData; }
 
+
 /**
  * Compares two depots that were written out by the CookCommandlet. Similar to running the cookcommandlet
  * with -diffonly, but compares two depots on disk rather than comparing one depot on disk against the
@@ -172,6 +174,8 @@ private:
 	[[nodiscard]] static FString GetNormalizedFlexPath(FStringView Path);
 	static void NormalizeFlexPath(FStringBuilderBase& Path);
 
+	static ELogVerbosity::Type ParseDiffVerbosity(const FString& Text);
+
 private:
 	FCommandLineArgs Args;
 	TArray<ITargetPlatform*> TargetPlatforms;
@@ -187,6 +191,12 @@ private:
 		MountsWithLeafFolderNameNotEqualLongPackageNameRoot;
 	FCookedDepot BaseDepot;
 	FCookedDepot CompDepot;
+	uint32 NumAdded = 0;
+	uint32 NumRemoved = 0;
+	uint32 NumModified = 0;
+	ELogVerbosity::Type AddedVerbosity = ELogVerbosity::Display;
+	ELogVerbosity::Type RemovedVerbosity = ELogVerbosity::Display;
+	ELogVerbosity::Type ModifiedVerbosity = ELogVerbosity::Display;
 	bool bShowSummary = true;
 	bool bShowPackages = true;
 	bool bShowHeaders = false; // Not yet implemented, change default to true once implemented
diff --git a/Engine/Source/Editor/UnrealEd/Private/Commandlets/DiffCookCommandlet.cpp b/Engine/Source/Editor/UnrealEd/Private/Commandlets/DiffCookCommandlet.cpp
--- a/Engine/Source/Editor/UnrealEd/Private/Commandlets/DiffCookCommandlet.cpp
+++ b/Engine/Source/Editor/UnrealEd/Private/Commandlets/DiffCookCommandlet.cpp
@@ -17,6 +17,7 @@
 #include "Misc/CoreMisc.h"
 #include "Misc/CoreMiscDefines.h"
 #include "Misc/CString.h"
+#include "Misc/FeedbackContext.h"
 #include "Misc/PackageSegment.h"
 #include "Misc/Parse.h"
 #include "Misc/Paths.h"
@@ -72,14 +73,14 @@ int32 UDiffCookCommandlet::Main(const FString& CmdLineParams)
 	if (CompDepot.bValid)
 	{
 		FDiffResult Diff = DiffDepotAsBinary();
-		if (bShowSummary)
-		{
-			PrintSummary(Diff);
-		}
 		if (bShowPackages)
 		{
 			PrintPackageDiffs(Diff);
 		}
+		if (bShowSummary)
+		{
+			PrintSummary(Diff);
+		}
 	}
 	else
 	{
@@ -172,6 +173,18 @@ bool UDiffCookCommandlet::TryParseCommandLine(const FString& CmdLineParams)
 		{
 			bSingleDepot = true;
 		}
+		else if (Key == TEXT("addedverbosity"))
+		{
+			AddedVerbosity = ParseDiffVerbosity(Value);
+		}
+		else if (Key == TEXT("removedverbosity"))
+		{
+			RemovedVerbosity = ParseDiffVerbosity(Value);
+		}
+		else if (Key == TEXT("modifiedverbosity"))
+		{
+			ModifiedVerbosity = ParseDiffVerbosity(Value);
+		}
 	}
 
 	if (Args.BasePath.IsEmpty() || (!bSingleDepot && Args.CompPath.IsEmpty()))
@@ -202,7 +215,10 @@ bool UDiffCookCommandlet::TryParseCommandLine(const FString& CmdLineParams)
 			TEXT("\n\t\t-package=<FileNameOrLongPackageNamesDelimitedBy+>: Show per-package diffs for these and only these packages.")
 			TEXT("\n\t\t-targetplatform=<CookPlatformName>: Interpret the result as for the given platform.")
 			TEXT("\n\t\t\tIf not specified, platform will be inferred from path, if that doesn't work, platform-specific data will be skipped.")
-			TEXT("\n\t\t-singledepot: Ignore -comp and display information about -base without any diff."));
+			TEXT("\n\t\t-singledepot: Ignore -comp and display information about -base without any diff.")
+			TEXT("\n\t\t-addedverbosity=error|warning|display: Specify the verbosity at which the count of added files will be reported.")
+			TEXT("\n\t\t-removedverbosity=error|warning|display: Specify the verbosity at which the count of removed files will be reported.")
+			TEXT("\n\t\t-modifiedverbosity=error|warning|display: Specify the verbosity at which the count of modified files will be reported."));
 	}
 	return bResult;
 }
@@ -922,13 +938,14 @@ UDiffCookCommandlet::FDiffResult UDiffCookCommandlet::DiffDepotAsBinary()
 
 void UDiffCookCommandlet::PrintSummary(FDiffResult& Diff)
 {
-	if (Diff.PackageDiffs.IsEmpty())
-	{
-		UE_LOG(LogCook, Display, TEXT("CookDiff Result: IDENTICAL"));
-		return;
-	}
-	UE_LOG(LogCook, Error, TEXT("CookDiff Result: DIFFERENT"));
-	UE_LOG(LogCook, Display, TEXT("%d files are added, removed, or modified."), Diff.PackageDiffs.Num());
+	UE_LOG(LogCook, Display, TEXT("CookDiff Result: %s"),
+		Diff.PackageDiffs.IsEmpty() ? TEXT("IDENTICAL") : TEXT("DIFFERENT"));
+	GWarn->CategorizedLogf(LogCook.GetCategoryName(), (NumAdded == 0 ? ELogVerbosity::Display : AddedVerbosity),
+		TEXT("%d files added."), NumAdded);
+	GWarn->CategorizedLogf(LogCook.GetCategoryName(), (NumRemoved == 0 ? ELogVerbosity::Display : RemovedVerbosity),
+		TEXT("%d files removed."), NumRemoved);
+	GWarn->CategorizedLogf(LogCook.GetCategoryName(), (NumModified == 0 ? ELogVerbosity::Display : ModifiedVerbosity),
+		TEXT("%d files modified."), NumModified);
 }
 
 void UDiffCookCommandlet::PrintPackageDiffs(FDiffResult& Diff)
@@ -965,14 +982,17 @@ void UDiffCookCommandlet::PrintPackageDiffs(FDiffResult& Diff)
 		switch (PackageDiff.Result)
 		{
 		case EPackageDiffResult::Removed:
+			++NumRemoved;
 			UE_LOG(LogCook, Display, TEXT("Removed:  %s\n\tBase: %s\n\tComp: %s (Missing)"),
 				*PackageNameStr, **BasePath, **CompPath);
 			break;
 		case EPackageDiffResult::Added:
+			++NumAdded;
 			UE_LOG(LogCook, Display, TEXT("Added:    %s\n\tBase: %s (Missing)\n\tComp: %s"),
 				*PackageNameStr, **BasePath, **CompPath);
 			break;
 		case EPackageDiffResult::Modified:
+			++NumModified;
 			bModified = true;
 			break;
 		case EPackageDiffResult::Error: [[fallthrough]];
@@ -1188,6 +1208,13 @@ void UDiffCookCommandlet::NormalizeFlexPath(FStringBuilderBase& Path)
 	FPathViews::NormalizeFilename(Path);
 }
 
+ELogVerbosity::Type UDiffCookCommandlet::ParseDiffVerbosity(const FString& Text)
+{
+	ELogVerbosity::Type Result = ParseLogVerbosityFromString(Text);
+	Result = (ELogVerbosity::Type)FMath::Min((uint8)Result, (uint8)ELogVerbosity::Display);
+	return Result;
+}
+
 bool UDiffCookCommandlet::FPackageData::HasExtension(EPackageExtension Segment) const
 {
 	check((uint32)Segment < sizeof(HasExtensionBitfield) * 8);
diff --git a/Engine/Source/Programs/AutomationTool/BuildGraph/Tasks/ZenMirrorTask.cs b/Engine/Source/Programs/AutomationTool/BuildGraph/Tasks/ZenMirrorTask.cs
new file mode 100644
--- /dev/null
+++ b/Engine/Source/Programs/AutomationTool/BuildGraph/Tasks/ZenMirrorTask.cs
@@ -0,0 +1,212 @@
+// Copyright Epic Games, Inc. All Rights Reserved.
+
+using System;
+using System.Collections.Generic;
+using System.Diagnostics.CodeAnalysis;
+using System.IO;
+using System.Linq;
+using System.Runtime.InteropServices;
+using System.Text;
+using System.Text.Json;
+using System.Text.Json.Serialization;
+using System.Threading.Tasks;
+using System.Xml;
+using EpicGames.Core;
+using EpicGames.Horde;
+using EpicGames.ProjectStore;
+using Microsoft.Extensions.Logging;
+
+namespace AutomationTool.Tasks
+{
+	/// <summary>
+	/// Parameters for a task that calls zen mirror to write the Zen oplog onto disk
+	/// </summary>
+	public class ZenMirrorTaskParameters
+	{
+		/// <summary>
+		/// The project from which to export the snapshot
+		/// </summary>
+		[TaskParameter(Optional = true)]
+		public FileReference Project { get; set; }
+
+		/// <summary>
+		/// The target platform to mirror the snapshot for
+		/// </summary>
+		[TaskParameter(Optional = true)]
+		public string Platform { get; set; }
+
+		/// <summary>
+		/// The path on the local disk to which the data will be mirrored
+		/// If empty then the path will be set to the %Project%\Saved\Cooked\%Platform% directory.
+		/// </summary>
+		[TaskParameter(Optional = true)]
+		public DirectoryReference DestinationFileDir { get; set; }
+	}
+
+	/// <summary>
+	/// Exports an snapshot from Zen to a specified destination.
+	/// </summary>
+	[TaskElement("ZenMirror", typeof(ZenMirrorTaskParameters))]
+	public class ZenMirrorTask : BgTaskImpl
+	{
+		/// <summary>
+		/// Parameters for the task
+		/// </summary>
+		readonly ZenMirrorTaskParameters _parameters;
+
+		/// <summary>
+		/// Constructor.
+		/// </summary>
+		/// <param name="parameters">Parameters for this task</param>
+		public ZenMirrorTask(ZenMirrorTaskParameters parameters)
+		{
+			_parameters = parameters;
+		}
+
+		/// <summary>
+		/// Gets the assumed path to where Zen should exist
+		/// </summary>
+		/// <returns></returns>
+		public static FileReference ZenExeFileReference()
+		{
+			return ResolveFile(String.Format("Engine/Binaries/{0}/zen{1}",
+				HostPlatform.Current.HostEditorPlatform.ToString(),
+				RuntimeInformation.IsOSPlatform(OSPlatform.Windows) ? ".exe" : ""));
+		}
+
+		/// <summary>
+		/// Ensures that ZenServer is running on this current machine. This is needed before running any oplog commands
+		/// This passes the sponsor'd process Id to launch zen.
+		/// This ensures that zen does not live longer than the lifetime of the a process that needs Zen.
+		/// </summary>
+		/// <param name="projectFile"></param>
+		public static void ZenLaunch(FileReference projectFile)
+		{
+			// Get the ZenLaunch executable path
+			FileReference zenLaunchExe = ResolveFile(String.Format("Engine/Binaries/{0}/ZenLaunch{1}",
+				HostPlatform.Current.HostEditorPlatform.ToString(),
+				RuntimeInformation.IsOSPlatform(OSPlatform.Windows) ? ".exe" : ""));
+
+			StringBuilder zenLaunchCommandline = new StringBuilder();
+			zenLaunchCommandline.AppendFormat("{0} -SponsorProcessID={1}",
+				CommandUtils.MakePathSafeToUseWithCommandLine(projectFile.FullName), Environment.ProcessId);
+
+			CommandUtils.RunAndLog(CommandUtils.CmdEnv, zenLaunchExe.FullName,
+				zenLaunchCommandline.ToString(), Options: CommandUtils.ERunOptions.Default);
+		}
+
+		private static bool TryRunAndLogWithoutSpew(string app, string commandLine, bool ignoreFailure)
+		{
+			ProcessResult.SpewFilterCallbackType silentOutputFilter = new ProcessResult.SpewFilterCallbackType(line =>
+				{
+					return null;
+				});
+			try
+			{
+				CommandUtils.RunAndLog(CommandUtils.CmdEnv, app, commandLine, MaxSuccessCode: 0,
+					Options: CommandUtils.ERunOptions.Default, SpewFilterCallback: silentOutputFilter);
+			}
+			catch (CommandUtils.CommandFailedException e)
+			{
+				if (!ignoreFailure)
+				{
+					Logger.LogWarning("{Text}", e.ToString());
+				}
+				return false;
+			}
+			return true;
+		}
+
+		private static bool TryRunMirrorCommand(string app, string commandLine)
+		{
+			int attemptLimit = 2;
+			int attempt = 0;
+			while (attempt < attemptLimit)
+			{
+				if (TryRunAndLogWithoutSpew(app, commandLine, false))
+				{
+					return true;
+				}
+				Logger.LogWarning("Attempt {AttemptNum} of mirroring the oplog failed, {Action}...", attempt + 1, attempt < (attemptLimit - 1) ? "retrying" : "abandoning");
+
+				attempt = attempt + 1;
+			}
+			return false;
+		}
+
+		/// <summary>
+		/// ExecuteAsync the task.
+		/// </summary>
+		/// <param name="job">Information about the current job</param>
+		/// <param name="buildProducts">Set of build products produced by this node.</param>
+		/// <param name="tagNameToFileSet">Mapping from tag names to the set of files they include</param>
+		public override Task ExecuteAsync(JobContext job, HashSet<FileReference> buildProducts,
+			Dictionary<string, HashSet<FileReference>> tagNameToFileSet)
+		{
+			FileReference projectFile = _parameters.Project;
+			if (!FileReference.Exists(projectFile))
+			{
+				throw new AutomationException("Missing project file - {0}", projectFile.FullName);
+			}
+			if (string.IsNullOrEmpty(_parameters.Platform))
+			{
+				throw new AutomationException("Missing platform");
+			}
+
+			ZenLaunch(projectFile);
+
+			DirectoryReference destinationFileDir = _parameters.DestinationFileDir;
+			if (destinationFileDir == null || destinationFileDir.FullName.Length == 0)
+			{
+				destinationFileDir
+					= DirectoryReference.Combine(projectFile.Directory, "Saved", "Cooked", _parameters.Platform);
+			}
+
+			// Get the Zen executable path
+			FileReference zenExe = ZenExeFileReference();
+
+			// Format the command line
+			StringBuilder commandLine = new StringBuilder();
+			commandLine.Append("oplog-mirror");
+
+			commandLine.Append(" --project ");
+			commandLine.Append(ProjectUtils.GetProjectPathId(projectFile));
+
+			commandLine.Append(" --oplog ");
+			commandLine.Append(_parameters.Platform);
+
+			commandLine.Append(" --target ");
+			commandLine.Append(destinationFileDir.FullName);
+
+			TryRunMirrorCommand(zenExe.FullName, commandLine.ToString());
+
+			return Task.CompletedTask;
+		}
+
+		/// <summary>
+		/// Output this task out to an XML writer.
+		/// </summary>
+		public override void Write(XmlWriter writer)
+		{
+			Write(writer, _parameters);
+		}
+
+		/// <summary>
+		/// Find all the tags which are used as inputs to this task
+		/// </summary>
+		/// <returns>The tag names which are read by this task</returns>
+		public override IEnumerable<string> FindConsumedTagNames()
+		{
+			yield break;
+		}
+
+		/// <summary>
+		/// Find all the tags which are modified by this task
+		/// </summary>
+		/// <returns>The tag names which are modified by this task</returns>
+		public override IEnumerable<string> FindProducedTagNames()
+		{
+			yield break;
+		}
+	}
+}
