BuildGraph/ZenMirrorTask: Remove unused using statements; they cause warnings in static analysis.
#rnx

diff --git a/Engine/Source/Programs/AutomationTool/BuildGraph/Tasks/ZenMirrorTask.cs b/Engine/Source/Programs/AutomationTool/BuildGraph/Tasks/ZenMirrorTask.cs
--- a/Engine/Source/Programs/AutomationTool/BuildGraph/Tasks/ZenMirrorTask.cs
+++ b/Engine/Source/Programs/AutomationTool/BuildGraph/Tasks/ZenMirrorTask.cs
@@ -2,18 +2,11 @@
 
 using System;
 using System.Collections.Generic;
-using System.Diagnostics.CodeAnalysis;
-using System.IO;
-using System.Linq;
 using System.Runtime.InteropServices;
 using System.Text;
-using System.Text.Json;
-using System.Text.Json.Serialization;
 using System.Threading.Tasks;
 using System.Xml;
 using EpicGames.Core;
-using EpicGames.Horde;
-using EpicGames.ProjectStore;
 using Microsoft.Extensions.Logging;
 
 namespace AutomationTool.Tasks
