[Gauntlet] Modify EpicRoot posix paths to support different configured builders

diff --git a/Engine/Source/Programs/AutomationTool/Gauntlet/Unreal/Utils/Gauntlet.UnrealHelpers.cs b/Engine/Source/Programs/AutomationTool/Gauntlet/Unreal/Utils/Gauntlet.UnrealHelpers.cs
--- a/Engine/Source/Programs/AutomationTool/Gauntlet/Unreal/Utils/Gauntlet.UnrealHelpers.cs
+++ b/Engine/Source/Programs/AutomationTool/Gauntlet/Unreal/Utils/Gauntlet.UnrealHelpers.cs
@@ -432,26 +432,35 @@ namespace Gauntlet
 	{
 		string PlatformPath;
 
-		public EpicRoot(string Path)
+		public EpicRoot(string InPath)
 		{
-			PlatformPath = Path;
+			PlatformPath = InPath;
 
 			if (BuildHostPlatform.Current.Platform == UnrealTargetPlatform.Mac || BuildHostPlatform.Current.Platform == UnrealTargetPlatform.Linux)
 			{
-				string PosixMountPath = CommandUtils.IsBuildMachine ? "/Volumes/epicgames.net/root" : "/Volumes/root";
+				string PosixMountPath;
+				if(CommandUtils.IsBuildMachine)
+				{
+					string PrimaryRoot = Path.Combine("/Volumes", "epicgames.net", "root");
+					string SecondaryRoot = "/Volumes";
+					PosixMountPath = Directory.Exists(PrimaryRoot) ? PrimaryRoot : SecondaryRoot;
+				}
+				else
+				{
+					PosixMountPath = Path.Combine("/Volumes", "root");
+				}
 
-				if (!Path.Contains("P:"))
+				if (!InPath.Contains("P:"))
 				{
-					PlatformPath = Regex.Replace(Path, @"\\\\epicgames.net\\root", PosixMountPath, RegexOptions.IgnoreCase);
+					PlatformPath = Regex.Replace(InPath, @"\\\\epicgames.net\\root", PosixMountPath, RegexOptions.IgnoreCase);
 				}
 				else
 				{
-					PlatformPath = Regex.Replace(Path, "P:", PosixMountPath, RegexOptions.IgnoreCase);
+					PlatformPath = Regex.Replace(InPath, "P:", PosixMountPath, RegexOptions.IgnoreCase);
 				}
 				
 				PlatformPath = PlatformPath.Replace(@"\", "/");
 			}
-
 		}
 
 		public static implicit operator string(EpicRoot Path)
