[Gauntlet] Updated implementation of TargetDeviceIOS
This CL has iOS reach parity all other TargetDevice platforms in the standard set forth in UE-201403
This includes launch flow function implementations, overlay executable support, and updated toolset to support modern versions of iOS and xCode
Namely, ios-deploy has become obsolete for iOS 17+ due to changes in Apple's system frameworks. Instead Apple's new tool, devicectl, is used in versions of xcode where available

The following types of builds are supported:
- Standard Apps (non-bulk)
- Standard IPAs (non-bulk)
- Bulk apps
- dev apps (bulk and non-bulk) aka overlay executables

#tests Preflight and local runs against all build variations

diff --git a/Engine/Source/Programs/AutomationTool/Gauntlet/Framework/Base/Gauntlet.OverlayExecutable.cs b/Engine/Source/Programs/AutomationTool/Gauntlet/Framework/Base/Gauntlet.OverlayExecutable.cs
--- a/Engine/Source/Programs/AutomationTool/Gauntlet/Framework/Base/Gauntlet.OverlayExecutable.cs
+++ b/Engine/Source/Programs/AutomationTool/Gauntlet/Framework/Base/Gauntlet.OverlayExecutable.cs
@@ -50,7 +50,7 @@ namespace Gauntlet
 				throw new AutomationException("Overlay Error: No base executable was specified.");
 			}
 
-			if (!File.Exists(BaseExecutable))
+			if (!File.Exists(BaseExecutable) && !Directory.Exists(BaseExecutable)) // mac/ios apps can be directories
 			{
 				throw new AutomationException("Overlay Error: Could not find base executable {0}.", BaseExecutable);
 			}
@@ -61,7 +61,7 @@ namespace Gauntlet
 			string ExecutableExtension = string.IsNullOrEmpty(ExtensionOverride) ? Path.GetExtension(BaseExecutable) : ExtensionOverride;
 			string LocalExecutable = Path.Combine(PlatformBinariesDirectory, ExecutableName + ExecutableExtension);
 
-			if (!File.Exists(LocalExecutable))
+			if (!File.Exists(LocalExecutable) && !Directory.Exists(LocalExecutable))
 			{
 				Log.Verbose("No local executable for {Platform} exists. Skipping overlay for this role", Platform);
 				return false;
diff --git a/Engine/Source/Programs/AutomationTool/Gauntlet/Platform/IOS/Gauntlet.IOSBuildSource.cs b/Engine/Source/Programs/AutomationTool/Gauntlet/Platform/IOS/Gauntlet.IOSBuildSource.cs
--- a/Engine/Source/Programs/AutomationTool/Gauntlet/Platform/IOS/Gauntlet.IOSBuildSource.cs
+++ b/Engine/Source/Programs/AutomationTool/Gauntlet/Platform/IOS/Gauntlet.IOSBuildSource.cs
@@ -22,8 +22,6 @@ namespace Gauntlet
 
 		public bool IsIPAFile;
 
-		public Dictionary<string, string> FilesToInstall;
-
 		public string PackageName;
 
 		public BuildFlags Flags { get; protected set; }
@@ -34,12 +32,11 @@ namespace Gauntlet
 
 		public bool SupportsAdditionalFileCopy { get; }
 
-		public AppleBuild(UnrealTargetConfiguration InConfig, string InPackageName, string InSourcePath, Dictionary<string, string> InFilesToInstall, BuildFlags InFlags)
+		public AppleBuild(UnrealTargetConfiguration InConfig, string InPackageName, string InSourcePath, BuildFlags InFlags)
 		{
 			Configuration = InConfig;
 			PackageName = InPackageName;
 			SourcePath = InSourcePath;
-			FilesToInstall = InFilesToInstall;
 			Flags = InFlags;
 			SupportsAdditionalFileCopy = true;
 			IsIPAFile = Path.GetExtension(InSourcePath).Equals(".ipa", StringComparison.OrdinalIgnoreCase);
@@ -249,17 +246,17 @@ namespace Gauntlet
 			}
 		}
 
-		public static T CreateFromPath<T>(string InProjectName, string InPath, AppleBuildSource<T> BuildSource)
+		public static T CreateFromPath<T>(string InProjectName, string InRootPath, string InBuildPath, AppleBuildSource<T> BuildSource)
 			where T : AppleBuild
 		{
 			FileSystemInfo BuildPath;
-			if (Directory.Exists(InPath))
+			if (Directory.Exists(InBuildPath))
 			{
-				BuildPath = new DirectoryInfo(InPath);
+				BuildPath = new DirectoryInfo(InBuildPath);
 			}
-			else if (File.Exists(InPath))
+			else if (File.Exists(InBuildPath))
 			{
-				BuildPath = new FileInfo(InPath);
+				BuildPath = new FileInfo(InBuildPath);
 			}
 			else
 			{
@@ -267,29 +264,31 @@ namespace Gauntlet
 				return null;
 			}
 
-			UnrealTargetConfiguration Configuration = UnrealHelpers.GetConfigurationFromExecutableName(InProjectName, InPath);
-			if (Configuration == UnrealTargetConfiguration.Unknown)
-			{
-				Log.Verbose("Could not deduce iOS build configuration from build at path {BuildPath}. Skipping.", InPath);
-				return null;
-			}
-
 			// Check there's an executable with the right name
 			string AppShortName = Regex.Replace(InProjectName, "Game", string.Empty, RegexOptions.IgnoreCase);
+			UnrealTargetConfiguration Configuration = UnrealTargetConfiguration.Unknown;
+
 			if (BuildPath is DirectoryInfo App)
 			{
+				Configuration = UnrealHelpers.GetConfigurationFromExecutableName(InProjectName, InBuildPath);
+				if (Configuration == UnrealTargetConfiguration.Unknown)
+				{
+					Log.Verbose("Could not deduce iOS build configuration from build at path {BuildPath}. Skipping.", InBuildPath);
+					return null;
+				}
+
 				if (App.GetFiles(AppShortName + '*', SearchOption.TopDirectoryOnly).FirstOrDefault() == null)
 				{
-					Log.Verbose("Could not find an executable within build path {BuildPath}. Skipping", InPath);
+					Log.Verbose("Could not find an executable within build path {BuildPath}. Skipping", InBuildPath);
 					return null;
 				}
 			}
 			else
 			{
 				// Get a list of files in the IPA
-				if (!ExecuteIPAZipCommand(string.Format("-Z1 {0}", InPath), out string Output))
+				if (!ExecuteIPAZipCommand(string.Format("-Z1 {0}", InBuildPath), out string Output))
 				{
-					Log.Info("Unable to list files for IPA {IPAPath}", InPath);
+					Log.Info("Unable to list files for IPA {IPAPath}", InBuildPath);
 					return null;
 				}
 
@@ -300,6 +299,12 @@ namespace Gauntlet
 				{
 					if(File.Contains(AppShortName, StringComparison.OrdinalIgnoreCase))
 					{
+						Configuration = UnrealHelpers.GetConfigurationFromExecutableName(InProjectName, File);
+						if (Configuration == UnrealTargetConfiguration.Unknown)
+						{
+							Log.Verbose("Could not deduce iOS build configuration from build at path {BuildPath}. Skipping.", InBuildPath);
+							return null;
+						}
 						bFoundExecutable = true;
 						break;
 					}
@@ -307,24 +312,11 @@ namespace Gauntlet
 
 				if (!bFoundExecutable)
 				{
-					Log.Verbose("Could not find an executable within build path {BuildPath}. Skipping", InPath);
+					Log.Verbose("Could not find an executable within build path {BuildPath}. Skipping", InBuildPath);
 					return null;
 				}
 			}
 
-			Log.Verbose("Pulling package data from {BuildPath}", BuildPath);
-
-			// IOS builds are always packaged, and can always replace the command line and executable (even as IPAs because we cache the unzipped app)
-			BuildFlags Flags = BuildFlags.Packaged | BuildFlags.CanReplaceCommandLine | BuildFlags.CanReplaceExecutable;
-			if (BuildPath.FullName.Contains("Bulk"))
-			{
-				Flags |= BuildFlags.Bulk;
-			}
-			else
-			{
-				Flags |= BuildFlags.NotBulk;
-			}
-
 			PlistInfo Info = GetPlistInfo(BuildPath.FullName);
 			if(Info == null)
 			{
@@ -345,12 +337,44 @@ namespace Gauntlet
 				return null;
 			}
 
+			// IOS builds are always packaged, and can always replace the command line and executable (even as IPAs because we cache the unzipped app)
+			BuildFlags Flags = BuildFlags.Packaged | BuildFlags.CanReplaceCommandLine | BuildFlags.CanReplaceExecutable;
+			if (BuildPath.FullName.Contains("Bulk"))
+			{
+				if(BuildPath is FileInfo)
+				{
+					Log.Info("Bulk builds cannot be contained in an IPA file. Skipping.");
+					return null;
+				}
+				
+				Flags |= BuildFlags.Bulk;
+
+				int PayloadIndex = BuildPath.FullName.IndexOf("/Payload");
+				if(PayloadIndex < 0)
+				{
+					Log.Info("Could not locate the payload directory for bulk build at path {BuildPath}. Skipping.", BuildPath);
+					return null;
+				}
+				else
+				{
+					string RootBuildDirectory = BuildPath.FullName.Substring(0, PayloadIndex);
+					string BulkContentCopyCommands = Path.Combine(RootBuildDirectory, "RequiredCommands.txt");
+					if(string.IsNullOrEmpty(BulkContentCopyCommands) || !File.Exists(BulkContentCopyCommands))
+					{
+						Log.Info("Could not locate bulk content command list {CommandsFile}. Skipping.", BulkContentCopyCommands);
+						return null;
+					}
+				}
+			}
+			else
+			{
+				Flags |= BuildFlags.NotBulk;
+			}
+
 			Log.Verbose("Found bundle id: {PackageName}", PackageName);
 			Log.Verbose("Found {Configuration} {Flags} build at {BuildPath}", Configuration, ((Flags & BuildFlags.Bulk) == BuildFlags.Bulk) ? "(bulk)" : "(not bulk)", BuildPath);
 
-			// Todo: Consider handling Bulk content copies here?
-			Dictionary<string, string> FilesToInstall = new Dictionary<string, string>();
-			return Activator.CreateInstance(typeof(T), new object[] { Configuration, PackageName, BuildPath.FullName, FilesToInstall, Flags }) as T;
+			return Activator.CreateInstance(typeof(T), new object[] { Configuration, PackageName, BuildPath.FullName, Flags }) as T;
 		}
 	}
 
@@ -388,7 +412,7 @@ namespace Gauntlet
 
 			// Interface default parameters don't let us modify the default if calling from an interface cast...
 			// IOS builds are often located deeper within a client directory, so increase the depth here
-			MaxRecursion = MaxRecursion > 5 ? MaxRecursion : 5;
+			MaxRecursion = MaxRecursion > 7 ? MaxRecursion : 7;
 
 			DirectoryInfo BuildDirectory = new DirectoryInfo(InPath);
 			if (BuildDirectory.Exists)
@@ -416,11 +440,12 @@ namespace Gauntlet
 					Apps.AddRange(ValidApps);
 					Apps.AddRange(ValidIPAs);
 					SubDirectories = SubDirectories.SelectMany(Directory => Directory.GetDirectories("*", SearchOption.TopDirectoryOnly)).ToArray();
+
 				}
 
 				foreach (FileSystemInfo App in Apps)
 				{
-					AppleBuild Build = AppleBuild.CreateFromPath(InProjectName, App.FullName, this);
+					AppleBuild Build = AppleBuild.CreateFromPath(InProjectName, InPath, App.FullName, this);
 					if (Build != null)
 					{
 						Builds.Add(Build);
@@ -430,13 +455,12 @@ namespace Gauntlet
 
 			return Builds;
 		}
-
 	}
 
 	public class IOSBuild : AppleBuild
 	{
-		public IOSBuild(UnrealTargetConfiguration InConfig, string InPackageName, string InSourcePath, Dictionary<string, string> InFilesToInstall, BuildFlags InFlags)
-			: base(InConfig, InPackageName, InSourcePath, InFilesToInstall, InFlags)
+		public IOSBuild(UnrealTargetConfiguration InConfig, string InPackageName, string InSourcePath, BuildFlags InFlags)
+			: base(InConfig, InPackageName, InSourcePath, InFlags)
 		{ }
 
 		public override UnrealTargetPlatform Platform => UnrealTargetPlatform.IOS;
@@ -450,8 +474,8 @@ namespace Gauntlet
 
 	public class TVOSBuild : AppleBuild
 	{
-		public TVOSBuild(UnrealTargetConfiguration InConfig, string InPackageName, string InSourcePath, Dictionary<string, string> InFilesToInstall, BuildFlags InFlags)
-			: base(InConfig, InPackageName, InSourcePath, InFilesToInstall, InFlags)
+		public TVOSBuild(UnrealTargetConfiguration InConfig, string InPackageName, string InSourcePath, BuildFlags InFlags)
+			: base(InConfig, InPackageName, InSourcePath, InFlags)
 		{ }
 
 		public override UnrealTargetPlatform Platform => UnrealTargetPlatform.TVOS;
diff --git a/Engine/Source/Programs/AutomationTool/Gauntlet/Platform/IOS/Gauntlet.TargetDeviceIOS.cs b/Engine/Source/Programs/AutomationTool/Gauntlet/Platform/IOS/Gauntlet.TargetDeviceIOS.cs
--- a/Engine/Source/Programs/AutomationTool/Gauntlet/Platform/IOS/Gauntlet.TargetDeviceIOS.cs
+++ b/Engine/Source/Programs/AutomationTool/Gauntlet/Platform/IOS/Gauntlet.TargetDeviceIOS.cs
@@ -1,25 +1,50 @@
 // Copyright Epic Games, Inc. All Rights Reserved.
 
 using System;
-using System.Collections.Generic;
 using System.IO;
-using System.Threading;
 using System.Linq;
-using System.Security.Cryptography;
-using AutomationTool;
-using UnrealBuildTool;
+using System.Data;
 using System.Text;
 using System.Text.RegularExpressions;
+using System.Collections.Generic;
+using System.Security.Cryptography;
 using EpicGames.Core;
-
-/*
-
-General Device Notes (and areas for improvement):
-
-1) We don't currently support parallel iOS tests, see https://jira.it.epicgames.net/browse/UEATM-219
-2) Device Farm devices should be in airplane mode + wifi to avoid No Sim warning notification
-
-*/
+using Gauntlet.Utils;
+using AutomationTool;
+using UnrealBuildTool;
+using static AutomationTool.CommandUtils;
+
+/** IOS General Notes:
+ * 
+ *	Functionality:
+ *	
+ *		Tooling for iOS automation is in a tumultuous state.
+ *		Historically, ios-deploy was used to manage all aspects of automation.
+ *		However, recent changes in Apple's system framework have made ios-deploy obsolete for launching apps for iOS versions 17+ (See for details: [REDACTED_PROFILE_URL]/ios-deploy/issues/588)
+ *		Alternative tools offer most - but not all - of the features needed to install, launch, monitor, and debug an app.
+ *		Because of this, Gauntlet iOS automation is restricted to running on a Mac host and mixes and matches which tools are used for a given operation.
+ *		
+ *		For now, TargetDeviceIOS uses the following tools to support automation:
+ *			ios-deploy			- Third party CLI tool, still used for legacy operations and copying files to/from device - [REDACTED_PROFILE_URL]/ios-deploy
+ *			devicectl			- Apple's core device framework. Has issues with copying files FROM the device, but handles just about everything else - run 'xcrun devicectl' for details
+ *			libimobiledevice	- Cross-platform third party library, only used for ease of bulk content copies - https://libimobiledevice.org/
+ *			
+ *		In general:
+ *			IF you are using xcode versions 15 or below AND iOS version 16 or below, ios-deploy will be primarily used
+ *			IF you are using xcode version 16+, devicectl will be primarily used
+ *			IF you are using trying to test on iOS versions 17+, you MUST upgrade to XCode 16 which contains critical app monitoring features. (See https://forums.developer.apple.com/forums/thread/756393)
+			
+		*** IMPORTANT ****
+		Testing iOS 17+ and iOS 16 or lower in a single pass is not supported. Doing so can result in undefined behavior.
+		If you need to do this, separate the testing into separate passes targeting each device.
+ *			
+ *	Other:
+ *	
+ *		- Builds (including local developer builds used with -dev) must already be signed. If built with UBT, signing is already part of the process. Automated re-signing may be added in the future.
+ *		- Installing builds requires your device to be added to a mobile provision, and for your build to have been signed with an embedded mobile provision containing the device (See https://developer.apple.com/documentation/xcode/distributing-your-app-to-registered-devices)
+ *		- If your app installs content after launch, consider using a bulk build to prevent timeouts. This also lets you more rapidly test -dev apps because code and content are separated
+ *		- If your app requires additional permissions, consider using MDM profiles to prevent permission pop-ups from blocking your tests
+ */
 
 namespace Gauntlet
 {
@@ -28,94 +53,176 @@ namespace Gauntlet
 	/// </summary>
 	public class TargetDeviceIOS : ITargetDevice
 	{
+		public class ConnectedDevice
+		{
+			public string UUID;
+			public int IOSMajorVersion;
+			public bool bConnected;
+			public ITargetDevice AssignedDevice;
+			public bool IsAssigned => AssignedDevice != null;
+			public ConnectedDevice(string UUID, int IOSMajorVersion, bool bConnected)
+			{
+				this.UUID = UUID;
+				this.IOSMajorVersion = IOSMajorVersion;
+				this.bConnected = bConnected;
+			}		
+		}
+
+		/// <summary>
+		/// Whether or not devicectl is used to run most device operations
+		/// ios-deploy will be used when this is false.
+		/// This property is set based on the installed version of xcode and the ios version of the device
+		/// </summary>
+		private static bool UseDeviceCtl;
+
+		/// <summary>
+		/// Collection of devices visible to the UAT host
+		/// </summary>
+		private static Dictionary<string, ConnectedDevice> ConnectedDevices;
+
+		/// <summary>
+		/// Currently active version of XCode. This can be changed between installed versions using 'xcodes select'
+		/// </summary>
+		private static int XCodeMajorVersion;
+
+		/// <summary>
+		/// Creating an iOS device requires several queries to be performed first.
+		/// These results are cached and gated from being refreshed behind this bool
+		/// </summary>
+		private static bool bStaticInitialized;
+
+		/// <summary>
+		/// Prevents separate start threads from managing connected device management at the same time (bulk only)
+		/// </summary>
+		private static object InitializerMutex = new object();
+
+		/// <summary>
+		/// Prevents separate start threads from writing to the network build cache directory at the same time(bulk only)
+		/// </summary>
+		private static object CacheMutex = new object();
+
+		/// <summary>
+		/// Used to prevent re-copies of an entire network share build in the event of a device failure (bulk only)
+		/// </summary>
+		private static bool HasCopiedNetworkBuild = false;
 
 		#region ITargetDevice
 		public string Name { get; protected set; }
 
 		public UnrealTargetPlatform? Platform => UnrealTargetPlatform.IOS;
 
-		public CommandUtils.ERunOptions RunOptions { get; set; }
+		public ERunOptions RunOptions { get; set; }
 
-		public bool IsConnected => Connected;
+		public bool IsConnected => ConnectedDevices.ContainsKey(UUID) && ConnectedDevices[UUID].bConnected;
 
 		public bool IsOn => true;
 
 		public bool IsAvailable => true;
 		#endregion
 
+		[AutoParam(60 * 60 * 2)] // iOS installs can be slow... 2 hour default
+		public int MaxInstallTime { get; protected set; }
+
 		/// <summary>
 		/// Temp path we use to push/pull things from the device
 		/// </summary>
 		public string LocalCachePath { get; protected set; }
 
-		/// <summary>
-		/// Artifact (e.g. Saved) path on the device
-		/// </summary>
-		public string DeviceArtifactPath { get; protected set; }
-
 		/// <summary>
 		/// Low-level device name (uuid)
 		/// </summary>
-		public string DeviceName { get; protected set; }
-
-		protected Dictionary<EIntendedBaseCopyDirectory, string> LocalDirectoryMappings { get; set; }
-
-		private static object AppLock = new object();
+		public string UUID { get; protected set; }
 
-		private static bool ZombiesKilled = false;
-
-		private static Dictionary<string, bool> ConnectedDevices = new Dictionary<string, bool>();
-
-		private bool Connected = false;
+		protected Dictionary<EIntendedBaseCopyDirectory, string> LocalDirectoryMappings { get; set; } = new();
 
 		private bool IsDefaultDevice = false;
 
-		// path to locally extracted (and possibly resigned) app bundle
-		// Note: ios-deploy works with app bundles, which requires the IPA be unzipped for deployment (this will allow us to resign in the future as well)
-		private string LocalAppBundle = null;
-
-		// Gauntlet cache folder for tracking device/ipa state
-		string GauntletAppCache => Path.Combine(Globals.TempDir, "IOSAppCache");
-
-		// the current App MD5 hash, which is tracked to avoid unneccessary deployments and unzip operations
-		private string AppHashFilename => Path.Combine(GauntletAppCache, "AppHash.txt");
-
-		// file whose presence signals that cache was resigned
-		private string CacheResignedFilename => Path.Combine(GauntletAppCache, "Resigned.txt");
+		private IOSAppInstall InstallCache = null;
 
 		public TargetDeviceIOS(string InName, string InCachePath = null)
 		{
-			Name = InName;
-			LocalDirectoryMappings = new Dictionary<EIntendedBaseCopyDirectory, string>();
-			IsDefaultDevice = string.IsNullOrEmpty(InName) || InName.Equals("default", StringComparison.OrdinalIgnoreCase);
+			AutoParam.ApplyParamsAndDefaults(this, Globals.Params.AllArguments);
 
-			List<string> DefaultDevices = GetConnectedDeviceUUID();	
+			if (BuildHostPlatform.Current.Platform != UnrealTargetPlatform.Mac)
+			{
+				throw new AutomationException("Creation of an iOS target device is only supported on Mac at this time.");
+			}
 
-			KillZombies();
+			InitializeDevices();
+
+			IsDefaultDevice = string.IsNullOrEmpty(InName) || InName.Equals("default", StringComparison.OrdinalIgnoreCase);
+			Name = IsDefaultDevice ? "default" : InName;
 
-			// If no device name or its 'default' then use the first default device
 			if (IsDefaultDevice)
 			{
-				if (DefaultDevices.Count() == 0)
+				IEnumerable<ConnectedDevice> AvailableDevices = ConnectedDevices.Values.Where(Device => !Device.IsAssigned);
+				if(!AvailableDevices.Any())
 				{
-					throw new AutomationException("No default device available");
+					throw new AutomationException("Ran out of unassigned default devices. Ensure all devices are properly connected to the host.");
 				}
 
-				DeviceName = DefaultDevices.First();
+				if(!AvailableDevices.Where(Device => Device.bConnected).Any())
+				{
+					// If some devices are detected, but none connected, attempt connection
+					UUID = AvailableDevices.First().UUID;
+					Log.Info("Attempting to connect to {UUID}", UUID);
+					if (!Connect())
+					{
+						throw new AutomationException("Failed to connect to {DeviceName}. Ensure the device is properly connected to the host.", UUID);
+					}
+				}
+				else
+				{
+					// Just select the first available device
+					UUID = AvailableDevices.Where(Device => Device.bConnected).First().UUID;
+				}
 
-				Log.Verbose("Selected device {0} as default", DeviceName);
+				Log.Verbose("Selected device {DeviceName} as default device", UUID);
 			}
 			else
 			{
-				DeviceName = InName.Trim();
-				if (!DefaultDevices.Contains(DeviceName))
+				UUID = Name;
+				if (!ConnectedDevices.ContainsKey(Name))
+				{
+					throw new AutomationException("Device with UUID {0} not found in device list", UUID);
+				}
+				if (!ConnectedDevices[UUID].bConnected)
+				{
+					Log.Info("Attempting to connect to {DeviceName}");
+					if(!Connect())
+					{
+						throw new AutomationException("Failed to connect to {DeviceName}", UUID);
+					}
+				}
+			}
+
+			ConnectedDevices[UUID].AssignedDevice = this;
+
+			if(XCodeMajorVersion < 15)
+			{
+				UseDeviceCtl = false;
+			}
+			else if(XCodeMajorVersion == 15)
+			{
+				if(ConnectedDevices[UUID].IOSMajorVersion < 17)
 				{
-					throw new AutomationException("Device with UUID {0} not found in device list", DeviceName);
+					UseDeviceCtl = false;
 				}
+				else
+				{
+					string Message = "As of XCode 15, Apple has removed DeveloperDiskImages for iOS 17 and beyond. " +
+						"Because of this, ios-deploy cannot be used to run and montior apps within Gauntlet. " +
+						"Instead, DeviceCtl replaces this functionality which is fully available as of XCode 16. " +
+						"Update to XCode 16 to prevent this error.";
+					throw new AutomationException(Message);
+				}
+			}
+			else
+			{
+				UseDeviceCtl = true;
 			}
 
-			// setup local cache
-			LocalCachePath = InCachePath ?? Path.Combine(GauntletAppCache, "IOSDevice_" + DeviceName);
+			LocalCachePath = InCachePath ?? Path.Combine(Globals.TempDir, "IOSDevice_" + UUID);
 		}
 
 		#region IDisposable Support
@@ -146,51 +253,64 @@ namespace Gauntlet
 		#region ITargetDevice
 		public bool Connect()
 		{
-			lock (Globals.MainLock)
+			if(UseDeviceCtl)
 			{
-				if (Connected)
+				if(ConnectedDevices.ContainsKey(UUID) && ConnectedDevices[UUID].bConnected)
 				{
-					return true;
+					return true; // already connected
 				}
 
-				bool ExistingConnection = false;
-				if (ConnectedDevices.TryGetValue(DeviceName, out ExistingConnection))
+				IProcessResult ConnectResult = ExecuteDevicectlCommand("manage pair");
+				if(ConnectResult.ExitCode != 0)
 				{
-					if (ExistingConnection)
-					{
-						throw new AutomationException("Connected to already connected device");
-					}
+					Log.Info("Connecting to {DeviceName} exited with code {ExitCode}:\n{Output}", UUID, ConnectResult.ExitCode, ConnectResult.Output);
+					return false;
 				}
 
-				ConnectedDevices[DeviceName] = true;
+				string DetailsFile = Path.GetTempFileName();
+				string DetailsCommand = string.Format("device info details -j {0}", DetailsFile);
+				IProcessResult DetailResult = ExecuteDevicectlCommand(DetailsCommand);
+				if(DetailResult.ExitCode != 0)
+				{
+					Log.Info("Detailing {DeviceName} exited with code {ExitCode}. Skipping connection:\n{Output}", UUID, DetailResult.ExitCode, ConnectResult.Output);
+				}
 
-				Connected = true;
-			}
+				JsonObject JsonDevice = JsonObject.Parse(File.ReadAllText(DetailsFile)).GetObjectField("result");
 
-			return true;
-		}
+				JsonObject ConnectionProperties = JsonDevice
+					.GetObjectField("connectionProperties");
+				string PairStatus = ConnectionProperties.GetStringField("pairingState");
 
-		public bool Disconnect(bool bForce = false)
-		{
-			lock (Globals.MainLock)
-			{
-				if (!Connected)
+				if(!PairStatus.Equals("paired", StringComparison.OrdinalIgnoreCase))
 				{
-					return true;
+					Log.Info("{DeviceName}'s pair status is still not paired. Skipping connection.", UUID);
+					return false;
 				}
 
-				Connected = false;
-
-				if (ConnectedDevices.ContainsKey(DeviceName))
+				if(ConnectedDevices.ContainsKey(UUID))
 				{
-					ConnectedDevices.Remove(DeviceName);
+					ConnectedDevices[UUID].bConnected = true;
 				}
+				else
+				{
+					string OSVersionString = JsonDevice.GetObjectField("deviceProperties").GetStringField("osVersionNumber");
+					int IOSMajorVersion = int.Parse(OSVersionString.Substring(0, OSVersionString.IndexOf('.')));
 
+					ConnectedDevice Device = new ConnectedDevice(UUID, IOSMajorVersion, true);
+					ConnectedDevices.Add(UUID, Device);
+				}
 			}
 
 			return true;
 		}
 
+		public bool Disconnect(bool bForce = false)
+		{
+			// Disconnecting from an iOS device causes a lot of problems for device farm automation
+			// Namely, re-pairing a device requires someone to re-trust the host PC manually
+			return true;
+		}
+
 		public bool PowerOn()
 		{
 			return true;
@@ -203,53 +323,14 @@ namespace Gauntlet
 
 		public bool Reboot()
 		{
-			const string Cmd = "/usr/local/bin/idevicediagnostics";
-			if (!File.Exists(Cmd))
-			{
-				Log.Verbose("Rebooting iOS device requires idevicediagnostics binary");
-				return true;
-			}
-
-			var Result = IOSBuild.ExecuteCommand(Cmd, string.Format("restart -u {0}", DeviceName));
-			if (Result.ExitCode != 0)
-			{
-				Log.Warning(string.Format("Failed to reboot iOS device {0}, restart command failed", DeviceName));
-				return true;
-			}
-
-			// initial wait 20 seconds
-			Thread.Sleep(20 * 1000);
-
-			const int WaitPeriod = 10;
-			int WaitTime = 120;
-			bool rebooted = false;
-			do
-			{
-				Result = IOSBuild.ExecuteCommand(Cmd, string.Format("diagnostics WiFi -u {0}", DeviceName));
-				if (Result.ExitCode == 0)
-				{
-					rebooted = true;
-					break;
-				}
-
-				Thread.Sleep(WaitPeriod * 1000);
-				WaitTime -= WaitPeriod;
-
-			} while (WaitTime > 0);
-
-			if (!rebooted)
-			{
-				Log.Warning(KnownLogEvents.Gauntlet_DeviceEvent, "Failed to reboot iOS device {Name}, device didn't come back after restart", DeviceName);
-			}
-
-			return true;
+			// Rebooting tends to lock the device and break connections requiring manual intervention to unlock the device
+			return true; 
 		}
 
 		public Dictionary<EIntendedBaseCopyDirectory, string> GetPlatformDirectoryMappings()
 		{
 			return LocalDirectoryMappings;
 		}
-
 		public override string ToString()
 		{
 			return Name;
@@ -258,663 +339,650 @@ namespace Gauntlet
 
 		public void FullClean()
 		{
+			string AppOutput;
+			if (UseDeviceCtl)
+			{
+				IProcessResult AppListResult = ExecuteDevicectlCommand("device info apps");
+				if (AppListResult.ExitCode != 0)
+				{
+					throw new AutomationException("App list exited with code {0}: {1}", AppListResult.ExitCode, AppListResult.Output);
+				}
 
-		}
+				AppOutput = AppListResult.Output;
+			}
+			else
+			{
+				IProcessResult AppListResult = ExecuteIOSDeployCommand("-B");
+				if (AppListResult.ExitCode != 0)
+				{
+					throw new AutomationException("App list exited with code {0}: {1}", AppListResult.ExitCode, AppListResult.Output);
+				}
 
-		public void CleanArtifacts()
-		{
+				AppOutput = AppListResult.Output;
+			}
 
-		}
+			Regex AppRegex = new Regex("com[^\\s]*"); // Ex. com.epicgames.enginetest
+			IEnumerable<string> AppList = AppRegex.Matches(AppOutput)
+				.Where(Match => Match.Success)
+				.Select(Match => Match.Value)
+				.Where(Match => !Match.Contains("com.apple.", StringComparison.OrdinalIgnoreCase));
 
-		public void InstallBuild(UnrealAppConfig AppConfiguration)
-		{
+			foreach (string App in AppList)
+			{
+				Log.Info("Uninstalling {App}...", App);
 
-		}
+				IProcessResult UninstallResult;
+				if(UseDeviceCtl)
+				{
+					UninstallResult = ExecuteDevicectlCommand(string.Format("device uninstall app {0}", App));
+				}
+				else
+				{
+					UninstallResult = ExecuteIOSDeployCommand(string.Format("--bundle_id {0} --uninstall_only", App));
+				}
 
-		public IAppInstall CreateAppInstall(UnrealAppConfig AppConfig)
-		{
-			return null;
+				UninstallResult.WaitForExit();
+				if (UninstallResult.ExitCode != 0)
+				{
+					throw new AutomationException("Uninstalling {0} exited with code {1}: {2}", App, UninstallResult.ExitCode, UninstallResult.Output);
+				}
+			}
 		}
 
-		public void CopyAdditionalFiles(IEnumerable<UnrealFileToCopy> FilesToCopy)
+		public void CleanArtifacts()
 		{
+			if(InstallCache == null)
+			{
+				throw new AutomationException("Could not deduce Bundle ID and project name due to invalid installation cache. " +
+					"Either specify the Bundle ID and Project Name using this function's overload OR ensure CreateAppInstall has been called first.");
+			}
 
+			CleanArtifacts(InstallCache.PackageName, InstallCache.ProjectName);
 		}
 
-		public void PopulateDirectoryMappings(string ProjectDir)
+		public void CleanArtifacts(string BundleID, string ProjectName)
 		{
-			LocalDirectoryMappings.Add(EIntendedBaseCopyDirectory.Build, Path.Combine(ProjectDir, "Build"));
-			LocalDirectoryMappings.Add(EIntendedBaseCopyDirectory.Binaries, Path.Combine(ProjectDir, "Binaries"));
-			LocalDirectoryMappings.Add(EIntendedBaseCopyDirectory.Config, Path.Combine(ProjectDir, "Config"));
-            LocalDirectoryMappings.Add(EIntendedBaseCopyDirectory.Content, Path.Combine(ProjectDir, "Content"));
-            LocalDirectoryMappings.Add(EIntendedBaseCopyDirectory.Demos, Path.Combine(ProjectDir, "Demos"));
-			LocalDirectoryMappings.Add(EIntendedBaseCopyDirectory.PersistentDownloadDir, Path.Combine(ProjectDir, "Saved", "PersistentDownloadDir"));
-			LocalDirectoryMappings.Add(EIntendedBaseCopyDirectory.Profiling, Path.Combine(ProjectDir, "Profiling"));
-            LocalDirectoryMappings.Add(EIntendedBaseCopyDirectory.Saved, ProjectDir);
-        }
+			string EngineSavedDirectory = "/Documents/Engine/Saved";
+			string ProjectSavedDirectory = string.Format("/Documents/{0}/Saved", ProjectName);
 
-		public IAppInstance Run(IAppInstall App)
-		{
-			IOSAppInstall IOSApp = App as IOSAppInstall;
+			string CleanEngine = string.Format("--bundle_id {0} --rmtree {1}", BundleID, EngineSavedDirectory);
+			string CleanProject = string.Format("--bundle_id {0} --rmtree {1}", BundleID, ProjectSavedDirectory);
 
-			if (IOSApp == null)
+			IProcessResult Result = ExecuteIOSDeployCommand(CleanEngine);
+			Result.WaitForExit();
+			if (Result.ExitCode != 0)
 			{
-				throw new DeviceException("AppInstance is of incorrect type!");
+				throw new AutomationException("Cleaning Engine artifacts for {0} exited with code {1}:\n{2}", BundleID, Result.ExitCode, Result.Output);
 			}
 
-			string CommandLine = IOSApp.CommandLine.Replace("\"", "\\\\\"");
-
-			Log.Info("Launching {0} on {1}", App.Name, ToString());
-			Log.Verbose("\t{0}", CommandLine);
-
-			// ios-deploy notes: -L launches detached, -I non-interactive  (exits when app exits), -r uninstalls before install (removes app Documents folder)
-			// -t <seconds> number of seconds to wait for device to be connected
-
-			// check xcodebuild -version, to see if it larger than 15.0
-			bool bIsXcode15 = false;
-			IProcessResult XcodeVersionResult = CommandUtils.Run("xcodebuild", "-version");
-			if (XcodeVersionResult.ExitCode == 0)
+			Result = ExecuteIOSDeployCommand(CleanProject);
+			Result.WaitForExit();
+			if (Result.ExitCode != 0)
 			{
-				// XcodeVersionResult.Output is: Xcode 15, try get version number
-				Match VersionMatch = Regex.Match(XcodeVersionResult.Output, @"Xcode\ (?<Version>\d+)");
-				if (VersionMatch.Success)
-				{
-					int Version = 0;
-					if (int.TryParse(VersionMatch.Groups["Version"].Value, out Version))
-					{
-						bIsXcode15 = Version >= 15;
-					}
-				}
+				throw new AutomationException("Cleaning Project artifacts for {0} exited with code {1}:\n{2}", BundleID, Result.ExitCode, Result.Output);
 			}
+		}
 
-			IProcessResult Result;
-			if (bIsXcode15)
-			{
-				CommandLine = CommandLine.Replace("\\\\", "");
-				string CL = String.Format("device process launch --device {0} {1} '{2}'", DeviceName, IOSApp.PackageName, CommandLine.Trim());
-				Result = ExecuteIOSDeployCommand(CL, 0, true, false, true);
-				Thread.Sleep(5000);
-				// check if IOSApp.PackageName process running using devicectl device info processes
-				while(true)
-				{
-					CL = String.Format("device info processes --device {0}", DeviceName);
-					Result = ExecuteIOSDeployCommand(CL, 0, true, false, true);
-					// Get BundleName from IOSApp.PackageName
-					string BundleName = IOSApp.PackageName.Substring(IOSApp.PackageName.LastIndexOf('.') + 1);
-					Log.Info("Process running: " + BundleName);
-					// check if IOSApp.PackageName process not running
-					if (!Result.Output.Contains(BundleName))
-					{
-						break;
-					}
-					Thread.Sleep(5000);
-				}
+		public void InstallBuild(UnrealAppConfig AppConfig)
+		{
+			if(AppConfig.Build is not IOSBuild Build)
+			{
+				throw new AutomationException("Unsupported build type {0} for iOS!", AppConfig.Build.GetType());
+			}
 
+			string BuildPath = Build.SourcePath;
+			string AppName = Path.GetFileName(Build.SourcePath);
+			if (Globals.IsRunningDev && AppConfig.OverlayExecutable.GetOverlay(BuildPath, out string OverlayExecutable))
+			{
+				BuildPath = OverlayExecutable;
 			}
-			else
+			else if(Build.IsIPAFile)
 			{
+				// If the build is an IPA file, unzip it to the cache directory
+				string BuildCache = Path.Combine(LocalCachePath, "TempBuildCache");
+				string PayloadDirectory = Path.Combine(BuildCache, "Payload");
 
-				// setup symbols if available
-				string DSymBundle = "";
-				string DSymDir = Path.Combine(GauntletAppCache, "Symbols");
-
-				if (Directory.Exists(DSymDir))
+				if(Directory.Exists(PayloadDirectory))
 				{
-					DSymBundle = Directory.GetDirectories(DSymDir).Where(D => Path.GetExtension(D).ToLower() == ".dsym").FirstOrDefault();
-					DSymBundle = string.IsNullOrEmpty(DSymBundle) ? "" : DSymBundle = " -S \"" + DSymBundle + "\"";
+					Directory.Delete(PayloadDirectory, true);
 				}
+				Directory.CreateDirectory(BuildCache);
 
-				string CL = "--noinstall -I" + DSymBundle + " -b \"" + LocalAppBundle + "\" --args '" + CommandLine.Trim() + "'";
-
-				Result = ExecuteIOSDeployCommand(CL, 0);
-
-				Thread.Sleep(5000);
+				string IPACommand = string.Format("-x -k {0} {1}", BuildPath, BuildCache);
+				if (!AppleBuild.ExecuteIPADittoCommand(IPACommand, out string Output, PayloadDirectory))
+				{
+					throw new AutomationException("Unable to extract IPA {0}:\n{1}", Build.SourcePath, Output);
+				}
 
-				// Give ios-deploy a chance to throw out any errors...
-				if (Result.HasExited)
+				BuildPath = Directory.EnumerateDirectories(PayloadDirectory, "*.app").FirstOrDefault();
+				if(BuildPath == null)
 				{
-					Log.Warning(KnownLogEvents.Gauntlet_DeviceEvent, "ios-deploy exited early: " + Result.Output);
-					throw new DeviceException("Failed to launch on {0}. {1}", Name, Result.Output);
+					throw new AutomationException("Failed to find app within IPA payload directory!");
 				}
 			}
 
-			return new IOSAppInstance(IOSApp, Result, IOSApp.CommandLine);
-		}
-
-		public IProcessResult ExecuteIOSDeployCommand(String CommandLine, int WaitTime = 60, bool WarnOnTimeout = true, bool UseDeviceID = true, bool UseDeviceCtl = false)
-		{
-			String IOSDeployPath = Path.Combine(Globals.UnrealRootDir, "Engine/Extras/ThirdPartyNotUE/ios-deploy/bin/ios-deploy");
-			if (UseDeviceCtl)
+			// Uninstall any existing form of the app
+			IProcessResult Result;
+			if(UseDeviceCtl)
 			{
-				CommandLine = String.Format("devicectl {0}", CommandLine);
-				IOSDeployPath = "/usr/bin/xcrun";
+				Result = ExecuteDevicectlCommand(string.Format("device uninstall app {0}", Build.PackageName));
 			}
 			else
 			{
-				if (UseDeviceID && !IsDefaultDevice)
-				{
-					CommandLine = String.Format("--id {0} {1}", DeviceName, CommandLine);
-				}
+				Result = ExecuteIOSDeployCommand(string.Format("--bundle_id {0} --uninstall_only", Build.PackageName));
 			}
 
-			if (!File.Exists(IOSDeployPath))
+			Result.WaitForExit();
+			if(Result.ExitCode != 0)
 			{
-				throw new AutomationException("Unable to run ios-deploy binary at {0}", IOSDeployPath);
+				throw new AutomationException("Uninstall exited with code {0}:\n{1}", Result.ExitCode, Result.Output);
 			}
 
-			CommandUtils.ERunOptions RunOptions = CommandUtils.ERunOptions.NoWaitForExit;
-
-			if (Log.IsVeryVerbose)
+			// Install the app
+			// Bad provisions can cause this to spit out errors, ignore these and let higher levels manage the exception's verbosity
+			using(ScopedSuspendECErrorParsing ErrorSuspension = new())
 			{
-				RunOptions |= CommandUtils.ERunOptions.AllowSpew;
+				if(UseDeviceCtl)
+				{
+					Result = ExecuteDevicectlCommand(string.Format("device install app {0} -v", BuildPath), MaxInstallTime);
+				}
+				else
+				{
+					Result = ExecuteIOSDeployCommand(string.Format("-b \"{0}\"", BuildPath), MaxInstallTime);			
+				}
 			}
-			else
+
+			if (Result.ExitCode != 0)
 			{
-				RunOptions |= CommandUtils.ERunOptions.NoLoggingOfRunCommand;
+				throw new AutomationException("Install exited with code {0}:\n{1}", Result.ExitCode, Result.Output);
 			}
-
-			Log.Verbose("ios-deploy executing '{0}'", CommandLine);
-
-			IProcessResult Result = CommandUtils.Run(IOSDeployPath, CommandLine, Options: RunOptions);
-
-			if (WaitTime > 0)
+			
+			if(Build.Flags.HasFlag(BuildFlags.Bulk))
 			{
-				DateTime StartTime = DateTime.Now;
-
-				Result.ProcessObject.WaitForExit(WaitTime * 1000);
+				int PayloadIndex = Build.SourcePath.IndexOf("/Payload");
+				DirectoryInfo RootBuildDirectory = new(Build.SourcePath.Substring(0, PayloadIndex));
+					
+				string PackageDirectory = RootBuildDirectory.Parent.FullName;
+				string BulkBuildPath = RootBuildDirectory.FullName;
+				string BulkContentsPath = Path.Combine(PackageDirectory, "BulkContents");
+				if(!Directory.Exists(BulkContentsPath))
+				{
+					throw new AutomationException("Could not locate a bulk contents directory when attempting to install a bulk build. Is your bulk build configured correctly?");
+				}
 
-				if (Result.HasExited == false)
+				// When targeting a build in the network share, pull the build down to the host first for a faster install.
+				// Note: it's only faster when the device is connected via USB 3.0, so WiFi is actually faster for our device farm because we use 2.0 for host connections
+				// When running on builders, we'll just copy the build directly from the share
+				string BuildVolume = GetVolumeName(PackageDirectory);
+				string CurrentVolume = GetVolumeName(Environment.CurrentDirectory);
+				if(!BuildVolume.Equals(CurrentVolume, StringComparison.OrdinalIgnoreCase) && !IsBuildMachine)
 				{
-					if ((DateTime.Now - StartTime).TotalSeconds >= WaitTime)
+					string BuildCache = Path.Combine(LocalCachePath, "TempBuildCache");
+					string RelativeBuildPath = BulkBuildPath.Replace(PackageDirectory, string.Empty).TrimStart('\\').Trim('/');
+					string RelativeContentPath = BulkContentsPath.Replace(PackageDirectory, string.Empty).Trim('\\').Trim('/');
+
+					string DestinationBuildPath = Path.Combine(BuildCache, RelativeBuildPath);
+					string DestinationContentPath = Path.Combine(BuildCache, RelativeContentPath);
+					
+					lock(CacheMutex)
 					{
-						string Message = String.Format("IOSDeployPath timeout after {0} secs: {1}, killing process", WaitTime, CommandLine);
-
-						if (WarnOnTimeout)
-						{ 
-							Log.Warning(Message);
-						}
-						else
+						if(!HasCopiedNetworkBuild)
 						{
-							Log.Info(Message);
-						}
+							if (Directory.Exists(BuildCache))
+							{
+								Directory.Delete(BuildCache, true);
+							}
+							Directory.CreateDirectory(BuildCache);
 
-						Result.ProcessObject.Kill();
-						// wait up to 15 seconds for process exit
-						Result.ProcessObject.WaitForExit(15000);
+							Log.Info("Copying {Source} to {Dest}", BulkBuildPath, DestinationBuildPath);
+							SystemHelpers.CopyDirectory(BulkBuildPath, DestinationBuildPath, SystemHelpers.CopyOptions.Mirror);
+							Log.Info("Copying {Source} to {Dest}", BulkContentsPath, DestinationContentPath);
+							SystemHelpers.CopyDirectory(BulkContentsPath, DestinationContentPath, SystemHelpers.CopyOptions.Mirror);
+
+							Log.Info("Copy to host complete. Starting deploy to {DeviceName}", UUID);
+
+							// If a device has an issue and we need to select a new one, don't re-copy the whole build from the network again!
+							HasCopiedNetworkBuild = true;
+						}
 					}
+					
+					BulkBuildPath = DestinationBuildPath;
 				}
-			}
 
-			return Result;
-		}
+				string IDeviceFS = Path.Combine(BulkBuildPath, "libimobiledevice", "mac", "idevicefs");
+				string BulkContentCopyCommand = string.Format("-b {0} -x RequiredCommands.txt", Build.PackageName);
+				Result = ExecuteDeploymentCommand(IDeviceFS, BulkContentCopyCommand, MaxInstallTime, true, BulkBuildPath, ERunOptions.NoStdOutRedirect);
 
-		/// <summary>
-		/// Remove the application entirely from the iOS device, this includes any persistent app data in /Documents
-		/// </summary>
-		private void RemoveApplication(IOSBuild Build)
-		{
-			string CommandLine = String.Format("--bundle_id {0} --uninstall_only", Build.PackageName);
-			ExecuteIOSDeployCommand(CommandLine);
+				if(Result.ExitCode != 0)
+				{
+					throw new AutomationException("Copying bulk content failed with exit code {0}. See above for info", Result.ExitCode);
+				}
+			}
 		}
 
-		/// <summary>
-		/// Checks whether version of deployed bundle matches local Build
-		/// </summary>
-		private bool CheckDeployedApp(IOSBuild Build)
+		public IAppInstall CreateAppInstall(UnrealAppConfig AppConfig)
 		{
-			try
+			if (AppConfig.Build is not IOSBuild Build)
 			{
-				Log.Verbose("Checking deployed App hash");
+				throw new AutomationException("Unsupported build type {0} for iOS!", AppConfig.Build.GetType());
+			}
 
-				string CommandLine = String.Format("--bundle_id {0} --download={1} --to {2}", Build.PackageName, "/Documents/AppHash.txt", LocalCachePath);
-				IProcessResult Result = ExecuteIOSDeployCommand(CommandLine, 120);
+			string ProjectDirectory = string.Format("/Documents/{0}", AppConfig.ProjectName);
+			PopulateDirectoryMappings(ProjectDirectory);
 
-				if (Result.ExitCode != 0)
+			string AppPath = null;
+			if(!UseDeviceCtl)
+			{
+				if(Globals.IsRunningDev)
 				{
-					return false;
+					AppConfig.OverlayExecutable.GetOverlay(Build.SourcePath, out AppPath);
 				}
-
-				string Hash = File.ReadAllText(LocalCachePath + "/Documents/AppHash.txt").Trim();
-				string StoredHash = File.ReadAllText(AppHashFilename).Trim();
-
-				if (Hash == StoredHash)
+				else
 				{
-					Log.Verbose("Deployed app hash matched cached App hash");
-					return true;
+					AppPath = Build.SourcePath;
 				}
+			}
 
+			InstallCache = new IOSAppInstall(AppConfig.Name, this, Build.PackageName, AppConfig.CommandLine, AppConfig.ProjectName, AppPath);
+			return InstallCache;
+		}
+
+		public void CopyAdditionalFiles(IEnumerable<UnrealFileToCopy> FilesToCopy)
+		{
+			if(InstallCache == null || string.IsNullOrEmpty(InstallCache.PackageName))
+			{
+				throw new AutomationException("Could not deduce Bundle ID due to invalid installation cache. " +
+					"Either specify the Bundle ID using this function's overload OR ensure CreateAppInstall has been called first.");
 			}
-			catch (Exception Ex)
+
+			CopyAdditionalFiles(FilesToCopy, InstallCache.PackageName);
+		}
+
+		public void CopyAdditionalFiles(IEnumerable<UnrealFileToCopy> FilesToCopy, string BundleID)
+		{
+			if (FilesToCopy == null || !FilesToCopy.Any())
+			{
+				return;
+			}
+
+			if (!LocalDirectoryMappings.Any())
 			{
-				if (!Ex.Message.Contains("is denied"))
+				throw new AutomationException("Attempted to copy additional files before LocalDirectoryMappings were populated." +
+					"{0} must call PopulateDirectoryMappings before attempting to call CopyAdditionalFiles", this.GetType());
+			}
+
+			foreach (UnrealFileToCopy FileToCopy in FilesToCopy)
+			{
+				FileInfo SourceFile = new FileInfo(FileToCopy.SourceFileLocation);
+				FileInfo DestinationFile = new FileInfo(string.Format("/{0}/{1}", LocalDirectoryMappings[FileToCopy.TargetBaseDirectory], FileToCopy.TargetRelativeLocation));
+				if (!SourceFile.Exists)
 				{
-					Log.Verbose("Unable to pull cached App cache from device, cached file may not exist: {0}", Ex.Message);
+					Log.Warning(KnownLogEvents.Gauntlet_DeviceEvent, "File to copy {File} not found. Skipping copy.", FileToCopy);
 				}
-			}
+				SourceFile.IsReadOnly = false;
 
-			Log.Verbose("Deployed app hash doesn't match, App will be installed");
-			return false;
+				Log.Verbose("Copying {Source} to {Destination}", SourceFile, DestinationFile);
 
+				string CopyCommand = string.Format("--bundle_id {0} --upload={1} --to {2}",
+					BundleID, SourceFile, DestinationFile);
+				IProcessResult Result = ExecuteIOSDeployCommand(CopyCommand, 120);
+	
+				if(Result.ExitCode != 0)
+				{
+					if (Result.ExitCode != 0)
+					{
+						throw new AutomationException("Copy of {0} to {1} exited with code {2}:\n{3}",
+							SourceFile, DestinationFile, Result.ExitCode, Result.Output);
+					}
+				}
+			}
 		}
 
-		/// <summary>
-		/// Resign application using local executable and update debug symbols
-		/// </summary>
-		private void ResignApplication(UnrealAppConfig AppConfig)
+		public void PopulateDirectoryMappings(string ProjectDirectory)
 		{
-			// check that we have the signing stuff we need
-			string SignProvision = Globals.Params.ParseValue("signprovision", String.Empty);
-			string SignEntitlements = Globals.Params.ParseValue("signentitlements", String.Empty);
-			string SigningIdentity =  Globals.Params.ParseValue("signidentity", String.Empty);
+			string SavedDirectory = ProjectDirectory + "/Saved";
+
+			LocalDirectoryMappings.Add(EIntendedBaseCopyDirectory.Build, ProjectDirectory + "/Build");
+			LocalDirectoryMappings.Add(EIntendedBaseCopyDirectory.Binaries, ProjectDirectory + "/Binaries");
+			LocalDirectoryMappings.Add(EIntendedBaseCopyDirectory.Config, SavedDirectory + "/Config");
+            LocalDirectoryMappings.Add(EIntendedBaseCopyDirectory.Content, ProjectDirectory + "/Content");
+            LocalDirectoryMappings.Add(EIntendedBaseCopyDirectory.Demos, SavedDirectory + "/Demos");
+			LocalDirectoryMappings.Add(EIntendedBaseCopyDirectory.PersistentDownloadDir, SavedDirectory + "/PersistentDownloadDir");
+			LocalDirectoryMappings.Add(EIntendedBaseCopyDirectory.Profiling, SavedDirectory + "/Profiling");
+            LocalDirectoryMappings.Add(EIntendedBaseCopyDirectory.Saved, SavedDirectory);
+        }
 
-			// handle signing provision
-			if (string.IsNullOrEmpty(SignProvision) || !File.Exists(SignProvision))
+		public IAppInstance Run(IAppInstall Install)
+		{
+			if (Install is not IOSAppInstall IOSInstall)
 			{
-				throw new AutomationException("Absolute path to existing provision must be specified, example: -signprovision=/path/to/myapp.provision");
+				throw new DeviceException("AppInstance is of incorrect type!");
 			}
 
-			// handle entitlements
-			// Note this extracts entitlements: which may be useful when using same provision/entitlements?: codesign -d --entitlements :entitlements.plist ~/.gauntletappcache/Payload/Example.app/
+			// Copy the commandline into a uecommandline.txt
+			string CommandLine = IOSInstall.CommandLine.Trim();
+			string TempCommandLine = Path.GetTempFileName();
+			File.WriteAllText(TempCommandLine, CommandLine);
 
-			if (string.IsNullOrEmpty(SignEntitlements) || !File.Exists(SignEntitlements))
+			IProcessResult CommandLineCopyResult;
+			if(UseDeviceCtl)
 			{
-				throw new AutomationException("Absolute path to existing entitlements must be specified, example: -signprovision=/path/to/entitlements.plist");
+				string CopyCommand = string.Format("device copy to --device {0} --source {1} --destination {2} --domain-type appDataContainer --domain-identifier {3}", 
+					UUID, TempCommandLine, Path.Combine("Documents", "uecommandline.txt"), IOSInstall.PackageName);
+				CommandLineCopyResult = ExecuteDevicectlCommand(CopyCommand, UseDeviceID: false);
+			}
+			else
+			{
+				string CopyCommand = string.Format("--bundle_id {0} --upload={1} --to {2}",
+					IOSInstall.PackageName, TempCommandLine, Path.Combine("/Documents", "uecommandline.txt"));
+				CommandLineCopyResult = ExecuteIOSDeployCommand(CopyCommand);
 			}
 
-			// signing identity
-			if (string.IsNullOrEmpty(SigningIdentity))
+			if(CommandLineCopyResult.ExitCode != 0)
 			{
-				throw new AutomationException("Signing identity must be specified, example: -signidentity=\"iPhone Developer: John Smith\"");
+				throw new AutomationException("Failed to copy uecommandline file to device!");
 			}
 
-			string ProjectName = AppConfig.ProjectName;
-			string BundleName = Path.GetFileNameWithoutExtension(LocalAppBundle);
-			string ExecutableName = UnrealHelpers.GetExecutableName(ProjectName, UnrealTargetPlatform.IOS, AppConfig.Configuration, AppConfig.ProcessType, "");
-			string CachedAppPath = Path.Combine(GauntletAppCache, "Payload", string.Format("{0}.app", BundleName));
+			Log.Info("Launching {0} on {1}", Install.Name, this);
+			Log.Verbose(IOSInstall.CommandLine);
 
-			string LocalExecutable = Path.Combine(Environment.CurrentDirectory, ProjectName, string.Format("Binaries/IOS/{0}", ExecutableName));
-			if (!File.Exists(LocalExecutable))
+			IProcessResult AppProcess;
+			if(UseDeviceCtl)
 			{
-				throw new AutomationException("Local executable not found for -dev argument: {0}", LocalExecutable);
+				// device process launch	base command
+				// --terminate-existing		exit any existing forms of the app
+				// --console				pipe process stdout to devicectl, also holds onto the process until exit instead of immediate exit
+				// --device					device to run on
+				string LaunchCommand = string.Format("device process launch --terminate-existing --console --device {0} {1}", UUID, IOSInstall.PackageName);
+				AppProcess = ExecuteDevicectlCommand(LaunchCommand, 0, UseDeviceID: false, AdditionalOptions: ERunOptions.SpewIsVerbose);
 			}
+			else
+			{
+				if (IOSInstall.AppPath == null)
+				{
+					// Sanity check, this should never happen, but it would be tricky to debug without an explicit message if it did!
+					throw new AutomationException("The provided install to run has a null app path, this likely means the install was designed " +
+						"to run with devicectl instead of ios-deploy.");
+				}
 
-			File.WriteAllText(CacheResignedFilename, "The application has been resigned");
-
-			// copy local executable
-			FileInfo SrcInfo = new FileInfo(LocalExecutable);
-			string DestPath = Path.Combine(CachedAppPath, BundleName);
-			SrcInfo.CopyTo(DestPath, true);
-			Log.Verbose("Copied local executable from {0} to {1}", LocalExecutable, DestPath);
+				// -m		Skip installation
+				// -I		Non-interactive mode, just launch and end launch process when app process ends
+				// -b		path to bundle to run
+				string LaunchCommand = string.Format("-m -I -b \"{0}\"", IOSInstall.AppPath);
+				AppProcess = ExecuteIOSDeployCommand(LaunchCommand, 0);
+			}
 
-			// copy provision
-			SrcInfo = new FileInfo(SignProvision);
-			DestPath = Path.Combine(CachedAppPath, "embedded.mobileprovision");
-			SrcInfo.CopyTo(DestPath, true);
-			Log.Verbose("Copied provision from {0} to {1}", SignProvision, DestPath);
+			if (AppProcess.HasExited)
+			{
+				throw new DeviceException("Failed to launch on {0}. {1}", Name, AppProcess.Output);
+			}
 
-			// handle symbols
-			string LocalSymbolsDir = Path.Combine(Environment.CurrentDirectory, ProjectName, string.Format("Binaries/IOS/{0}.dSYM", ExecutableName));
-			DestPath = Path.Combine(GauntletAppCache, string.Format("Symbols/{0}.dSYM", ExecutableName));
+			return new IOSAppInstance(IOSInstall, AppProcess, IOSInstall.CommandLine);
+		}
 
-			if (Directory.Exists(DestPath))
+		public static IProcessResult ExecuteIOSDeployCommand(string CommandLine, string Device, int WaitTime = 60, bool WarnOnTimeout = true, ERunOptions AdditionalOptions = ERunOptions.None)
+		{
+			string IOSDeploy = Path.Combine(Globals.UnrealRootDir, "Engine/Extras/ThirdPartyNotUE/ios-deploy/bin/ios-deploy");
+			if (!string.IsNullOrEmpty(Device))
 			{
-				Directory.Delete(DestPath, true);
+				CommandLine = string.Format("{0} -i {1}", CommandLine, Device);
 			}
 
-			if (Directory.Exists(LocalSymbolsDir))
+			return ExecuteDeploymentCommand(IOSDeploy, CommandLine, WaitTime, WarnOnTimeout, AdditionalOptions: AdditionalOptions);
+		}
+
+		public IProcessResult ExecuteIOSDeployCommand(string CommandLine, int WaitTime = 60, bool WarnOnTimeout = true, bool UseDeviceID = true, ERunOptions AdditionalOptions = ERunOptions.None)
+		{
+			if (UseDeviceID && !IsDefaultDevice)
 			{
-				CommandUtils.CopyDirectory_NoExceptions(LocalSymbolsDir, DestPath, true);
+				return ExecuteIOSDeployCommand(CommandLine, UUID, WaitTime, WarnOnTimeout, AdditionalOptions);
 			}
 			else
 			{
-				Log.Warning(KnownLogEvents.Gauntlet_DeviceEvent, "No symbols found for local build at {Directory}, removing cached app symbols", LocalSymbolsDir);
-			}
-
-			// resign application
-			// @todo: this asks for password unless "Always Allow" is selected, also for builders, document how to permanently grant codesign access to keychain
-			string SignArgs = string.Format("-f -s \"{0}\" --entitlements \"{1}\" \"{2}\"", SigningIdentity, SignEntitlements, CachedAppPath);
-			Log.Info("\nResigning app, please enter keychain password if prompted:\n\ncodesign {0}", SignArgs);
-			var Result = IOSBuild.ExecuteCommand("codesign", SignArgs);
-			if (Result.ExitCode != 0)
-			{
-				throw new AutomationException("Failed to resign application");
+				return ExecuteIOSDeployCommand(CommandLine, null, WaitTime, WarnOnTimeout, AdditionalOptions);
 			}
 		}
 
-		/// <summary>
-		/// Get UUID of all connected iOS devices
-		/// </summary>
-		private List<string> GetConnectedDeviceUUID()
+		public static IProcessResult ExecuteDevicectlCommand(string CommandLine, string Device, int WaitTime = 60, bool WarnOnTimeout = true, ERunOptions AdditionalOptions = ERunOptions.None)
 		{
-			var Result = ExecuteIOSDeployCommand("--detect", 60, true, false);
+			CommandLine = string.Format("devicectl {0}", CommandLine);
 
-			if (Result.ExitCode != 0)
+			string XCRun = "/usr/bin/xcrun";
+			if (!string.IsNullOrEmpty(Device))
 			{
-				return new List<string>();
+				CommandLine = string.Format("{0} --device {1}", CommandLine, Device);
 			}
 
-			MatchCollection DeviceMatches = Regex.Matches(Result.Output, @"(.?)Found\ ([a-z0-9]{40}|[A-Z0-9]{8}-[A-Z0-9]{16})");
-
-			return DeviceMatches.Cast<Match>().Select<Match, string>(
-				M => M.Groups[2].ToString()
-			).ToList();
+			return ExecuteDeploymentCommand(XCRun, CommandLine, WaitTime, WarnOnTimeout, AdditionalOptions: AdditionalOptions);
 		}
 
-		// Get rid of any zombie lldb/iosdeploy processes, this needs to be reworked to use tracked process id's when running parallel tests across multiple AutomationTool.exe processes on test workers
-		private void KillZombies()
+		public IProcessResult ExecuteDevicectlCommand(string CommandLine, int WaitTime = 60, bool WarnOnTimeout = true, bool UseDeviceID = true, ERunOptions AdditionalOptions = ERunOptions.None)
 		{
+			if(UseDeviceID)
+			{
+				return ExecuteDevicectlCommand(CommandLine, UUID, WaitTime, WarnOnTimeout, AdditionalOptions);
+			}
+			else
+			{
+				return ExecuteDevicectlCommand(CommandLine, null, WaitTime, WarnOnTimeout, AdditionalOptions);
+			}
+		}
 
-			if (ZombiesKilled)
+		private static IProcessResult ExecuteDeploymentCommand(string Executable, string CommandLine, int WaitTime = 60, bool WarnOnTimeout = true, string WorkingDir = null, ERunOptions AdditionalOptions = ERunOptions.None)
+		{
+			if (!AdditionalOptions.HasFlag(ERunOptions.UseShellExecute) && !File.Exists(Executable))
 			{
-				return;
+				throw new AutomationException("Unable to find deployment binary at {0}", Executable);
 			}
 
-			ZombiesKilled = true;
+			Log.Info("{DeploymentExecutable} executing '{Command}'", Executable, CommandLine);
 
-			IOSBuild.ExecuteCommand("killall", "ios-deploy");
-			Thread.Sleep(2500);
-			IOSBuild.ExecuteCommand("killall", "lldb");
-			Thread.Sleep(2500);
-		}
+			ERunOptions RunOptions = ERunOptions.NoWaitForExit | AdditionalOptions;
+			RunOptions |= Log.IsVeryVerbose
+				? ERunOptions.AllowSpew
+				: ERunOptions.NoLoggingOfRunCommand;
 
-		/// <summary>
-		/// Generate MD5 and cache IPA bundle files
-		/// </summary>
-		private bool PrepareIPA(IOSBuild Build)
-		{
-			Log.Info("Preparing IPA {0}", Build.SourcePath);
 
+			IProcessResult Result = null;
 			try
 			{
-				// cache the unzipped app using a MD5 checksum, avoiding needing to unzip
-				string Hash = null;
-				string StoredHash = null;
-				using (var MD5Hash = MD5.Create())
-				{
-					using (var Stream = File.OpenRead(Build.SourcePath))
-					{
-						Hash = BitConverter.ToString(MD5Hash.ComputeHash(Stream)).Replace("-", "").ToLowerInvariant();
-					}
-				}
-				string PayloadDir = Path.Combine(GauntletAppCache, "Payload");
-				string SymbolsDir = Path.Combine(GauntletAppCache, "Symbols");
+				Result = CommandUtils.Run(Executable, CommandLine, Options: RunOptions, WorkingDir: WorkingDir);
 
-				if (File.Exists(AppHashFilename) && Directory.Exists(PayloadDir))
+				if (WaitTime > 0)
 				{
-					StoredHash = File.ReadAllText(AppHashFilename).Trim();
-					if (Hash != StoredHash)
-					{
-						Log.Verbose("App hash out of date, clearing cache");
-						StoredHash = null;
-					}
-				}
+					DateTime StartTime = DateTime.Now;
 
-				if (String.IsNullOrEmpty(StoredHash) || Hash != StoredHash)
-				{
-					if (Directory.Exists(PayloadDir))
-					{
-						Directory.Delete(PayloadDir, true);
-					}
+					Result.ProcessObject.WaitForExit(WaitTime * 1000);
 
-					if (Directory.Exists(SymbolsDir))
+					if (Result.HasExited == false)
 					{
-						Directory.Delete(SymbolsDir, true);
-					}
-
-					if (File.Exists(CacheResignedFilename))
-					{
-						File.Delete(CacheResignedFilename);
-					}
+						if ((DateTime.Now - StartTime).TotalSeconds >= WaitTime)
+						{
+							string Message = string.Format("{0} timeout after {1} secs: {2}, killing process", Executable, WaitTime, CommandLine);
 
-					Log.Verbose("Unzipping IPA {0} to cache at: {1}", Build.SourcePath, GauntletAppCache);
+							if (WarnOnTimeout)
+							{
+								Log.Warning(Message);
+							}
+							else
+							{
+								Log.Info(Message);
+							}
 
-					string Output;
-					if (!IOSBuild.ExecuteIPADittoCommand(String.Format("-x -k {0} {1}", Build.SourcePath, GauntletAppCache), out Output, PayloadDir))
-					{
-						throw new Exception(String.Format("Unable to extract IPA {0}", Build.SourcePath));
+							Result.ProcessObject?.Kill();
+							Result.ProcessObject?.WaitForExit(15000);
+						}
 					}
+				}
+			}
+			catch (Exception Ex)
+			{
+				Log.Warning("Encountered an {ExceptionType} while running device command {Exception}", Ex.GetType().Name, Ex.Message);
+			}
 
-					// Cache symbols for symbolicated callstacks
-					string SymbolsZipFile = string.Format("{0}/../../Symbols/{1}.dSYM.zip", Path.GetDirectoryName(Build.SourcePath), Path.GetFileNameWithoutExtension(Build.SourcePath));
-
-					Log.Verbose("Checking Symbols at {0}", SymbolsZipFile);
+			return Result;
+		}
 
-					if (File.Exists(SymbolsZipFile))
+		private static void InitializeDevices()
+		{
+			lock (InitializerMutex)
+			{
+				if (bStaticInitialized)
+				{
+					if (!ConnectedDevices.Any())
 					{
-						Log.Verbose("Unzipping Symbols {0} to cache at: {1}", SymbolsZipFile, SymbolsDir);
-
-						if (!IOSBuild.ExecuteIPAZipCommand(String.Format("{0} -d {1}", SymbolsZipFile, SymbolsDir), out Output, SymbolsDir))
-						{
-							throw new Exception(String.Format("Unable to extract build symbols {0} -> {1}", SymbolsZipFile, SymbolsDir));
-						}
+						throw new AutomationException("No iOS device connections were detected! Verify the host's setup.");
 					}
 
-					// store hash
-					File.WriteAllText(AppHashFilename, Hash);
-
-					Log.Verbose("App cached");
-				}
-				else
-				{
-					Log.Verbose("Using cached App");
+					return;
 				}
 
-				LocalAppBundle = Directory.GetDirectories(PayloadDir).Where(D => Path.GetExtension(D) == ".app").FirstOrDefault();
+				bStaticInitialized = true;
 
-				if (string.IsNullOrEmpty(LocalAppBundle))
+				// Determine which version of XCode is installed
+				// On XCode version 15 or greater, devicectl will be used to manage deployments
+				// On XCode version 14 or lower, ios-deploy will be used to manage deployments
+				IProcessResult VersionResult = CommandUtils.Run("xcodebuild", "-version");
+				if (VersionResult.ExitCode != 0)
 				{
-					throw new Exception(string.Format("Unable to find app in local app bundle {0}", PayloadDir));
+					throw new AutomationException("XCode version query exited with code {0}. Do you have XCode installed?\n{1}", VersionResult.ExitCode, VersionResult.Output);
 				}
 
-			}
-			catch (Exception Ex)
-			{
-				throw new AutomationException("Unable to prepare {0} : {1}", Build.SourcePath, Ex.Message);
-			}
+				Match VersionMatch = Regex.Match(VersionResult.Output, @"Xcode\ (?<Version>\d+)");
+				if (!VersionMatch.Success)
+				{
+					throw new AutomationException("Failed to match version output, has Apple changed output formatting?");
+				}
 
-			return true;
-		}
+				XCodeMajorVersion = int.Parse(VersionMatch.Groups[1].Value);
 
-		#region Legacy Implementations
+				UseDeviceCtl = XCodeMajorVersion >= 15;
 
-		/// <summary>
-		/// Remove artifacts from device
-		/// </summary>
-		private bool CleanDeviceArtifacts(IOSBuild Build)
-		{			
-			try
-			{
-				Log.Verbose("Cleaning device artifacts");
+				// Gather the iOS devices are connected to this UAT host
+				ConnectedDevices = GetConnectedDevices();
 
-				string CleanCommand = String.Format("--bundle_id {0} --rmtree {1}", Build.PackageName, DeviceArtifactPath);
-				IProcessResult Result = ExecuteIOSDeployCommand(CleanCommand, 120);
+				// Clear any zombie processes previous UAT instances may have left running
+				AppleBuild.ExecuteCommand("killall", "xcrun").WaitForExit();
+				AppleBuild.ExecuteCommand("killall", "ios-deploy").WaitForExit();
+				AppleBuild.ExecuteCommand("killall", "lldb").WaitForExit();
 
-				if (Result.ExitCode != 0)
+				if (!ConnectedDevices.Any())
 				{
-					Log.Warning(KnownLogEvents.Gauntlet_DeviceEvent, "Failed to clean artifacts from device");
-					return false;
+					throw new AutomationException("No iOS device connections were detected! Verify the host's setup.");
 				}
-
-			}
-			catch (Exception Ex)
-			{
-				Log.Verbose("Exception while cleaning artifacts from device: {0}", Ex.Message);
 			}
-
-			return true;
-
 		}
 
-		public IAppInstall InstallApplication(UnrealAppConfig AppConfig)
+		private static Dictionary<string, ConnectedDevice> GetConnectedDevices()
 		{
-			IOSBuild Build = AppConfig.Build as IOSBuild;
-
-			// Ensure Build exists
-			if (Build == null)
-			{
-				throw new AutomationException("Invalid build for IOS!");
-			}
+			Dictionary<string, ConnectedDevice> Devices = new();
 
-			if (!AppConfig.SkipInstall)
+			if (UseDeviceCtl)
 			{
-				bool CacheResigned = false;
-				bool UseLocalExecutable = Globals.Params.ParseParam("dev");
+				string ListFile = Path.GetTempFileName();
+				string ListCommand = string.Format("list devices -j {0}", ListFile);
 
-				lock (AppLock)
+				IProcessResult ListResult = ExecuteDevicectlCommand(ListCommand, null);
+				if (ListResult.ExitCode != 0)
 				{
-					Log.Info("Installing using build source {0}", Build.SourcePath);
-
-					// device artifact path
-					DeviceArtifactPath = string.Format("/Documents/{0}/Saved", AppConfig.ProjectName);
-
-					CacheResigned = File.Exists(CacheResignedFilename);
-
-					if (CacheResigned && !UseLocalExecutable)
-					{
-						if (File.Exists(AppHashFilename))
-						{
-							Log.Verbose("App was resigned, invalidating app cache");
-							File.Delete(AppHashFilename);
-						}
-					}
-
-					if (Build.IsIPAFile)
-					{
-						PrepareIPA(Build);
-					}
-					else
-					{
-						LocalAppBundle = Build.SourcePath;
-					}
-
-					// local executable support
-					if (UseLocalExecutable)
-					{
-						ResignApplication(AppConfig);
-					}
+					Log.Error("List devices exited with code {0}:\n{1}", ListResult.ExitCode, ListResult.Output);
+					return null;
 				}
 
-				if (CacheResigned || UseLocalExecutable || !CheckDeployedApp(Build))
-				{
-					// uninstall will clean all device artifacts
-					IProcessResult Result = ExecuteIOSDeployCommand(String.Format("--uninstall -b \"{0}\"", LocalAppBundle), 20 * 60);
-					Log.Info(Result.Output);
-				}
-				else
+				if(!File.Exists(ListFile))
 				{
-					// remove device artifacts
-					CleanDeviceArtifacts(Build);
+					Log.Error("Could not find list device output file at {FilePath}", ListFile);
+					return null;
 				}
 
-				// parallel iOS tests use same app install folder, so lock it as setup is quick
-				lock (Globals.MainLock)
-				{
-					// local app install with additional files, this directory will be mirrored to device in a single operation
-					string AppInstallPath;
-
-					AppInstallPath = Path.Combine(Globals.TempDir, "iOSAppInstall");
-
-					if (Directory.Exists(AppInstallPath))
-					{
-						Directory.Delete(AppInstallPath, true);
-					}
-
-					Directory.CreateDirectory(AppInstallPath);
-
-					if (LocalDirectoryMappings.Count == 0)
-					{
-						PopulateDirectoryMappings(AppInstallPath);
-					}
-
-					//@todo: Combine Build and AppConfig files, this should be done in higher level code, not per device implementation
-
-					if (AppConfig.FilesToCopy != null)
-					{
-						foreach (UnrealFileToCopy FileToCopy in AppConfig.FilesToCopy)
-						{
-							string PathToCopyTo = Path.Combine(LocalDirectoryMappings[FileToCopy.TargetBaseDirectory], FileToCopy.TargetRelativeLocation);
-
-							if (File.Exists(FileToCopy.SourceFileLocation))
-							{
-								FileInfo SrcInfo = new FileInfo(FileToCopy.SourceFileLocation);
-								SrcInfo.IsReadOnly = false;
-								string DirectoryToCopyTo = Path.GetDirectoryName(PathToCopyTo);
-								if (!Directory.Exists(DirectoryToCopyTo))
-								{
-									Directory.CreateDirectory(DirectoryToCopyTo);
-								}
-								if (File.Exists(PathToCopyTo))
-								{
-									FileInfo ExistingFile = new FileInfo(PathToCopyTo);
-									ExistingFile.IsReadOnly = false;
-								}
-
-								SrcInfo.CopyTo(PathToCopyTo, true);
-								Log.Verbose("Copying app install: {0} to {1}", FileToCopy, DirectoryToCopyTo);
-							}
-							else
-							{
-								Log.Warning(KnownLogEvents.Gauntlet_DeviceEvent, "File to copy {File} not found", FileToCopy);
-							}
-						}
-					}
+				string Text = File.ReadAllText(ListFile);
+				JsonObject[] JsonDevices = JsonObject.Parse(File.ReadAllText(ListFile))
+					.GetObjectField("result")
+					.GetObjectArrayField("devices");
 
-					// copy mapped files in a single pass
-					string CopyCommand = String.Format("--bundle_id {0} --upload={1} --to {2}", Build.PackageName, AppInstallPath, DeviceArtifactPath);
-					ExecuteIOSDeployCommand(CopyCommand, 120);
+				foreach (JsonObject JsonDevice in JsonDevices)
+				{
+					string UUID = JsonDevice.GetObjectField("hardwareProperties").GetStringField("udid");
+					string OSVersionString = JsonDevice.GetObjectField("deviceProperties").GetStringField("osVersionNumber");
+					int IOSMajorVersion = int.Parse(OSVersionString.Substring(0, OSVersionString.IndexOf('.')));
+					bool bIsPaired = JsonDevice.GetObjectField("connectionProperties").GetStringField("pairingState").Equals("paired", StringComparison.OrdinalIgnoreCase);
 
-					// store the App hash to avoid redundant deployments
-					if (File.Exists(AppHashFilename))
-					{
-						CopyCommand = String.Format("--bundle_id {0} --upload={1} --to {2}", Build.PackageName, AppHashFilename, "/Documents/AppHash.txt");
-						ExecuteIOSDeployCommand(CopyCommand, 120);
-					}
+					ConnectedDevice Device = new ConnectedDevice(UUID, IOSMajorVersion, bIsPaired);
+					Devices.Add(UUID, Device);
 				}
 			}
 			else
 			{
-				// Skipping installation
-				Log.Info("Skipping installation as per config.");
-				if (Build.IsIPAFile)
+				IProcessResult DetectResult = ExecuteIOSDeployCommand("--detect", null);
+				if (DetectResult.ExitCode != 0)
 				{
-					string PayloadDir = Path.Combine(GauntletAppCache, "Payload");
-					LocalAppBundle = Directory.GetDirectories(PayloadDir).Where(D => Path.GetExtension(D) == ".app").FirstOrDefault();
-					if (string.IsNullOrEmpty(LocalAppBundle))
-					{
-						throw new AutomationException($"Could not find local IPA payload at {PayloadDir}");
-					}
+					Log.Warning("Detect devices exited with code {0}:\n{1}", DetectResult.ExitCode, DetectResult.Output);
+					return null;
 				}
-				else
+
+				IEnumerable<Match> DeviceMatches = Regex
+					.Matches(DetectResult.Output, @"(.?)Found\ ([a-z0-9]{40}|[A-Z0-9]{8}-[A-Z0-9]{16})(.*?)(\d+\.)")
+					.Where(Match => Match.Success);
+
+				foreach(Match Match in DeviceMatches)
 				{
-					LocalAppBundle = Build.SourcePath;
+					string UUID = Match.Groups[2].ToString();
+					string OSVersionString = Match.Groups[4].ToString();
+					int IOSMajorVersion = int.Parse(OSVersionString.Substring(0, OSVersionString.IndexOf('.')));
+
+					ConnectedDevice Device = new ConnectedDevice(UUID, IOSMajorVersion, true);
+					Devices.Add(UUID, Device);
 				}
 			}
 
-			IOSAppInstall IOSApp = new IOSAppInstall(AppConfig.Name, this, Build.PackageName, AppConfig.CommandLine);
-			return IOSApp;
+			return Devices;
+		}
+
+		private static string GetVolumeName(string InPath)
+		{
+			Match M = Regex.Match(InPath, @"/Volumes/(.+?)/");
+
+			if (M.Success)
+			{
+				return M.Groups[1].ToString();
+			}
+
+			return string.Empty;
+		}
+
+		// DEPRECATED
+		public IAppInstall InstallApplication(UnrealAppConfig AppConfig)
+		{
+			InstallBuild(AppConfig);
+			return CreateAppInstall(AppConfig);
 		}
-		#endregion
 	}
 
 	class IOSAppInstall : IAppInstall
 	{
 		public string Name { get; protected set; }
 
+		public string AppPath { get; protected set; }
+
 		public string CommandLine { get; protected set; }
 
 		public string PackageName { get; protected set; }
 
+		public string ProjectName { get; protected set; }
+
 		public ITargetDevice Device => IOSDevice;
 
 		public TargetDeviceIOS IOSDevice;
 
-		public IOSAppInstall(string InName, TargetDeviceIOS InDevice, string InPackageName, string InCommandLine)
+		public IOSAppInstall(string InName, TargetDeviceIOS InDevice, string InPackageName,
+			string InCommandLine, string InProjectName, string InAppPath = null)
 		{
 			Name = InName;
+			AppPath = InAppPath;
+			IOSDevice = InDevice;
 			CommandLine = InCommandLine;
 			PackageName = InPackageName;
-
-			IOSDevice = InDevice;
+			ProjectName = InProjectName;
 		}
 
 		public IAppInstance Run()
@@ -926,11 +994,12 @@ namespace Gauntlet
 	class IOSAppInstance : IAppInstance
 	{
 		public ITargetDevice Device => Install.Device;
-		public IProcessResult ProcessResult { get; private set; }
+		public TargetDeviceIOS IOSDevice => Device as TargetDeviceIOS;
+		public IProcessResult LaunchProcess { get; private set; }
 		public string CommandLine { get; private set; }
 		public bool WasKilled { get; protected set; }
-		public int ExitCode => ProcessResult.ExitCode;
-		public bool HasExited => ProcessResult.HasExited;
+		public int ExitCode => LaunchProcess.ExitCode;
+		public bool HasExited => LaunchProcess.HasExited;
 
 		public string ArtifactPath
 		{
@@ -945,7 +1014,7 @@ namespace Gauntlet
 					}
 				}
 
-				return Install.IOSDevice.LocalCachePath + "/" + Install.IOSDevice.DeviceArtifactPath;
+				return Path.Combine(IOSDevice.LocalCachePath, "Saved");
 			}
 		}
 
@@ -958,16 +1027,15 @@ namespace Gauntlet
 					// The ios application is being run under lldb by ios-deploy
 					// lldb catches crashes and we have it setup to dump thread callstacks
 					// parse any crash dumps into Unreal crash format and append to output
-					string CrashLog = LLDBCrashParser.GenerateCrashLog(ProcessResult.Output);
+					string CrashLog = LLDBCrashParser.GenerateCrashLog(LaunchProcess.Output);
 
 					if (!string.IsNullOrEmpty(CrashLog))
 					{
-						return String.Format("{0}\n{1}", ProcessResult.Output, CrashLog);
+						return string.Format("{0}\n{1}", LaunchProcess.Output, CrashLog);
 					}
 				}
 
-				return ProcessResult.Output;
-
+				return LaunchProcess.Output;
 			}
 		}
 
@@ -979,14 +1047,14 @@ namespace Gauntlet
 		{
 			Install = InInstall;
 			this.CommandLine = InCommandLine;
-			this.ProcessResult = InProcess;
+			this.LaunchProcess = InProcess;
 		}
 
 		public int WaitForExit()
 		{
 			if (!HasExited)
 			{
-				ProcessResult.WaitForExit();
+				LaunchProcess.WaitForExit();
 			}
 
 			return ExitCode;
@@ -997,22 +1065,53 @@ namespace Gauntlet
 			if (!HasExited)
 			{
 				WasKilled = true;
-				ProcessResult.ProcessObject.Kill(true);
+				LaunchProcess.ProcessObject.Kill(true);
 			}
 		}
 
 		protected void SaveArtifacts()
 		{
-			TargetDeviceIOS Device = Install.IOSDevice;
+			IProcessResult DownloadResult;
+			string SourceArtifactPath = IOSDevice.GetPlatformDirectoryMappings()[EIntendedBaseCopyDirectory.Saved];
+			string DestinationArtifactPath = Path.Combine(IOSDevice.LocalCachePath, "Saved");
 
-			// copy remote artifacts to local
-			string CommandLine = String.Format("--bundle_id {0} --download={1} --to {2}", Install.PackageName, Device.DeviceArtifactPath, Device.LocalCachePath);
+			if (Directory.Exists(DestinationArtifactPath))
+			{
+				Directory.Delete(DestinationArtifactPath, true);
+			}
 
-			IProcessResult DownloadCmd = Device.ExecuteIOSDeployCommand(CommandLine, 120);
+			// ios-deploy will copy the entire Documents directory, we just want the contents
+			// We'll copy the directory to a temp path before moving the contents into the saved directory
+			string TempPath = Path.Combine(IOSDevice.LocalCachePath, "Temp");
+			if (Directory.Exists(TempPath))
+			{
+				Directory.Delete(TempPath, true);
+			}
+			Directory.CreateDirectory(TempPath);
 
-			if (DownloadCmd.ExitCode != 0)
+			string CopyCommand = string.Format("--bundle_id {0} --download={1} --to {2}",
+				Install.PackageName, SourceArtifactPath, TempPath);
+			DownloadResult = IOSDevice.ExecuteIOSDeployCommand(CopyCommand, 120);
+			if (DownloadResult.ExitCode != 0)
+			{
+				string Error = string.Format("Copying artifacts from device failed with code {0}:\n{1}\nArtifacts will not be saved.",
+					DownloadResult.ExitCode, DownloadResult.Output);
+				Log.Error(KnownLogEvents.Gauntlet_DeviceEvent, Error);
+			}
+
+			try
+			{
+				string DocumentsPath = Path.Combine(TempPath, "Documents");
+				if (Directory.Exists(DocumentsPath))
+				{
+					Directory.Move(DocumentsPath, DestinationArtifactPath);
+				}
+
+				Directory.Delete(TempPath, true);
+			}
+			catch (Exception Ex)
 			{
-				Log.Warning(KnownLogEvents.Gauntlet_DeviceEvent, "Failed to retrieve artifacts. {Output}", DownloadCmd.Output);
+				Log.Error(KnownLogEvents.Gauntlet_DeviceEvent, "Failed to move artifacts out of temp directory:\n{Exception}", Ex.Message);
 			}
 		}
 	}
@@ -1036,9 +1135,9 @@ namespace Gauntlet
 			public override string ToString()
 			{
 				// symbolicated
-				if (!String.IsNullOrEmpty(Source))
+				if (!string.IsNullOrEmpty(Source))
 				{
-					return string.Format("Error: [Callstack] 0x{0} {1}!{2} [{3}{4}]", Address, Module, Symbol.Replace(" ", "^"), Source, String.IsNullOrEmpty(Line) ? "" : ":" + Line);
+					return string.Format("Error: [Callstack] 0x{0} {1}!{2} [{3}{4}]", Address, Module, Symbol.Replace(" ", "^"), Source, string.IsNullOrEmpty(Line) ? "" : ":" + Line);
 				}
 
 				// unsymbolicated
