- Added functionality to ExportAppleBuilds script:
  - Makes use of the AppStoreConnect API, which requires a .p8 file to have been downloaded from Apple's Integrations tab
  - Can check if the version already exists on TestFlight, and can (on a Mac) increment the verrsion of the .xcarchive to be a valid version
  - Can wait until the upload has completed and set the WhatsNew field in the TF build, for automated testing notes

diff --git a/Engine/Source/Programs/AutomationTool/Apple/Apple.Automation.csproj b/Engine/Source/Programs/AutomationTool/Apple/Apple.Automation.csproj
--- a/Engine/Source/Programs/AutomationTool/Apple/Apple.Automation.csproj
+++ b/Engine/Source/Programs/AutomationTool/Apple/Apple.Automation.csproj
@@ -35,6 +35,14 @@
     </Compile>
   </ItemGroup>
   <ItemGroup>
+
+	<PackageReference Include="Microsoft.IdentityModel.Tokens" Version="6.34.0" />
+	<PackageReference Include="System.IdentityModel.Tokens.Jwt" Version="6.34.0" />
+	<Reference Include="BouncyCastle.Crypto, Version=1.8.3.0, Culture=neutral, PublicKeyToken=0e99375e54769942">
+		<HintPath>..\..\..\..\Binaries\ThirdParty\IOS\BouncyCastle.Crypto.dll</HintPath>
+	</Reference>
+	<PackageReference Include="Portable.BouncyCastle" Version="1.8.6.7" />
+
     <ProjectReference Include="..\..\Shared\EpicGames.Core\EpicGames.Core.csproj" PrivateAssets="All"><Private>false</Private></ProjectReference>
     <ProjectReference Include="..\..\UnrealBuildTool\UnrealBuildTool.csproj" PrivateAssets="All"><Private>false</Private></ProjectReference>
     <ProjectReference Include="..\AutomationUtils\AutomationUtils.Automation.csproj" PrivateAssets="All"><Private>false</Private></ProjectReference>
diff --git a/Engine/Source/Programs/AutomationTool/Apple/ExportAppleBuild.Automation.cs b/Engine/Source/Programs/AutomationTool/Apple/ExportAppleBuild.Automation.cs
--- a/Engine/Source/Programs/AutomationTool/Apple/ExportAppleBuild.Automation.cs
+++ b/Engine/Source/Programs/AutomationTool/Apple/ExportAppleBuild.Automation.cs
@@ -10,6 +10,18 @@ using System.Threading;
 using UnrealBuildBase;
 using Microsoft.Extensions.Logging;
 using UnrealBuildTool;
+using System.IdentityModel.Tokens.Jwt;
+using Org.BouncyCastle.Crypto.Parameters;
+using System.Security.Cryptography;
+using Microsoft.IdentityModel.Tokens;
+using OpenTracing.Tag;
+using EpicGames.Horde.Storage;
+using Google.Protobuf.WellKnownTypes;
+using System.Net.Http;
+using System.Xml.Linq;
+using System.Net.Http.Headers;
+using System.Text.Json;
+using System.Text.RegularExpressions;
 
 class ExportAppleBuild : BuildCommand
 {
@@ -27,6 +39,12 @@ class ExportAppleBuild : BuildCommand
 		string LatestXcArchiveParam = ParseParamValue("LatestXcArchiveForTarget");
 		string LatestXcArchiveSearch = ParseParamValue("LatestXcArchiveSearchPath");
 		string ExportPath = ParseParamValue("ExportPath");
+		string WhatsNew = ParseParamValue("WhatsNew");
+		bool bUpload = !ParseParam("SkipUpload");
+		bool bCheckVersion = ParseParam("CheckVersion");
+		string VersionConflictSigningIdentity = ParseParamValue("VersionConflictSigningIdentity");
+		string VersionConflictIncrement = ParseParamValue("VersionConflictIncrement");
+		string ExportAuthOverride = ParseParamValue("ExportAuthOverride");
 
 		DirectoryReference XcArchive;
 		if (!string.IsNullOrEmpty(LatestXcArchiveParam))
@@ -41,7 +59,7 @@ class ExportAppleBuild : BuildCommand
 		}
 		else
 		{
-			XcArchive = new DirectoryReference(XcArchivePathParam);			
+			XcArchive = new DirectoryReference(XcArchivePathParam);
 		}
 
 		if (XcArchive == null || !DirectoryReference.Exists(XcArchive))
@@ -49,47 +67,560 @@ class ExportAppleBuild : BuildCommand
 			Logger.LogError("No XCArchive found, with -XCArchive or -LatestXCArchiveForTarget params");
 			return ExitCode.Error_Arguments;
 		}
+		
+		string JWTToken = null;
+		if (bCheckVersion || !string.IsNullOrEmpty(WhatsNew))
+		{
+			JWTToken = CreateJWTToken(ProjectFile);
+		}
 
-		string PlistPath;
-		if (!string.IsNullOrEmpty(PlistParam) && PlistParam[0] == '/')
+		if (bCheckVersion)
 		{
-			PlistPath = PlistParam;
+			if (CheckVersion(JWTToken, ref XcArchive, VersionConflictSigningIdentity, VersionConflictIncrement) == false)
+			{
+				return ExitCode.Error_Arguments;
+			}
 		}
-		else
+
+		if (bUpload)
 		{
-			// look in project and engine for named plist
-			PlistPath = Path.Combine(ProjectFile.Directory.FullName, "Build", "Xcode", PlistParam + ".plist");
+			string PlistPath;
+			if (!string.IsNullOrEmpty(PlistParam) && PlistParam[0] == '/')
+			{
+				PlistPath = PlistParam;
+			}
+			else
+			{
+				// look in project and engine for named plist
+				PlistPath = Path.Combine(ProjectFile.Directory.FullName, "Build", "Xcode", PlistParam + ".plist");
+				if (!File.Exists(PlistPath))
+				{
+					PlistPath = Path.Combine(Unreal.EngineDirectory.FullName, "Build", "Xcode", PlistParam + ".plist");
+				}
+			}
+
 			if (!File.Exists(PlistPath))
 			{
-				PlistPath = Path.Combine(Unreal.EngineDirectory.FullName, "Build", "Xcode", PlistParam + ".plist");
+				Logger.LogError("No plist options file found, with -plist param");
+				return ExitCode.Error_Arguments;
+			}
+
+			Logger.LogInformation("Project: {Project}, XCArchive: {XC}, Plist {Plist}", ProjectFile, XcArchive, PlistPath);
+
+			// build commandline
+			string CommandLine = $"-exportArchive";
+			CommandLine += $" -archivePath \"{XcArchive}\"";
+			CommandLine += $" -exportOptionsPlist \"{PlistPath}\"";
+			CommandLine += $" -allowProvisioningUpdates";
+			if (!string.IsNullOrEmpty(ExportAuthOverride))
+			{
+				Logger.LogInformation("Auth Overrides: {Overrides}", ExportAuthOverride);
+				CommandLine += " " + ExportAuthOverride;
+			}
+			else
+			{
+				CommandLine += AppleExports.GetXcodeBuildAuthOptions(ProjectFile);
+			}
+
+			if (!string.IsNullOrEmpty(ExportPath))
+			{
+				CommandLine += $" -exportPath {ExportPath}";
+			}
+
+			Logger.LogInformation($"Running 'xcodebuild {CommandLine}'...");
+
+			int Return;
+			string Output = Utils.RunLocalProcessAndReturnStdOut("/usr/bin/xcodebuild", CommandLine, null, out Return);
+			Logger.LogInformation(Output);
+
+			if (Return != 0)
+			{
+				return (ExitCode)Return;
 			}
 		}
-		if (!File.Exists(PlistPath))
+
+		if (!string.IsNullOrEmpty(WhatsNew))
 		{
-			Logger.LogError("No plist options file found, with -plist param");
-			return ExitCode.Error_Arguments;
+			// get some info from the .xcarchive that we just uploaded
+			string BundleId, AppVersion, BundleVersion;
+			GetInfoFromXcArchive(XcArchive, out BundleId, out AppVersion, out BundleVersion);
+
+			Logger.LogInformation("Setting What's New for {BundleId}, app version{Version}, build version {BundleVersion}", BundleId, AppVersion, BundleVersion);
+
+			UpdateWhatsNew(JWTToken, BundleId, AppVersion, BundleVersion, WhatsNew);
+		}
+
+		return (ExitCode)0;
+	}
+
+	HttpClient CreateClient(string Token)
+	{
+		HttpClient HttpClient = new HttpClient();
+		HttpClient.BaseAddress = new Uri("https://api.appstoreconnect.apple.com/v1/");
+		HttpClient.DefaultRequestHeaders.Authorization = new AuthenticationHeaderValue("Bearer", Token);
+		return HttpClient;
+	}
+
+
+
+	record Links(string self, string related);
+	record Relationship(Links links);
+
+	record AppInfo(string type, string id);
+	record AppResponse(List<AppInfo> data);
+
+
+	record BuildRelationships(Relationship preReleaseVersion);
+	record BuildAttributes(string version, string processingState, string uploadedDate);
+	record BuildInfo(string type, string id, BuildAttributes attributes, BuildRelationships relationships);
+	record BuildsResponse(List<BuildInfo> data);
+
+	record VersionAttributes(string version, string platform);
+	record VersionInfo(string id, VersionAttributes attributes);
+	record VersionResponse(VersionInfo data);
+	record VersionsResponse(List<VersionInfo> data);
+
+	record BuildsWithVersionsResponse(List<BuildInfo> data, List<VersionInfo> included);
+
+	record PagedLinks(string next);
+	record PagedBuildsResponse(List<BuildInfo> data, PagedLinks links);
+
+	record LocalizationAttributes(string whatsNew);
+	record LocalizationInfo(string type, string id, LocalizationAttributes attributes);
+	record LocalizationResponse(List<LocalizationInfo> data);
+	record LocalizationUpdate(LocalizationInfo data);
+
+	record LocalizationRelationship(string id, string type);
+	record LocalizationRelationshipsData(LocalizationRelationship data);
+	record LocalizationRelationships(LocalizationRelationshipsData build);
+	record LocalizationCreateAttributes(string locale, string whatsNew);
+	record LocalizationCreateInfo(LocalizationCreateAttributes attributes, LocalizationRelationships relationships, string type);
+	record LocalizationCreateRequest(LocalizationCreateInfo data);
+
+
+
+
+
+	static CancellationToken CancellationToken = new();
+
+	// shared client object
+	HttpClient HttpClient = null;
+
+	string GetAppId(string BundleId)
+	{
+		AppResponse Response = HttpClient.GetAsync<AppResponse>($"apps?filter[bundleId]={BundleId}&fields[apps]=name", CancellationToken).Result;
+		return Response.data[0].id;
+	}
+
+	string GetBuildId(string AppId, string PreReleaseVersion, string BuildVersion)
+	{
+		int Limit = 1;
+		string SubURL = $"builds?filter[app]={AppId}&filter[version]={BuildVersion}&sort=-uploadedDate&limit={Limit}";
+
+		bool bHasProcessedAll = false;
+		while (!bHasProcessedAll)
+		{
+			PagedBuildsResponse Response = HttpClient.GetAsync<PagedBuildsResponse>(SubURL, CancellationToken).Result;			
+			foreach (BuildInfo Build in Response.data)
+			{
+				// follow link to preReleaseVersion
+				SubURL = Build.relationships.preReleaseVersion.links.related.Replace(HttpClient.BaseAddress.ToString(), "");
+				VersionResponse VersionResponse = HttpClient.GetAsync<VersionResponse>(SubURL, CancellationToken).Result;
+
+				// if it matches, return this Build!
+				if (VersionResponse.data.attributes.version == PreReleaseVersion)
+				{
+					return Build.id;
+				}
+			}
+
+			if (Response.links.next == null)
+			{
+				bHasProcessedAll = true;
+			}
+			else
+			{
+				// get suburl for next page
+				SubURL = Response.links.next.Replace(HttpClient.BaseAddress.ToString(), "");
+			}
+		}
+
+		// at this point, if no build was found, then it's not been uploaded yet
+		return null;
+	}
+
+
+	string GetLocalizationId(string BuildId)
+	{
+		string SubURL = $"builds/{BuildId}/betaBuildLocalizations";
+		LocalizationResponse Response = HttpClient.GetAsync<LocalizationResponse>(SubURL, CancellationToken).Result;
+
+		if (Response.data.Count == 0)
+		{
+			Console.WriteLine("no current beta build localized data, will create new english data");
+			return null;
+		}
+		Console.WriteLine("current whats new is: {0}", Response.data[0].attributes.whatsNew);
+
+		return Response.data[0].id;
+	}
+
+	// return true if the version doesn't exist - which means it can be uploaded
+	bool CheckVersion(string JWTToken, ref DirectoryReference XcArchive, string SigningIdentity, string Increment)
+	{
+		// get some info from the .xcarchive
+		string BundleId, AppVersion, BundleVersion;
+		GetInfoFromXcArchive(XcArchive, out BundleId, out AppVersion, out BundleVersion);
+
+
+Console.WriteLine("Checking for version >= {0}", BundleVersion);
+		Version BundleVersionNumber = new Version(BundleVersion);
+		using (HttpClient = CreateClient(JWTToken))
+		{
+			string AppId = GetAppId(BundleId);
+
+			// get prerelease version id
+			string SubURL = $"preReleaseVersions?filter[app]={AppId}&filter[version]={AppVersion}&fields[preReleaseVersions]=version";
+//Console.WriteLine("request: {0}", SubURL);
+			VersionsResponse VersionResponse = HttpClient.GetAsync<VersionsResponse>(SubURL, CancellationToken).Result;
+			// if the app version doesn't even exist, then the bundle verrsion can't exist in it!
+			if (VersionResponse.data.Count == 0)
+			{
+				return true;
+			}
+
+			Version MaxVersion = new Version(0, 0);
+
+			// now walk over all builds until we find one >= the bundle version
+			string VersionId = VersionResponse.data[0].id;
+			SubURL = $"preReleaseVersions/{VersionId}/builds?limit=5&fields[builds]=version";
+//Console.WriteLine("request: {0}", SubURL);
+			bool bHasProcessedAll = false;
+			while (!bHasProcessedAll)
+			{
+				PagedBuildsResponse Response = HttpClient.GetAsync<PagedBuildsResponse>(SubURL, CancellationToken).Result;
+
+				foreach (BuildInfo Build in Response.data)
+				{
+					string VersionString = Build.attributes.version;
+					// append .0.0 if needed
+					if (!VersionString.Contains("."))
+					{
+						VersionString += ".0.0";
+					}
+
+					Version Version = new Version(VersionString);
+					if (Version > MaxVersion)
+					{
+						MaxVersion = Version;
+					}
+				}
+
+				if (Response.links.next == null)
+				{
+					bHasProcessedAll = true;
+				}
+				else
+				{
+					// get suburl for next page
+					SubURL = Response.links.next.Replace(HttpClient.BaseAddress.ToString(), "");
+				}
+			}
+
+			if (MaxVersion >= BundleVersionNumber)
+			{
+				if (SigningIdentity != null && Increment != null)
+				{
+					Version IncrementVersion = new Version(Increment);
+					if (IncrementVersion == new Version(0,0,0))
+					{
+						throw new Exception($"Invalid version increment {Increment}");
+					}
+
+					string TempDir = Path.GetTempFileName();
+					File.Delete(TempDir);
+					Directory.CreateDirectory(TempDir);
+
+					FileReference Script = FileReference.Combine(Unreal.EngineDirectory, "Build/BatchFiles/Mac/ResignApp.sh");
+					string FixedXcArchive = TempDir + "/FixedVersion.xcarchive";
+					Version FixedVersion = new Version(MaxVersion.Major + IncrementVersion.Major, MaxVersion.Minor + IncrementVersion.Minor, MaxVersion.Build + IncrementVersion.Build);
+					string Options = $"{Script} -s \"{XcArchive}\" -d \"{FixedXcArchive}\" -i \"{SigningIdentity}\" -bv {FixedVersion}";
+
+					if (MaxVersion == BundleVersionNumber)
+					{
+						Logger.LogInformation("A build exists with same version that is in the .xcarchive ({BundleVersion}). Making a new\n" +
+							".xcarchive ({FixedXcArchive}) with a version using the supplied increment. The new version is {FixedVersion}.",
+							BundleVersion, FixedXcArchive, FixedVersion);
+					}
+					else
+					{
+						Logger.LogInformation("A build exists with LARGER version ({MaxVersion}) that is in the .xcarchive ({BundleVersion})\n" + 
+							"which would cause the uploaded build to be auto-incremented. Making a new .xcarchive ({FixedXcArchive})\n" + 
+							"with a incremented FROM THE LARGER VERSION (not the version in your .xcarchive!). The new version is {FixedVersion}.",
+							MaxVersion, BundleVersion, FixedXcArchive, FixedVersion);
+					}
+
+					Logger.LogInformation("Running: {Script} {Options}...", "/usr/bin/env", Options);
+					int ExitCode;
+					string Output = Utils.RunLocalProcessAndReturnStdOut("/usr/bin/env", Options, out ExitCode);
+					if (ExitCode != 0)
+					{
+						Logger.LogError("{Output}", Output);
+						return false;
+					}
+
+					Logger.LogInformation("{Output}", Output);
+
+					// update the XcArchive location to the resigned one
+					XcArchive = new DirectoryReference(FixedXcArchive);
+
+					// we know that this xcarchive has a version that is not in the appstore, so return true 
+				}
+				else
+				{
+					Version PossibleVersion = new Version(MaxVersion.Major, MaxVersion.Minor + 1, MaxVersion.Build);
+
+					Logger.LogError("A build exists with same or larger version ({MaxVersion}) than is in the .xcarchive ({BundleVersion}). Either make a new build,\n" + 
+						"or modify the build version in the .xcarchive with a command like:\n" + 
+						"   Engine/Build/BatchFiles/Mac/ResignApp.sh -s {XcArchive} -d UpdatedVersion.xcarchive -i <signing identity> -bv {NewVersion}\n" + 
+						"and then upload that .xcarchive",
+						MaxVersion, BundleVersion, XcArchive, PossibleVersion);
+
+					return false;
+				}
+
+			}
+
+			return true;
+		}
+	}
+
+
+	bool CreateWhatsNew(string BuildId, string Text)
+	{
+		LocalizationRelationship Relationship = new(id: BuildId, type:"builds");
+		LocalizationRelationshipsData RelationshipData = new(data: Relationship);
+		LocalizationRelationships Relationships = new(build: RelationshipData);
+		LocalizationCreateAttributes Attributes = new(locale:"en-US", whatsNew:Text);
+		LocalizationCreateInfo Info = new(type:"betaBuildLocalizations", attributes:Attributes, relationships:Relationships);
+		LocalizationCreateRequest LocUpdate = new(data: Info);
+
+		StringContent JsonContent = new StringContent(JsonSerializer.Serialize<LocalizationCreateRequest>(LocUpdate), System.Text.Encoding.UTF8, "application/json");
+
+		string SubURL = $"betaBuildLocalizations";
+		HttpResponseMessage Response = HttpClient.PostAsync(SubURL, JsonContent).Result;
+
+		if (!Response.IsSuccessStatusCode)
+		{
+			Logger.LogError("Create Localization Data failed!: StatusCode {Code}\n{Response}", Response.StatusCode, Response.Content.ReadAsStringAsync().Result);
+			return false;
 		}
 
-		Logger.LogInformation("Project: {Project}, XCArchive: {XC}, Plist {Plist}", ProjectFile, XcArchive, PlistPath);
+		return true;
+	}
+
+	bool SetWhatsNew(string LocId, string Text)
+	{
+		LocalizationAttributes Attributes = new(whatsNew: Text);
+		LocalizationInfo Info = new(type: "betaBuildLocalizations", id: LocId, attributes: Attributes);
+		LocalizationUpdate LocUpdate = new(data: Info);
+
+		StringContent JsonContent = new StringContent(JsonSerializer.Serialize<LocalizationUpdate>(LocUpdate), System.Text.Encoding.UTF8, "application/json");
 
-		// build commandline
-		string CommandLine = $"-exportArchive";
-		CommandLine += $" -archivePath \"{XcArchive}\"";
-		CommandLine += $" -exportOptionsPlist \"{PlistPath}\"";
-		CommandLine += $" -allowProvisioningUpdates";
-		CommandLine += AppleExports.GetXcodeBuildAuthOptions(ProjectFile);
+		string SubURL = $"betaBuildLocalizations/{LocId}";
+		HttpResponseMessage Response = HttpClient.PatchAsync(SubURL, JsonContent).Result;
 
-		if (!string.IsNullOrEmpty(ExportPath))
+		if (!Response.IsSuccessStatusCode)
 		{
-			CommandLine += $" -exportPath {ExportPath}";
+			Logger.LogError("Patch Localization Data failed!: StatusCode {Code}\n{Response}", Response.StatusCode, Response.Content.ReadAsStringAsync().Result);
+			return false;
 		}
 
-		Logger.LogInformation($"Running 'xcodebuild {CommandLine}'...");
+		return true;
+	}
+
+	bool UpdateWhatsNew(string Token, string BundleId, string PreReleaseVersion, string BundleVersion, string Text)
+	{
+		using (HttpClient = CreateClient(Token))
+		{
+			// wait 15 second between queries for the version
+
+			string AppId = GetAppId(BundleId);
+
+			int Delay = 15;
+			int Tries = 10 * 60 / Delay;
+
+			string BuildId = null;
+			while (BuildId == null && Tries-- > 0)
+			{
+				BuildId = GetBuildId(AppId, PreReleaseVersion, BundleVersion);
+				if (BuildId == null)
+				{
+					Logger.LogInformation("Build not found yet, waiting and will try again...");
+					Thread.Sleep(Delay * 1000);
+				}
+			}
+
+			if (BuildId == null)
+			{
+				Logger.LogError("Unable to find build for BundleId: {BundleId}, App Version {PreReleaseVersion}, Build Version {BundleVersion}. It could be due to:\n" + 
+					"  * The build was never actually uploaded to AppStoreConnect to TestFlight\n" + 
+					"  * A later version was already there, and Apple changed the version number automatically. Run again with -CheckVersion -SkipUpload to check\n" + 
+					"  * There was a problem with the app. Check your email (or others on your team) to see if Apple rejected the app after uploading",
+					BundleId, PreReleaseVersion, BundleVersion);
+				return false;
+			}
+			// in the case that we don't use "most-recent" build, then we are looking for exact build version match to bundle version
+			string BuildVersion = BundleVersion;
+			
+			// string BuildVersion;
+			// string BuildId = GetMostRecentBuildId(AppId, PreReleaseVersion, out BuildVersion);
+			// if (BuildId == null)
+			// {
+			// 	Logger.LogError("Unable to find a recent build for version {Version} of {BundleId}", PreReleaseVersion, BundleId);
+			// 	return false;
+			// }
+
+			Logger.LogInformation("Found a build with version: {BuildVersion} [id is {Id}]", BuildVersion, BuildId);
+
+			string LocId = GetLocalizationId(BuildId);	
+
+			bool bUpdateSucceeded;
+			if (LocId == null)
+			{
+				bUpdateSucceeded = CreateWhatsNew(BuildId, Text);
+			}
+			else 
+			{
+				bUpdateSucceeded = SetWhatsNew(LocId, Text);
+			}
+			if (bUpdateSucceeded)
+			{
+				GetLocalizationId(BuildId);
+				return true;
+			}
+			
+			return false;
+		}
+	}
+
+	static void GetInfoFromXcArchive(DirectoryReference XcArchive, out string BundleId, out string AppVersion, out string BundleVersion)
+	{
+		BundleId = null;
+		AppVersion = null;
+		BundleVersion = null;
+		FileReference Plist = FileReference.Combine(XcArchive, "Info.plist");
+
+		string[] Lines = FileReference.ReadAllLines(Plist);
+		bool bMatchBundleId = false;
+		bool bMatchAppVersion = false;
+		bool bMatchBundleVersion = false;
+		foreach (string Line in Lines)
+		{
+			if (bMatchBundleId || bMatchAppVersion || bMatchBundleVersion)
+			{
+				Match Match = Regex.Match(Line, @"\<string\>(.*)\<\/string\>");
+				if (Match.Success)
+				{
+					if (bMatchBundleId)
+					{
+						BundleId = Match.Groups[1].Value;
+					}
+					else if (bMatchAppVersion)
+					{
+						AppVersion = Match.Groups[1].Value;
+					}
+					else if (bMatchBundleVersion)
+					{
+						BundleVersion = Match.Groups[1].Value;
+					}
+				}
+				bMatchBundleId = bMatchAppVersion = bMatchBundleVersion = false;
+			}
+			else if (Line.Contains("CFBundleIdentifier"))
+			{
+				bMatchBundleId = true;
+			}
+			else if (Line.Contains("CFBundleShortVersionString"))
+			{
+				bMatchAppVersion = true;
+			}
+			else if (Line.Contains("CFBundleVersion"))
+			{
+				bMatchBundleVersion = true;
+			}
+		}
+
+		if (BundleId == null || AppVersion == null || BundleVersion == null)
+		{
+			throw new Exception($"Failed to get bundle ID and versions from {Plist}");
+		}
+	}
+
+	static string CreateJWTToken(FileReference ProjectFile)
+	{
+		// @todo: refactor the AppleExports.GetXcodeBuildAuthOptions() code so that we can get this info in a shared way
+		ConfigHierarchy SharedPlatformIni = ConfigCache.ReadHierarchy(ConfigHierarchyType.Engine, ProjectFile?.Directory, UnrealTargetPlatform.Mac);
+		string IssuerID, KeyID, KeyPath;
+		string P8Path;
+		if (SharedPlatformIni.TryGetValue("/Script/MacTargetPlatform.XcodeProjectSettings", "AppStoreConnectIssuerID", out IssuerID) &&
+			SharedPlatformIni.TryGetValue("/Script/MacTargetPlatform.XcodeProjectSettings", "AppStoreConnectKeyID", out KeyID) &&
+			SharedPlatformIni.TryGetValue("/Script/MacTargetPlatform.XcodeProjectSettings", "AppStoreConnectKeyPath", out KeyPath))
+		{
+			P8Path = AppleExports.ConvertFilePath(ProjectFile?.Directory, KeyPath).FullName;
+		}
+		else
+		{
+			throw new Exception("Checking the version, or setting WhatsNew requires the AppStoreConnect API, which needs these config properties:\n" +
+				"  [/Script/MacTargetPlatform.XcodeProjectSettings]\n" +
+				"  AppStoreConnectIssuerID\n" +
+				"  AppStoreConnectKeyID\n" +
+				"  AppStoreConnectKeyPath\n" + 
+				"to be set correctly in your Engine config files. Note that you don't need to set bUseAppStoreConnect to true for this to work");
+		}
+
+
+		// since we can't use Cng apis on mac (or linux) we have to use BouncyCastle to convert the .p8 text to a
+		// ECDsa key, using some crazy elliptical cryptography black magic
+		ECDsa Key;
+		using (TextReader P8Reader = File.OpenText(P8Path))
+		{
+			ECPrivateKeyParameters ECPrivateKeyParameters = (ECPrivateKeyParameters)new Org.BouncyCastle.OpenSsl.PemReader(P8Reader).ReadObject();
+
+            Org.BouncyCastle.Math.EC.ECPoint Q = ECPrivateKeyParameters.Parameters.G.Multiply(ECPrivateKeyParameters.D).Normalize();
+			ECParameters MsEcp = new ECParameters {
+				Curve = ECCurve.NamedCurves.nistP256, 
+				Q = { 
+					X = Q.AffineXCoord.GetEncoded(),
+					Y = Q.AffineYCoord.GetEncoded()
+				}, 
+				D = ECPrivateKeyParameters.D.ToByteArrayUnsigned()
+			};
+			Key = ECDsa.Create(MsEcp);
+		}
+
+		// set up the header with the parmas that AppStoreConnect needs (kid, typ and alg)
+		ECDsaSecurityKey SecurityKey = new ECDsaSecurityKey(Key) { KeyId = KeyID };
+		SigningCredentials Credentials = new SigningCredentials(SecurityKey, SecurityAlgorithms.EcdsaSha256);
+		JwtHeader Header = new JwtHeader(Credentials);
+		
+		// set up the payload manually with the body that AppStoreConnect needs (aud, exp, and iss - iat and scope are optional)
+		// @todo: use scope to limit the token's access
+		JwtPayload Payload = new JwtPayload
+		{
+			{ "aud", "appstoreconnect-v1" },
+			{ "exp", (long)(DateTime.UtcNow - DateTime.UnixEpoch).TotalSeconds + 60 * 19 },
+			{ "iss", IssuerID }
+		};
 
-		int Return;
-		string Output = Utils.RunLocalProcessAndReturnStdOut("/usr/bin/xcodebuild", CommandLine, null, out Return);
-		Logger.LogInformation(Output);
+		// make a token and output as string
+		JwtSecurityToken JwtToken = new JwtSecurityToken(Header, Payload);
+		string FullToken = new JwtSecurityTokenHandler().WriteToken(JwtToken);
 
-		return (ExitCode)Return;
+		// When debugging things with curl, you'll need this token, but we don't write it by default since it could technically 
+		// be used to hack ASC within 20 minutes
+		//	Console.WriteLine("Auth header:\n  Authorization: Bearer {0}", FullToken);
+		
+		return FullToken;
 	}
 }
