[Gauntlet] Fixes ReportGenUtils on unix build machines incorrectly inquiring $Path for the dotnet executable

diff --git a/Engine/Source/Programs/AutomationTool/Gauntlet/Framework/Utils/Gauntlet.ReportGenUtils.cs b/Engine/Source/Programs/AutomationTool/Gauntlet/Framework/Utils/Gauntlet.ReportGenUtils.cs
--- a/Engine/Source/Programs/AutomationTool/Gauntlet/Framework/Utils/Gauntlet.ReportGenUtils.cs
+++ b/Engine/Source/Programs/AutomationTool/Gauntlet/Framework/Utils/Gauntlet.ReportGenUtils.cs
@@ -28,9 +28,19 @@ namespace Gauntlet
 			}
 			else
 			{
-				PathDelimter = ':';
-				DotnetPathLocation = "dotnet";
 				BasePythonLocation = @"Engine/Binaries/ThirdParty/Python3/Mac/bin/python3";
+
+				if (CommandUtils.IsBuildMachine)
+				{
+					DotnetPathLocation = "/usr/local/bin/dotnet/dotnet";
+					return;
+				}
+				else
+				{
+					PathDelimter = ':';
+					DotnetPathLocation = "dotnet";
+					BasePythonLocation = @"Engine/Binaries/ThirdParty/Python3/Mac/bin/python3";
+				}
 			}
 
 			DotnetPathLocation = Path.Combine(Environment.GetEnvironmentVariable("PATH")
