Fix staging an application under the restricted directory

#rnx

diff --git a/Engine/Source/Programs/AutomationTool/Scripts/CopyBuildToStagingDirectory.Automation.cs b/Engine/Source/Programs/AutomationTool/Scripts/CopyBuildToStagingDirectory.Automation.cs
--- a/Engine/Source/Programs/AutomationTool/Scripts/CopyBuildToStagingDirectory.Automation.cs
+++ b/Engine/Source/Programs/AutomationTool/Scripts/CopyBuildToStagingDirectory.Automation.cs
@@ -5577,7 +5577,7 @@ namespace AutomationScripts
 						// if we have a uproject file that is under the engine dir, then the target/exe will be in Engine/Binaries
 						if (ProjectDir.IsUnderDirectory(EngineDir))
 						{
-							ReceiptBaseDir = EngineDir;
+							ReceiptBaseDir = Unreal.GetExtensionDirs(EngineDir).Where(x => ProjectDir.IsUnderDirectory(x)).OrderByDescending(x => x.FullName.Length).FirstOrDefault() ?? EngineDir;
 						}
 
 						//if RawProgramProjectOverride then the binary will likely be under <Game>/Binaries
