Windows Code Signing Adjustments

* Fixing code so the first server is used
* removing a server as it appears to be not reliable anymore

#rnx

diff --git a/Engine/Source/Programs/AutomationTool/AutomationUtils/CommandUtils.cs b/Engine/Source/Programs/AutomationTool/AutomationUtils/CommandUtils.cs
--- a/Engine/Source/Programs/AutomationTool/AutomationUtils/CommandUtils.cs
+++ b/Engine/Source/Programs/AutomationTool/AutomationUtils/CommandUtils.cs
@@ -3024,18 +3024,18 @@ namespace AutomationTool
 			SHA256
 		};
 
+		// List of known / valid timestamp servers: https://gist.github.com/Manouchehri/fd754e402d98430243455713efada710
 		static readonly string[] TimestampServersSHA1 =
 		{
 			"http://timestamp.digicert.com",
-			"http://timestamp.comodoca.com/authenticode",
-			"http://timestamp.globalsign.com/scripts/timstamp.dll"
+			"http://ts.ssl.com",
+			"http://kstamp.keynectis.com/KSign"
 		};
 
 		static readonly string[] TimestampServersSHA256 =
 		{
 			"http://timestamp.digicert.com",
-			"http://sha256timestamp.ws.symantec.com/sha256/timestamp",
-			"http://timestamp.comodoca.com/?td=sha256",
+			"http://ts.ssl.com",
 			"http://rfc3161timestamp.globalsign.com/advanced"
 		};
 
