Horde: Update CreateArtifactTask to include keys and metadata, and to use the stream id passed by the agent through an envronment variable.

#rnx

diff --git a/Engine/Source/Programs/AutomationTool/BuildGraph/Tasks/CreateArtifactTask.cs b/Engine/Source/Programs/AutomationTool/BuildGraph/Tasks/CreateArtifactTask.cs
--- a/Engine/Source/Programs/AutomationTool/BuildGraph/Tasks/CreateArtifactTask.cs
+++ b/Engine/Source/Programs/AutomationTool/BuildGraph/Tasks/CreateArtifactTask.cs
@@ -2,11 +2,8 @@
 
 using System;
 using System.Collections.Generic;
-using System.Data;
 using System.Diagnostics;
-using System.IO;
 using System.Linq;
-using System.Net.Http;
 using System.Threading.Tasks;
 using System.Xml;
 using EpicGames.Core;
@@ -14,7 +11,6 @@ using EpicGames.Horde;
 using EpicGames.Horde.Artifacts;
 using EpicGames.Horde.Commits;
 using EpicGames.Horde.Storage;
-using EpicGames.Horde.Storage.Clients;
 using EpicGames.Horde.Storage.Nodes;
 using EpicGames.Horde.Streams;
 using Microsoft.Extensions.DependencyInjection;
@@ -25,151 +21,178 @@ using Microsoft.Extensions.Logging;
 namespace AutomationTool.Tasks
 {
 	/// <summary>
-	/// Parameters for a CreateArtifact task
+	/// Parameters for a <see cref="CreateArtifactTask"/>.
 	/// </summary>
 	public class CreateArtifactTaskParameters
 	{
 		/// <summary>
-		/// Name for the artifact
+		/// Name of the artifact
 		/// </summary>
 		[TaskParameter]
-		public string Name { get; set; } = String.Empty;
+		public string Name { get; set; } = null!;
 
 		/// <summary>
-		/// Type of the artifact
+		/// The artifact type. Determines the permissions and expiration policy for the artifact.
 		/// </summary>
-		[TaskParameter(Optional = true)]
-		public string Type { get; set; } = "unknown";
+		[TaskParameter]
+		public string Type { get; set; } = null!;
 
 		/// <summary>
-		/// Description for the artifact
+		/// Description for the artifact. Will be shown through the Horde dashboard.
 		/// </summary>
 		[TaskParameter(Optional = true)]
 		public string? Description { get; set; }
 
 		/// <summary>
-		/// Base directory to resolve relative paths for input files.
+		/// Base path for the uploaded files. All the tagged files must be under this directory. Defaults to the workspace root directory.
 		/// </summary>
-		[TaskParameter]
+		[TaskParameter(Optional = true)]
 		public string? BaseDir { get; set; }
 
 		/// <summary>
-		/// Stream for this artifact
+		/// Stream containing the artifact.
 		/// </summary>
 		[TaskParameter(Optional = true)]
 		public string? StreamId { get; set; }
 
 		/// <summary>
-		/// Changelist number for this artifact
+		/// Commit for the uploaded artifact.
+		/// </summary>
+		[TaskParameter(Optional = true)]
+		public string? Commit { get; set; }
+
+		/// <summary>
+		/// Files to include in the artifact.
+		/// </summary>
+		[TaskParameter(ValidationType = TaskParameterValidationType.FileSpec)]
+		public string Files { get; set; } = null!;
+
+		/// <summary>
+		/// Queryable keys for this artifact, separated by semicolons.
 		/// </summary>
 		[TaskParameter(Optional = true)]
-		public int? Change { get; set; }
+		public string? Keys { get; set; }
 
 		/// <summary>
-		/// Files to be uploaded.
+		/// Other metadata for the artifact, separated by semicolons.
 		/// </summary>
-		[TaskParameter(Optional = true, ValidationType = TaskParameterValidationType.FileSpec)]
-		public string Files { get; set; } = "...";
+		[TaskParameter(Optional = true)]
+		public string? Metadata { get; set; }
 	}
 
 	/// <summary>
-	/// Deploys a tool update through Horde
+	/// Uploads an artifact to Horde
 	/// </summary>
 	[TaskElement("CreateArtifact", typeof(CreateArtifactTaskParameters))]
-	public class CreateArtifactTask : SpawnTaskBase
+	public class CreateArtifactTask : BgTaskImpl
 	{
-		class LoggerProviderAdapter : ILoggerProvider
-		{
-			readonly ILogger _logger;
-
-			public LoggerProviderAdapter(ILogger logger) => _logger = logger;
-			public ILogger CreateLogger(string categoryName) => _logger;
-			public void Dispose() { }
-		}
-
 		readonly CreateArtifactTaskParameters _parameters;
 
 		/// <summary>
-		/// Construct a Helm task
+		/// Constructor.
 		/// </summary>
-		/// <param name="parameters">Parameters for the task</param>
+		/// <param name="parameters">Parameters for this task.</param>
 		public CreateArtifactTask(CreateArtifactTaskParameters parameters)
-		{
-			_parameters = parameters;
-		}
+			=> _parameters = parameters;
 
 		/// <summary>
 		/// ExecuteAsync the task.
 		/// </summary>
-		/// <param name="job">Information about the current job</param>
+		/// <param name="job">Information about the current job.</param>
 		/// <param name="buildProducts">Set of build products produced by this node.</param>
-		/// <param name="tagNameToFileSet">Mapping from tag names to the set of files they include</param>
+		/// <param name="tagNameToFileSet">Mapping from tag names to the set of files they include.</param>
 		public override async Task ExecuteAsync(JobContext job, HashSet<FileReference> buildProducts, Dictionary<string, HashSet<FileReference>> tagNameToFileSet)
 		{
-			// Create a DI container for building the graph
-			ServiceCollection serviceCollection = new ServiceCollection();
-			serviceCollection.AddHorde();
-			serviceCollection.AddLogging(builder => builder.AddProvider(new LoggerProviderAdapter(Log.Logger)));
-			serviceCollection.Configure<HordeOptions>(x => x.AllowAuthPrompt = !Automation.IsBuildMachine);
-			serviceCollection.Configure<LoggerFilterOptions>(options => options.AddFilter(typeof(HttpClient).FullName, LogLevel.Warning));
-
-			await using ServiceProvider serviceProvider = serviceCollection.BuildServiceProvider();
-
-			ArtifactName artifactName = new ArtifactName(_parameters.Name);
-			ArtifactType artifactType = new ArtifactType(_parameters.Type);
+			ArtifactName name = new ArtifactName(_parameters.Name);
+			ArtifactType type = new ArtifactType(_parameters.Type);
+			string[] keys = (_parameters.Keys ?? string.Empty).Split(';', StringSplitOptions.RemoveEmptyEntries | StringSplitOptions.TrimEntries);
+			string[] metadata = (_parameters.Metadata ?? string.Empty).Split(';', StringSplitOptions.RemoveEmptyEntries | StringSplitOptions.TrimEntries);
+
+			// Figure out the current change and stream id
+			StreamId streamId;
+			if (!String.IsNullOrEmpty(_parameters.StreamId))
+			{
+				streamId = new StreamId(_parameters.StreamId);
+			}
+			else
+			{
+				string? streamIdEnvVar = Environment.GetEnvironmentVariable("UE_HORDE_STREAMID");
+				if (!String.IsNullOrEmpty(streamIdEnvVar))
+				{
+					streamId = new StreamId(streamIdEnvVar);
+				}
+				else
+				{
+					throw new AutomationException("Missing UE_HORDE_STREAMID environment variable; unable to determine current stream.");
+				}
+			}
 
-			HordeHttpClient hordeHttpClient = serviceProvider.GetRequiredService<HordeHttpClient>();
-			StreamId streamId = new StreamId(_parameters.StreamId ?? throw new InvalidOperationException("Missing StreamId parameter")); 
-			CommitId commitId = CommitId.FromPerforceChange(_parameters.Change ?? CommandUtils.P4Env.Changelist);
-			CreateArtifactResponse response = await hordeHttpClient.CreateArtifactAsync(artifactName, artifactType, _parameters.Description, streamId: streamId, commitId: commitId);
-			Logger.LogInformation("Creating artifact {ArtifactId} '{ArtifactName}' ({ArtifactType}) (ns: {NamespaceId}, ref: {RefName})", response.ArtifactId, artifactName, artifactType, response.NamespaceId, response.RefName);
+			CommitId commitId;
+			if (!String.IsNullOrEmpty(_parameters.Commit))
+			{
+				commitId = new CommitId(_parameters.Commit);
+			}
+			else
+			{
+				int change = CommandUtils.P4Env.Changelist;
+				if (change > 0)
+				{
+					commitId = CommitId.FromPerforceChange(CommandUtils.P4Env.Changelist);
+				}
+				else
+				{
+					throw new AutomationException("Unknown changelist. Please run with -P4.");
+				}
+			}
 
-			Stopwatch timer = Stopwatch.StartNew();
+			// Resolve the files to include
+			DirectoryReference baseDir = ResolveDirectory(_parameters.BaseDir);
+			List<FileReference> files = BgTaskImpl.ResolveFilespec(baseDir, _parameters.Files, tagNameToFileSet).ToList();
 
-			HttpStorageClientFactory httpStorageClientFactory = serviceProvider.GetRequiredService<HttpStorageClientFactory>();
-			IStorageClient client = httpStorageClientFactory.CreateClient(response.NamespaceId, response.Token);
+			bool validFiles = true;
+			foreach (FileReference file in files)
+			{
+				if (!file.IsUnderDirectory(baseDir))
+				{
+					Logger.LogError("Artifact file {File} is not under {BaseDir}", file, baseDir);
+					validFiles = false;
+				}
+			}
 
-			await using (IBlobWriter writer = client.CreateBlobWriter(response.RefName))
+			if (!validFiles)
 			{
-				DirectoryReference baseDir = ResolveDirectory(_parameters.BaseDir);
-				List<FileInfo> files = ResolveFilespec(baseDir, _parameters.Files, tagNameToFileSet).Select(x => x.ToFileInfo()).ToList();
+				throw new AutomationException($"Unable to create artifact {name} with given file list.");
+			}
 
-				int totalCount = files.Count;
-				long totalSize = files.Sum(x => x.Length);
+			// Create the new artifact
+			IHordeClient hordeClient = CommandUtils.ServiceProvider.GetRequiredService<IHordeClient>();
+			IArtifact artifact = await hordeClient.Artifacts.AddAsync(name, type, _parameters.Description, streamId, commitId, keys, metadata);
+			Logger.LogInformation("Creating artifact {ArtifactId} '{ArtifactName}' ({ArtifactType}) with namespace {NamespaceId}, ref {RefName} ({Link})", artifact.Id, name, type, artifact.NamespaceId, artifact.RefName, $"{hordeClient.ServerUrl}/api/v1/storage/{artifact.NamespaceId}/refs/{artifact.RefName}");
 
-				IHashedBlobRef<DirectoryNode> outputNodeRef = await writer.WriteFilesAsync(baseDir.ToDirectoryInfo(), files, progress: new UpdateStatsLogger(totalCount, totalSize, Logger));
-				await writer.FlushAsync();
+			// Upload the files
+			IStorageClient storage = hordeClient.CreateStorageClient(artifact.NamespaceId);
+			Stopwatch timer = Stopwatch.StartNew();
 
-				await client.WriteRefAsync(response.RefName, outputNodeRef);
+			IHashedBlobRef<DirectoryNode> rootRef;
+			await using (IBlobWriter blobWriter = storage.CreateBlobWriter(artifact.RefName))
+			{
+				rootRef = await blobWriter.WriteFilesAsync(baseDir, files);
 			}
+			await storage.WriteRefAsync(artifact.RefName, rootRef);
 
-			Logger.LogInformation("Completed in {Time:n1}s", timer.Elapsed.TotalSeconds);
+			Logger.LogInformation("Uploaded artifact {ArtifactId} in {Time:n1}s", artifact.Id, timer.Elapsed.TotalSeconds);
 		}
 
-		/// <summary>
-		/// Output this task out to an XML writer.
-		/// </summary>
+		/// <inheritdoc/>
 		public override void Write(XmlWriter writer)
-		{
-			Write(writer, _parameters);
-		}
+			=> Write(writer, _parameters);
 
-		/// <summary>
-		/// Find all the tags which are used as inputs to this task
-		/// </summary>
-		/// <returns>The tag names which are read by this task</returns>
+		/// <inheritdoc/>
 		public override IEnumerable<string> FindConsumedTagNames()
-		{
-			return FindTagNamesFromList(_parameters.Files);
-		}
+			=> FindTagNamesFromFilespec(_parameters.Files);
 
-		/// <summary>
-		/// Find all the tags which are modified by this task
-		/// </summary>
-		/// <returns>The tag names which are modified by this task</returns>
+		/// <inheritdoc/>
 		public override IEnumerable<string> FindProducedTagNames()
-		{
-			yield break;
-		}
+			=> Enumerable.Empty<string>();
 	}
 }
