Horde: Include artifacts created through the CreateArtifact BuildGraph task in the list of artifacts for the job.

diff --git a/Engine/Source/Programs/AutomationTool/BuildGraph/BgSchema.cs b/Engine/Source/Programs/AutomationTool/BuildGraph/BgSchema.cs
--- a/Engine/Source/Programs/AutomationTool/BuildGraph/BgSchema.cs
+++ b/Engine/Source/Programs/AutomationTool/BuildGraph/BgSchema.cs
@@ -646,7 +646,7 @@ namespace AutomationTool
 			XmlSchemaComplexType artifactType = new XmlSchemaComplexType();
 			artifactType.Name = GetTypeName(ScriptSchemaStandardType.Artifact);
 			artifactType.Attributes.Add(CreateSchemaAttribute("Name", ScriptSchemaStandardType.Name, XmlSchemaUse.Required));
-			artifactType.Attributes.Add(CreateSchemaAttribute("Type", ScriptSchemaStandardType.Name, XmlSchemaUse.Optional));
+			artifactType.Attributes.Add(CreateSchemaAttribute("Type", ScriptSchemaStandardType.Name, XmlSchemaUse.Required));
 			artifactType.Attributes.Add(CreateSchemaAttribute("Description", ScriptSchemaStandardType.BalancedString, XmlSchemaUse.Optional));
 			artifactType.Attributes.Add(CreateSchemaAttribute("BasePath", ScriptSchemaStandardType.BalancedString, XmlSchemaUse.Optional));
 			artifactType.Attributes.Add(CreateSchemaAttribute("Tag", ScriptSchemaStandardType.Tag, XmlSchemaUse.Optional));
diff --git a/Engine/Source/Programs/AutomationTool/BuildGraph/BgScriptReader.cs b/Engine/Source/Programs/AutomationTool/BuildGraph/BgScriptReader.cs
--- a/Engine/Source/Programs/AutomationTool/BuildGraph/BgScriptReader.cs
+++ b/Engine/Source/Programs/AutomationTool/BuildGraph/BgScriptReader.cs
@@ -5,11 +5,13 @@ using System.Collections.Generic;
 using System.Diagnostics.CodeAnalysis;
 using System.IO;
 using System.Linq;
+using System.Reflection;
 using System.Text;
 using System.Text.RegularExpressions;
 using System.Threading.Tasks;
 using System.Xml;
 using System.Xml.Schema;
+using AutomationTool.Tasks;
 using EpicGames.BuildGraph;
 using EpicGames.Core;
 using Microsoft.Extensions.Logging;
@@ -841,7 +843,7 @@ namespace AutomationTool
 								lines.RemoveAt(lines.Count - 1);
 							}
 
-							foreach(string line in lines)
+							foreach (string line in lines)
 							{
 								int whitespaceLen = 0;
 								while (whitespaceLen < line.Length && Char.IsWhiteSpace(line[whitespaceLen]))
@@ -997,7 +999,7 @@ namespace AutomationTool
 							LogError(element, $"String operation 'SplitLast' requires exactly 1 argument.");
 							return;
 						}
-	
+
 						operationResult = input.Split(arguments[0]).Last();
 						break;
 					default:
@@ -1315,7 +1317,7 @@ namespace AutomationTool
 				string[] keys = ReadListAttribute(element, "Keys");
 				string[] metadata = ReadListAttribute(element, "Metadata");
 
-				BgArtifactDef newArtifact = new BgArtifactDef(name, type, description, basePath, tag, keys, metadata);
+				BgArtifactDef newArtifact = new BgArtifactDef(name, type, description, basePath, null, tag, keys, metadata);
 				_graph.Artifacts.Add(newArtifact);
 			}
 		}
@@ -1742,7 +1744,39 @@ namespace AutomationTool
 					}
 				}
 				_enclosingNode!.Tasks.Add(info);
+
+				if (info.Name.Equals("CreateArtifact", StringComparison.OrdinalIgnoreCase))
+				{
+					AddArtifactFromTask(element, info.Arguments);
+				}
+			}
+		}
+
+		void AddArtifactFromTask(BgScriptElement element, Dictionary<string, string> arguments)
+		{
+			CreateArtifactTaskParameters parameters = new CreateArtifactTaskParameters();
+			foreach (PropertyInfo propertyInfo in typeof(CreateArtifactTaskParameters).GetProperties())
+			{
+				TaskParameterAttribute? attribute = propertyInfo.GetCustomAttribute<TaskParameterAttribute>();
+				if (attribute != null)
+				{
+					if (arguments.TryGetValue(propertyInfo.Name, out string? value))
+					{
+						propertyInfo.SetValue(parameters, value);
+					}
+					else if (!attribute.Optional)
+					{
+						LogError(element, $"Missing '{propertyInfo.Name}' property for CreateArtifact task.");
+						return;
+					}
+				}
 			}
+
+			string[] keys = (parameters.Keys ?? String.Empty).Split(';', StringSplitOptions.RemoveEmptyEntries | StringSplitOptions.TrimEntries);
+			string[] metadata = (parameters.Metadata ?? String.Empty).Split(';', StringSplitOptions.RemoveEmptyEntries | StringSplitOptions.TrimEntries);
+
+			BgArtifactDef artifact = new BgArtifactDef(parameters.Name, parameters.Type, parameters.Description, null, _enclosingNode!.Name, null, keys, metadata);
+			_graph.Artifacts.Add(artifact);
 		}
 
 		/// <summary>
diff --git a/Engine/Source/Programs/AutomationTool/BuildGraph/BuildGraph.cs b/Engine/Source/Programs/AutomationTool/BuildGraph/BuildGraph.cs
--- a/Engine/Source/Programs/AutomationTool/BuildGraph/BuildGraph.cs
+++ b/Engine/Source/Programs/AutomationTool/BuildGraph/BuildGraph.cs
@@ -1256,7 +1256,14 @@ namespace AutomationTool
 				}
 
 				// Find all the output tags that are published artifacts
-				Dictionary<string, BgArtifactDef> outputNameToArtifact = graph.Artifacts.ToDictionary(x => x.TagName, x => x);
+				Dictionary<string, BgArtifactDef> outputNameToArtifact = new Dictionary<string, BgArtifactDef>(StringComparer.OrdinalIgnoreCase);
+				foreach (BgArtifactDef artifact in graph.Artifacts)
+				{
+					if (artifact.TagName != null)
+					{
+						outputNameToArtifact.Add(artifact.TagName, artifact);
+					}
+				}
 
 				// Publish all the output tags
 				foreach (BgNodeOutput output in node.Outputs)
diff --git a/Engine/Source/Programs/Horde/Drivers/JobDriver/Execution/JobExecutor.cs b/Engine/Source/Programs/Horde/Drivers/JobDriver/Execution/JobExecutor.cs
--- a/Engine/Source/Programs/Horde/Drivers/JobDriver/Execution/JobExecutor.cs
+++ b/Engine/Source/Programs/Horde/Drivers/JobDriver/Execution/JobExecutor.cs
@@ -124,6 +124,7 @@ namespace JobDriver.Execution
 			public string? BasePath { get; set; }
 			public List<string> Keys { get; set; } = new List<string>();
 			public List<string> Metadata { get; set; } = new List<string>();
+			public string NodeName { get; set; } = String.Empty;
 			public string OutputName { get; set; } = String.Empty;
 		}
 
@@ -1032,6 +1033,7 @@ namespace JobDriver.Execution
 				createArtifact.BasePath = exportedArtifact.BasePath ?? String.Empty;
 				createArtifact.Keys.AddRange(exportedArtifact.Keys);
 				createArtifact.Metadata.AddRange(exportedArtifact.Metadata);
+				createArtifact.NodeName = exportedArtifact.NodeName;
 				createArtifact.OutputName = exportedArtifact.OutputName;
 
 				updateGraph.Artifacts.Add(createArtifact);
diff --git a/Engine/Source/Programs/Horde/Plugins/Build/HordeServer.Build.Tests/Jobs/JobCollectionTests.cs b/Engine/Source/Programs/Horde/Plugins/Build/HordeServer.Build.Tests/Jobs/JobCollectionTests.cs
--- a/Engine/Source/Programs/Horde/Plugins/Build/HordeServer.Build.Tests/Jobs/JobCollectionTests.cs
+++ b/Engine/Source/Programs/Horde/Plugins/Build/HordeServer.Build.Tests/Jobs/JobCollectionTests.cs
@@ -832,7 +832,7 @@ namespace HordeServer.Tests.Jobs
 			IGraph baseGraph = await GraphCollection.AddAsync(templateMock.Object, null);
 
 			List<NewGraphArtifact> newArtifacts = new List<NewGraphArtifact>();
-			newArtifacts.Add(new NewGraphArtifact(new ArtifactName("foo"), new ArtifactType("type"), "hello world", "Engine/Source", new List<string>(), new List<string>(), "fileset"));
+			newArtifacts.Add(new NewGraphArtifact(new ArtifactName("foo"), new ArtifactType("type"), "hello world", "Engine/Source", new List<string>(), new List<string>(), null, "fileset"));
 
 			IGraph graph = await GraphCollection.AppendAsync(baseGraph, newArtifactRequests: newArtifacts);
 			Assert.AreEqual(1, graph.Artifacts.Count);
diff --git a/Engine/Source/Programs/Horde/Plugins/Build/HordeServer.Build/Jobs/Graphs/GraphCollection.cs b/Engine/Source/Programs/Horde/Plugins/Build/HordeServer.Build/Jobs/Graphs/GraphCollection.cs
--- a/Engine/Source/Programs/Horde/Plugins/Build/HordeServer.Build/Jobs/Graphs/GraphCollection.cs
+++ b/Engine/Source/Programs/Horde/Plugins/Build/HordeServer.Build/Jobs/Graphs/GraphCollection.cs
@@ -194,7 +194,8 @@ namespace HordeServer.Jobs.Graphs
 			public string BasePath { get; set; }
 			public List<string> Keys { get; set; }
 			public List<string> Metadata { get; set; }
-			public string OutputName { get; set; }
+			public string? NodeName { get; set; }
+			public string? OutputName { get; set; }
 
 			IReadOnlyList<string> IGraphArtifact.Keys => Keys;
 			IReadOnlyList<string> IGraphArtifact.Metadata => Metadata;
@@ -205,10 +206,9 @@ namespace HordeServer.Jobs.Graphs
 				BasePath = String.Empty;
 				Keys = new List<string>();
 				Metadata = new List<string>();
-				OutputName = String.Empty;
 			}
 
-			public GraphArtifact(ArtifactName name, ArtifactType type, string description, string basePath, IReadOnlyList<string> keys, IReadOnlyList<string> metadata, string outputName)
+			public GraphArtifact(ArtifactName name, ArtifactType type, string description, string basePath, IReadOnlyList<string> keys, IReadOnlyList<string> metadata, string? nodeName, string? outputName)
 			{
 				Name = name;
 				Type = type;
@@ -216,6 +216,7 @@ namespace HordeServer.Jobs.Graphs
 				BasePath = basePath;
 				Keys = keys.ToList();
 				Metadata = metadata.ToList();
+				NodeName = nodeName;
 				OutputName = outputName;
 			}
 		}
@@ -335,7 +336,7 @@ namespace HordeServer.Jobs.Graphs
 				{
 					foreach (NewGraphArtifact newArtifactRequest in newArtifactRequests)
 					{
-						newArtifacts.Add(new GraphArtifact(newArtifactRequest.Name, newArtifactRequest.Type, newArtifactRequest.Description, newArtifactRequest.BasePath, newArtifactRequest.Keys, newArtifactRequest.Metadata, newArtifactRequest.OutputName));
+						newArtifacts.Add(new GraphArtifact(newArtifactRequest.Name, newArtifactRequest.Type, newArtifactRequest.Description, newArtifactRequest.BasePath, newArtifactRequest.Keys, newArtifactRequest.Metadata, newArtifactRequest.NodeName, newArtifactRequest.OutputName));
 					}
 				}
 
diff --git a/Engine/Source/Programs/Horde/Plugins/Build/HordeServer.Build/Jobs/Graphs/IGraph.cs b/Engine/Source/Programs/Horde/Plugins/Build/HordeServer.Build/Jobs/Graphs/IGraph.cs
--- a/Engine/Source/Programs/Horde/Plugins/Build/HordeServer.Build/Jobs/Graphs/IGraph.cs
+++ b/Engine/Source/Programs/Horde/Plugins/Build/HordeServer.Build/Jobs/Graphs/IGraph.cs
@@ -451,10 +451,15 @@ namespace HordeServer.Jobs.Graphs
 		/// </summary>
 		public IReadOnlyList<string> Metadata { get; }
 
+		/// <summary>
+		/// Name of the node producing this artifact
+		/// </summary>
+		public string? NodeName { get; }
+
 		/// <summary>
 		/// Tag for the artifact files
 		/// </summary>
-		public string OutputName { get; }
+		public string? OutputName { get; }
 	}
 
 	/// <summary>
@@ -687,5 +692,5 @@ namespace HordeServer.Jobs.Graphs
 	/// <summary>
 	/// Information about an artifact
 	/// </summary>
-	public record class NewGraphArtifact(ArtifactName Name, ArtifactType Type, string Description, string BasePath, List<string> Keys, List<string> Metadata, string OutputName);
+	public record class NewGraphArtifact(ArtifactName Name, ArtifactType Type, string Description, string BasePath, List<string> Keys, List<string> Metadata, string? NodeName, string? OutputName);
 }
diff --git a/Engine/Source/Programs/Horde/Plugins/Build/HordeServer.Build/Jobs/JobRpcService.cs b/Engine/Source/Programs/Horde/Plugins/Build/HordeServer.Build/Jobs/JobRpcService.cs
--- a/Engine/Source/Programs/Horde/Plugins/Build/HordeServer.Build/Jobs/JobRpcService.cs
+++ b/Engine/Source/Programs/Horde/Plugins/Build/HordeServer.Build/Jobs/JobRpcService.cs
@@ -460,7 +460,7 @@ namespace HordeServer.Jobs
 
 				foreach (IGraphArtifact artifact in graph.Artifacts)
 				{
-					if (node.OutputNames.Contains(artifact.OutputName))
+					if (artifact.OutputName != null && node.OutputNames.Contains(artifact.OutputName))
 					{
 						RpcCreateGraphArtifactRequest stepArtifact = new RpcCreateGraphArtifactRequest { Name = artifact.Name.ToString(), Type = artifact.Type.ToString(), Description = artifact.Description, BasePath = artifact.BasePath, OutputName = artifact.OutputName };
 						stepArtifact.Keys.AddRange(artifact.Keys);
@@ -679,7 +679,19 @@ namespace HordeServer.Jobs
 						description = artifact.Name;
 					}
 
-					newArtifacts.Add(new NewGraphArtifact(name, type, description, artifact.BasePath, artifact.Keys.ToList(), artifact.Metadata.ToList(), artifact.OutputName));
+					string? nodeName = null;
+					if (!String.IsNullOrEmpty(artifact.NodeName))
+					{
+						nodeName = artifact.NodeName;
+					}
+
+					string? outputName = null;
+					if (!String.IsNullOrEmpty(artifact.OutputName))
+					{
+						outputName = artifact.OutputName;
+					}
+
+					newArtifacts.Add(new NewGraphArtifact(name, type, description, artifact.BasePath, artifact.Keys.ToList(), artifact.Metadata.ToList(), nodeName, outputName));
 				}
 
 				// Create the new graph
diff --git a/Engine/Source/Programs/Horde/Plugins/Build/HordeServer.Build/Jobs/JobsController.cs b/Engine/Source/Programs/Horde/Plugins/Build/HordeServer.Build/Jobs/JobsController.cs
--- a/Engine/Source/Programs/Horde/Plugins/Build/HordeServer.Build/Jobs/JobsController.cs
+++ b/Engine/Source/Programs/Horde/Plugins/Build/HordeServer.Build/Jobs/JobsController.cs
@@ -554,10 +554,18 @@ namespace HordeServer.Jobs
 					}
 				}
 
+				Dictionary<string, IGraphArtifact> nodeNameToArtifact = new Dictionary<string, IGraphArtifact>(StringComparer.OrdinalIgnoreCase);
 				Dictionary<string, IGraphArtifact> outputNameToArtifact = new Dictionary<string, IGraphArtifact>(StringComparer.OrdinalIgnoreCase);
 				foreach (IGraphArtifact artifact in graph.Artifacts)
 				{
-					outputNameToArtifact[artifact.OutputName] = artifact;
+					if (artifact.NodeName != null)
+					{
+						nodeNameToArtifact[artifact.NodeName] = artifact;
+					}
+					if (artifact.OutputName != null)
+					{
+						outputNameToArtifact[artifact.OutputName] = artifact;
+					}
 				}
 
 				foreach (IJobStepBatch batch in job.Batches)
@@ -566,9 +574,14 @@ namespace HordeServer.Jobs
 					foreach (IJobStep step in batch.Steps)
 					{
 						INode node = group.Nodes[step.NodeIdx];
+
+						IGraphArtifact? graphArtifact;
+						if (nodeNameToArtifact.TryGetValue(node.Name, out graphArtifact) && !addedArtifacts.Contains((graphArtifact.Name, step.Id)))
+						{
+							response.Artifacts.Add(new GetJobArtifactResponse(null, graphArtifact.Name, graphArtifact.Type, graphArtifact.Description, graphArtifact.Keys.ToList(), graphArtifact.Metadata.ToList(), step.Id));
+						}
 						foreach (string outputName in node.OutputNames)
 						{
-							IGraphArtifact? graphArtifact;
 							if (outputNameToArtifact.TryGetValue(outputName, out graphArtifact) && !addedArtifacts.Contains((graphArtifact.Name, step.Id)))
 							{
 								response.Artifacts.Add(new GetJobArtifactResponse(null, graphArtifact.Name, graphArtifact.Type, graphArtifact.Description, graphArtifact.Keys.ToList(), graphArtifact.Metadata.ToList(), step.Id));
diff --git a/Engine/Source/Programs/Shared/EpicGames.BuildGraph/BgArtifactDef.cs b/Engine/Source/Programs/Shared/EpicGames.BuildGraph/BgArtifactDef.cs
--- a/Engine/Source/Programs/Shared/EpicGames.BuildGraph/BgArtifactDef.cs
+++ b/Engine/Source/Programs/Shared/EpicGames.BuildGraph/BgArtifactDef.cs
@@ -29,10 +29,15 @@ namespace EpicGames.BuildGraph
 		/// </summary>
 		public string? BasePath { get; }
 
+		/// <summary>
+		/// Node that produces the artifact. Either this or TagName should be set.
+		/// </summary>
+		public string? NodeName { get; }
+
 		/// <summary>
 		/// Tag to use for the artifact. Uses the artifact name by default.
 		/// </summary>
-		public string TagName { get; }
+		public string? TagName { get; }
 
 		/// <summary>
 		/// Keys that can be used to find the artifact
@@ -51,15 +56,17 @@ namespace EpicGames.BuildGraph
 		/// <param name="type">Type of the artifact</param>
 		/// <param name="description">Description for the artifact</param>
 		/// <param name="basePath">Base path for files included in the artifact</param>
+		/// <param name="nodeName">Name of the node producing this artifact</param>
 		/// <param name="tagName">Name of the tag producing this artifact</param>
 		/// <param name="keys">Keys that can be used to find the artifact</param>
 		/// <param name="metadata">Metadata for the artifact</param>
-		public BgArtifactDef(string name, string? type, string? description, string? basePath, string tagName, IReadOnlyList<string> keys, IReadOnlyList<string> metadata)
+		public BgArtifactDef(string name, string? type, string? description, string? basePath, string? nodeName, string? tagName, IReadOnlyList<string> keys, IReadOnlyList<string> metadata)
 		{
 			Name = name;
 			Type = type;
 			Description = description;
 			BasePath = basePath;
+			NodeName = nodeName;
 			TagName = tagName;
 			Keys = keys;
 			Metadata = metadata;
diff --git a/Engine/Source/Programs/Shared/EpicGames.BuildGraph/BgGraphDef.cs b/Engine/Source/Programs/Shared/EpicGames.BuildGraph/BgGraphDef.cs
--- a/Engine/Source/Programs/Shared/EpicGames.BuildGraph/BgGraphDef.cs
+++ b/Engine/Source/Programs/Shared/EpicGames.BuildGraph/BgGraphDef.cs
@@ -6,6 +6,7 @@ using System.Diagnostics.CodeAnalysis;
 using System.Linq;
 using System.Runtime.CompilerServices;
 using System.Text;
+using System.Xml.Linq;
 using EpicGames.Core;
 using Microsoft.Extensions.Logging;
 
@@ -307,8 +308,12 @@ namespace EpicGames.BuildGraph
 			}
 			TagNameToNodeOutput = newTagNameToNodeOutput;
 
-			// Remove any artifacts which are not produced
-			Artifacts.RemoveAll(x => !TagNameToNodeOutput.ContainsKey(x.TagName));
+			// Remove any artifacts whose outputs are not produced
+			Artifacts.RemoveAll(x => x.TagName != null && !TagNameToNodeOutput.ContainsKey(x.TagName));
+
+			// Remove any artifacts whose nodes are no longer run
+			HashSet<string> newNodeNames = Agents.SelectMany(x => x.Nodes).Select(x => x.Name).ToHashSet(StringComparer.OrdinalIgnoreCase);
+			Artifacts.RemoveAll(x => x.NodeName != null && !newNodeNames.Contains(x.NodeName));
 		}
 
 		/// <summary>
@@ -457,21 +462,35 @@ namespace EpicGames.BuildGraph
 						jsonWriter.WriteValue("BasePath", artifact.BasePath);
 					}
 
-					jsonWriter.WriteArrayStart("Keys");
-					foreach (string key in artifact.Keys)
+					if (artifact.Keys.Count > 0)
 					{
-						jsonWriter.WriteValue(key);
+						jsonWriter.WriteArrayStart("Keys");
+						foreach (string key in artifact.Keys)
+						{
+							jsonWriter.WriteValue(key);
+						}
+						jsonWriter.WriteArrayEnd();
+					}
+
+					if (artifact.Metadata.Count > 0)
+					{
+						jsonWriter.WriteArrayStart("Metadata");
+						foreach (string metadata in artifact.Metadata)
+						{
+							jsonWriter.WriteValue(metadata);
+						}
+						jsonWriter.WriteArrayEnd();
 					}
-					jsonWriter.WriteArrayEnd();
 
-					jsonWriter.WriteArrayStart("Metadata");
-					foreach (string metadata in artifact.Metadata)
+					if (!String.IsNullOrEmpty(artifact.NodeName))
 					{
-						jsonWriter.WriteValue(metadata);
+						jsonWriter.WriteValue("NodeName", artifact.NodeName);
+					}
+					if (!String.IsNullOrEmpty(artifact.TagName))
+					{
+						jsonWriter.WriteValue("OutputName", artifact.TagName);
 					}
-					jsonWriter.WriteArrayEnd();
 
-					jsonWriter.WriteValue("OutputName", artifact.TagName);
 					jsonWriter.WriteObjectEnd();
 				}
 				jsonWriter.WriteArrayEnd();
