Horde: Include keys for the current job when using the CreateArtifact task.

#rnx

diff --git a/Engine/Source/Programs/AutomationTool/BuildGraph/Tasks/CreateArtifactTask.cs b/Engine/Source/Programs/AutomationTool/BuildGraph/Tasks/CreateArtifactTask.cs
--- a/Engine/Source/Programs/AutomationTool/BuildGraph/Tasks/CreateArtifactTask.cs
+++ b/Engine/Source/Programs/AutomationTool/BuildGraph/Tasks/CreateArtifactTask.cs
@@ -105,8 +105,21 @@ namespace AutomationTool.Tasks
 		{
 			ArtifactName name = new ArtifactName(_parameters.Name);
 			ArtifactType type = new ArtifactType(_parameters.Type);
-			string[] keys = (_parameters.Keys ?? string.Empty).Split(';', StringSplitOptions.RemoveEmptyEntries | StringSplitOptions.TrimEntries);
-			string[] metadata = (_parameters.Metadata ?? string.Empty).Split(';', StringSplitOptions.RemoveEmptyEntries | StringSplitOptions.TrimEntries);
+			List<string> keys = (_parameters.Keys ?? string.Empty).Split(';', StringSplitOptions.RemoveEmptyEntries | StringSplitOptions.TrimEntries).ToList();
+			List<string> metadata = (_parameters.Metadata ?? string.Empty).Split(';', StringSplitOptions.RemoveEmptyEntries | StringSplitOptions.TrimEntries).ToList();
+
+			// Add keys for the job that's executing
+			string? jobId = Environment.GetEnvironmentVariable("UE_HORDE_JOBID");
+			if (!String.IsNullOrEmpty(jobId))
+			{
+				keys.Add($"job:{jobId}");
+
+				string? stepId = Environment.GetEnvironmentVariable("UE_HORDE_STEPID");
+				if (!String.IsNullOrEmpty(stepId))
+				{
+					keys.Add($"job:{jobId}/step:{stepId}");
+				}
+			}
 
 			// Figure out the current change and stream id
 			StreamId streamId;
