Horde: Add an error if a name or type is not specified for an artifact.

#rnx

diff --git a/Engine/Source/Programs/AutomationTool/BuildGraph/BgScriptReader.cs b/Engine/Source/Programs/AutomationTool/BuildGraph/BgScriptReader.cs
--- a/Engine/Source/Programs/AutomationTool/BuildGraph/BgScriptReader.cs
+++ b/Engine/Source/Programs/AutomationTool/BuildGraph/BgScriptReader.cs
@@ -5,7 +5,6 @@ using System.Collections.Generic;
 using System.Diagnostics.CodeAnalysis;
 using System.IO;
 using System.Linq;
-using System.Reflection;
 using System.Text;
 using System.Text.RegularExpressions;
 using System.Threading.Tasks;
@@ -1760,10 +1759,21 @@ namespace AutomationTool
 			arguments.TryGetValue(nameof(CreateArtifactTaskParameters.Keys), out string? keys);
 			arguments.TryGetValue(nameof(CreateArtifactTaskParameters.Metadata), out string? metadata);
 
+			if (String.IsNullOrEmpty(name))
+			{
+				LogError(element, "Missing artifact name");
+				return;
+			}
+			if (String.IsNullOrEmpty(type))
+			{
+				LogError(element, "Missing artifact type");
+				return;
+			}
+
 			string[] keysArray = (keys ?? String.Empty).Split(';', StringSplitOptions.RemoveEmptyEntries | StringSplitOptions.TrimEntries);
 			string[] metadataArray = (metadata ?? String.Empty).Split(';', StringSplitOptions.RemoveEmptyEntries | StringSplitOptions.TrimEntries);
 
-			BgArtifactDef artifact = new BgArtifactDef(name ?? String.Empty, type, description, null, _enclosingNode!.Name, null, keysArray, metadataArray);
+			BgArtifactDef artifact = new BgArtifactDef(name, type, description, null, _enclosingNode!.Name, null, keysArray, metadataArray);
 			_graph.Artifacts.Add(artifact);
 		}
 
