RPC Changes to get more accurate reponse from CheatCommand and to reduce low level logging to verbose

#rnx

diff --git a/Engine/Source/Programs/AutomationTool/Gauntlet/Framework/RpcFramework/Gauntlet.RpcExecutor.cs b/Engine/Source/Programs/AutomationTool/Gauntlet/Framework/RpcFramework/Gauntlet.RpcExecutor.cs
--- a/Engine/Source/Programs/AutomationTool/Gauntlet/Framework/RpcFramework/Gauntlet.RpcExecutor.cs
+++ b/Engine/Source/Programs/AutomationTool/Gauntlet/Framework/RpcFramework/Gauntlet.RpcExecutor.cs
@@ -218,7 +218,7 @@ namespace Gauntlet
 
 
 			RouteMapping CallPath = new RouteMapping(ListenerContext.Request.Url.LocalPath, ListenerContext.Request.HttpMethod);
-			Log.Info(string.Format("In ReceiveListenRequest callback for {0} | {1}", CallPath.Path, CallPath.Verb));
+			Log.Verbose(string.Format("In ReceiveListenRequest callback for {0} | {1}", CallPath.Path, CallPath.Verb));
 			ResponseData CallResponse = new ResponseData();
 			CallResponse.ResponseBody = "This endpoint is not handled";
 
diff --git a/Engine/Source/Programs/AutomationTool/Gauntlet/Framework/RpcFramework/Gauntlet.RpcLibrary.cs b/Engine/Source/Programs/AutomationTool/Gauntlet/Framework/RpcFramework/Gauntlet.RpcLibrary.cs
--- a/Engine/Source/Programs/AutomationTool/Gauntlet/Framework/RpcFramework/Gauntlet.RpcLibrary.cs
+++ b/Engine/Source/Programs/AutomationTool/Gauntlet/Framework/RpcFramework/Gauntlet.RpcLibrary.cs
@@ -214,8 +214,8 @@ namespace Gauntlet
 			}
 			Dictionary<string, object> RequestArgs = new Dictionary<string, object>();
 			RequestArgs.Add("command", CheatString);
-			RpcExecutor.CallRpc(InTarget, "CheatCommand", 5000, RequestArgs);
-			return true;
+			HttpResponseMessage Response = RpcExecutor.CallRpc(InTarget, "CheatCommand", 5000, RequestArgs);
+			return Response != null && Response.IsSuccessStatusCode;
 		}
 
 
