[Gauntlet] Install bulk content with -n while running on a build machine

diff --git a/Engine/Source/Programs/AutomationTool/Gauntlet/Platform/IOS/Gauntlet.TargetDeviceIOS.cs b/Engine/Source/Programs/AutomationTool/Gauntlet/Platform/IOS/Gauntlet.TargetDeviceIOS.cs
--- a/Engine/Source/Programs/AutomationTool/Gauntlet/Platform/IOS/Gauntlet.TargetDeviceIOS.cs
+++ b/Engine/Source/Programs/AutomationTool/Gauntlet/Platform/IOS/Gauntlet.TargetDeviceIOS.cs
@@ -552,6 +552,12 @@ namespace Gauntlet
 				string IDeviceFS = Path.Combine(Globals.UnrealRootDir, "Engine", "Extras", "ThirdPartyNotUE", "libimobiledevice", "mac", "idevicefs");
 				string BulkContentCopyCommand = string.Format("-u {0} -b {1} -x RequiredCommands.txt", UUID, Build.PackageName);
 
+				if(IsBuildMachine)
+				{
+					// Required for forcing install via pre-paired network device instead of USB
+					BulkContentCopyCommand += " -n";
+				}
+
 				Result = ExecuteDeploymentCommand(IDeviceFS, BulkContentCopyCommand, MaxInstallTime, true, BulkBuildPath, ERunOptions.NoStdOutRedirect);
 
 				if(Result.ExitCode != 0)
