Fix BuildCookRun -Runtest not matching client exit properly
* Use softquit to allow log line to be output in time so that BuildCookRun can match it

#rnx

diff --git a/Engine/Source/Programs/AutomationTool/Scripts/RunProjectCommand.Automation.cs b/Engine/Source/Programs/AutomationTool/Scripts/RunProjectCommand.Automation.cs
--- a/Engine/Source/Programs/AutomationTool/Scripts/RunProjectCommand.Automation.cs
+++ b/Engine/Source/Programs/AutomationTool/Scripts/RunProjectCommand.Automation.cs
@@ -660,11 +660,11 @@ namespace AutomationScripts
 			}
 			else if (Params.RunAutomationTest != "")
 			{
-				TempCmdLine += "-execcmds=\"automation runtests " + Params.RunAutomationTest + ";quit\" ";
+				TempCmdLine += "-execcmds=\"automation runtests " + Params.RunAutomationTest + ";softquit\" ";
 			}
 			else if (Params.RunAutomationTests)
 			{
-				TempCmdLine += "-execcmds=\"automation runall;quit;\" ";
+				TempCmdLine += "-execcmds=\"automation runall;softquit;\" ";
 			}
 			if (SC.StageTargetPlatform.UseAbsLog)
 			{
