Added a simple UAT command LaunchMultiServer to automatically launch multiple MultiServerReplication servers based on config/ini files.
Added UMultiServerSettings config class to define the server launch parameters. It's in a new module that doesn't depend on Engine in case we want to use it in a C++ program that doesn't depend on engine.

diff --git a/Engine/Plugins/Runtime/MultiServerReplication/MultiServerReplication.uplugin b/Engine/Plugins/Runtime/MultiServerReplication/MultiServerReplication.uplugin
--- a/Engine/Plugins/Runtime/MultiServerReplication/MultiServerReplication.uplugin
+++ b/Engine/Plugins/Runtime/MultiServerReplication/MultiServerReplication.uplugin
@@ -16,6 +16,11 @@
 			"Name": "MultiServerReplication",
 			"Type": "Runtime",
 			"LoadingPhase": "Default"
+		},
+		{
+			"Name": "MultiServerConfiguration",
+			"Type": "Runtime",
+			"LoadingPhase": "Default"
 		}
 	],
 	"Plugins": [
diff --git a/Engine/Plugins/Runtime/MultiServerReplication/Source/MultiServerConfiguration/MultiServerConfiguration.Build.cs b/Engine/Plugins/Runtime/MultiServerReplication/Source/MultiServerConfiguration/MultiServerConfiguration.Build.cs
new file mode 100644
--- /dev/null
+++ b/Engine/Plugins/Runtime/MultiServerReplication/Source/MultiServerConfiguration/MultiServerConfiguration.Build.cs
@@ -0,0 +1,16 @@
+// Copyright Epic Games, Inc. All Rights Reserved.
+
+using UnrealBuildTool;
+
+public class MultiServerConfiguration : ModuleRules
+{
+	public MultiServerConfiguration(ReadOnlyTargetRules Target) : base(Target)
+    {
+        PrivateDependencyModuleNames.AddRange(
+			new string[] {
+                "Core",
+                "CoreUObject"
+            }
+        );
+    }
+}
diff --git a/Engine/Plugins/Runtime/MultiServerReplication/Source/MultiServerConfiguration/Private/MultiServerConfigurationModule.cpp b/Engine/Plugins/Runtime/MultiServerReplication/Source/MultiServerConfiguration/Private/MultiServerConfigurationModule.cpp
new file mode 100644
--- /dev/null
+++ b/Engine/Plugins/Runtime/MultiServerReplication/Source/MultiServerConfiguration/Private/MultiServerConfigurationModule.cpp
@@ -0,0 +1,6 @@
+// Copyright Epic Games, Inc. All Rights Reserved.
+
+#include "MultiServerConfigurationModule.h"
+#include "Modules/ModuleManager.h"
+
+IMPLEMENT_MODULE(FMultiServerConfigurationModule, MultiServerConfiguration);
\ No newline at end of file
diff --git a/Engine/Plugins/Runtime/MultiServerReplication/Source/MultiServerConfiguration/Public/MultiServerConfigurationModule.h b/Engine/Plugins/Runtime/MultiServerReplication/Source/MultiServerConfiguration/Public/MultiServerConfigurationModule.h
new file mode 100644
--- /dev/null
+++ b/Engine/Plugins/Runtime/MultiServerReplication/Source/MultiServerConfiguration/Public/MultiServerConfigurationModule.h
@@ -0,0 +1,20 @@
+// Copyright Epic Games, Inc. All Rights Reserved.
+
+#pragma once
+
+#include "Modules/ModuleInterface.h"
+
+/** See MultiServerNode for more details about multi-server networking as a whole. */
+class FMultiServerConfigurationModule : public IModuleInterface
+{
+public:
+
+	FMultiServerConfigurationModule() {}
+	virtual ~FMultiServerConfigurationModule() {}
+
+	// IModuleInterface
+	virtual bool SupportsDynamicReloading() override { return false; }
+	virtual bool SupportsAutomaticShutdown() override { return false; }
+};
+
+
diff --git a/Engine/Plugins/Runtime/MultiServerReplication/Source/MultiServerConfiguration/Public/MultiServerSettings.h b/Engine/Plugins/Runtime/MultiServerReplication/Source/MultiServerConfiguration/Public/MultiServerSettings.h
new file mode 100644
--- /dev/null
+++ b/Engine/Plugins/Runtime/MultiServerReplication/Source/MultiServerConfiguration/Public/MultiServerSettings.h
@@ -0,0 +1,44 @@
+// Copyright Epic Games, Inc. All Rights Reserved.
+
+#pragma once
+
+#include "UObject/Object.h"
+#include "Containers/Array.h"
+#include "MultiServerSettings.generated.h"
+
+class FString;
+
+/** Parameters for a single server process in a multi-server setup. */
+USTRUCT()
+struct FMultiServerDefinition
+{
+	GENERATED_BODY()
+
+	/** User-specified identifier that must be unique for each server in a MultiServer cluster. */
+	UPROPERTY(Config)
+	FString LocalId;
+
+	/** UDP port on which the MultiServer beacon will listen for new connections. */
+	UPROPERTY(Config)
+	uint16 ListenPort = 0;
+};
+
+/**
+ * Settings that control how a multi-server setup is configured and launched.
+ * For example, to define 4 servers in a Game ini file it might look like this:
+ * [/Script/MultiServerConfiguration.MultiServerSettings]
+ * !ServerDefinitions=ClearArray
+ * +ServerDefinitions=(LocalId="0",ListenPort=17000)
+ * +ServerDefinitions=(LocalId="1",ListenPort=17001)
+ * +ServerDefinitions=(LocalId="2",ListenPort=17002)
+ * +ServerDefinitions=(LocalId="3",ListenPort=17003)
+ */
+UCLASS(Config=Game)
+class MULTISERVERCONFIGURATION_API UMultiServerSettings : public UObject
+{
+	GENERATED_BODY()
+
+public:
+	UPROPERTY(Config)
+	TArray<FMultiServerDefinition> ServerDefinitions;
+};
diff --git a/Engine/Plugins/Runtime/MultiServerReplication/Source/MultiServerReplication/MultiServerReplication.Build.cs b/Engine/Plugins/Runtime/MultiServerReplication/Source/MultiServerReplication/MultiServerReplication.Build.cs
--- a/Engine/Plugins/Runtime/MultiServerReplication/Source/MultiServerReplication/MultiServerReplication.Build.cs
+++ b/Engine/Plugins/Runtime/MultiServerReplication/Source/MultiServerReplication/MultiServerReplication.Build.cs
@@ -11,7 +11,8 @@ public class MultiServerReplication : ModuleRules
                 "Core",
                 "CoreUObject",
                 "Engine",
-				"NetCore"
+				"NetCore",
+				"MultiServerConfiguration"
             }
         );
 
diff --git a/Engine/Source/Programs/AutomationTool/Scripts/LaunchMultiServer.Automation.cs b/Engine/Source/Programs/AutomationTool/Scripts/LaunchMultiServer.Automation.cs
new file mode 100644
--- /dev/null
+++ b/Engine/Source/Programs/AutomationTool/Scripts/LaunchMultiServer.Automation.cs
@@ -0,0 +1,125 @@
+// Copyright Epic Games, Inc. All Rights Reserved.
+
+using System;
+using System.IO;
+using UnrealBuildBase;
+using UnrealBuildTool;
+using EpicGames.Core;
+using Microsoft.Extensions.Logging;
+using System.Collections.Generic;
+
+namespace AutomationTool
+{
+
+	// Mirrors FMultiServerDefinition from the MultiServerReplication plugin
+	class MultiServerDefinition
+	{
+		public string LocalId = "0";
+		public int ListenPort = 0;
+
+		public override string ToString()
+		{
+			return "MultiServerDefinition: LocalId \"" + LocalId + "\", ListenPort " + ListenPort;
+		}
+	}
+
+	[Help("Launches multiple server processes for a project using the MultiServerReplication plugin.")]
+	[Help("project=<project>", "Project to open. Will search current path and paths in ueprojectdirs.")]
+	[Help("map=<MapName>", "Map to load on startup.")]
+	public class LaunchMultiServer : BuildCommand, IProjectParamsHelpers
+	{
+		public override ExitCode Execute()
+		{
+			Logger.LogInformation("********** RUN MULTISERVER COMMAND STARTED **********");
+			var StartTime = DateTime.UtcNow;
+
+			var MapToRun = ParseParamValue("map", "");
+
+			// Parse server configuration from ini files
+			ConfigHierarchy ProjectGameConfig = ConfigCache.ReadHierarchy(ConfigHierarchyType.Game, DirectoryReference.FromFile(ProjectPath), UnrealTargetPlatform.Win64);
+
+			const String ServerDefConfigSection = "/Script/MultiServerConfiguration.MultiServerSettings";
+
+			ProjectGameConfig.TryGetValuesGeneric(ServerDefConfigSection, "ServerDefinitions", out MultiServerDefinition[] ConfigServerDefs);
+
+			if (ConfigServerDefs != null && ConfigServerDefs.Length > 0)
+			{
+				// Use a dictionary to store ServerDefs by unqiue ID
+				Dictionary<string, MultiServerDefinition> UniqueServerDefs = new Dictionary<string, MultiServerDefinition>();
+
+				// Check for duplicates and build peer list
+				string PeerList = "";
+				for (int i = 0; i < ConfigServerDefs.Length; i++)
+				{
+					if (UniqueServerDefs.TryAdd(ConfigServerDefs[i].LocalId, ConfigServerDefs[i]))
+					{
+						PeerList += String.Format("127.0.0.1:{0}{1}", ConfigServerDefs[i].ListenPort, i < (ConfigServerDefs.Length - 1) ? "," : "");
+					}
+					else
+					{
+						Logger.LogInformation("MultiServer server defintion with duplicate ID {0}, ignoring duplicate entry {1}", ConfigServerDefs[i].LocalId, ConfigServerDefs[i]);
+					}
+				}
+
+				var CommonArgs = ProjectPath.FullName + " " + MapToRun + " -server -log -newconsole -MultiServerPeers=" + PeerList;
+
+				var ServerApp = CombinePaths(CmdEnv.LocalRoot, "Engine/Binaries/Win64/UnrealEditor.exe");
+				PushDir(Path.GetDirectoryName(ServerApp));
+
+				try
+				{
+					foreach (var ServerDef in UniqueServerDefs.Values)
+					{
+						Logger.LogInformation("Starting MultiServer instance with LocalId {0} and port {1}", ServerDef.LocalId, ServerDef.ListenPort);
+
+						var Args = CommonArgs + String.Format(" -MultiServerLocalId={0} -MultiServerListenPort={1}", ServerDef.LocalId, ServerDef.ListenPort);
+						var ServerProcess = Run(ServerApp, Args, null, ERunOptions.Default | ERunOptions.NoWaitForExit | ERunOptions.NoStdOutRedirect);
+
+						if (ServerProcess != null)
+						{
+							// Remove started process so it won't be killed on UAT exit.
+							// Essentially forces the -NoKill command-line option behavior for these.
+							ProcessManager.RemoveProcess(ServerProcess);
+						}
+					}
+				}
+				catch
+				{
+					throw;
+				}
+				finally
+				{
+					PopDir();
+				}
+			}
+			else
+			{
+				Logger.LogInformation("No server configs found in {0}, not starting any servers.", ServerDefConfigSection);
+			}
+
+			Logger.LogInformation("Run command time: {0:0.00} s", (DateTime.UtcNow - StartTime).TotalMilliseconds / 1000);
+			Logger.LogInformation("********** RUN MULTISERVER COMMAND COMPLETED **********");
+
+			return ExitCode.Success;
+		}
+
+		private FileReference ProjectFullPath;
+		public virtual FileReference ProjectPath
+		{
+			get
+			{
+				if (ProjectFullPath == null)
+				{
+					ProjectFullPath = ParseProjectParam();
+
+					if (ProjectFullPath == null)
+					{
+						throw new AutomationException("No project file specified. Use -project=<project>.");
+					}
+				}
+
+				return ProjectFullPath;
+			}
+		}
+	}
+}
