[Automation] Avoid exceptions during reading properties of the Process object

#rnx

diff --git a/Engine/Source/Programs/AutomationTool/AutomationUtils/ProcessUtils.cs b/Engine/Source/Programs/AutomationTool/AutomationUtils/ProcessUtils.cs
--- a/Engine/Source/Programs/AutomationTool/AutomationUtils/ProcessUtils.cs
+++ b/Engine/Source/Programs/AutomationTool/AutomationUtils/ProcessUtils.cs
@@ -122,8 +122,19 @@ namespace AutomationTool
 			for (int ProcessIndex = ProcessesToKill.Count - 1; ProcessIndex >= 0; --ProcessIndex )
 			{
 				IProcess Process =  ProcessesToKill[ProcessIndex];
-				var ProcessName = Process.GetProcessName();
-				if (Process.HasExited)
+				var ProcessName = string.Empty;
+				bool bProcessHasExited = true;
+				try
+				{
+					ProcessName = Process.GetProcessName();
+					bProcessHasExited = Process.HasExited;
+				}
+				catch (InvalidOperationException Ex)
+				{
+					Logger.LogDebug("Exception accessing the Process properties:\n{Exception}", Ex.ToString());
+				}
+					
+				if (bProcessHasExited)
 				{
 					ProcessesToKill.RemoveAt(ProcessIndex);
 				}
@@ -138,7 +149,15 @@ namespace AutomationTool
 				Logger.LogDebug("Trying to kill {Arg0} spawned processes.", ProcessesToKill.Count);
 				foreach (var Proc in ProcessesToKill)
 				{
-					Logger.LogDebug("  {Arg0}", Proc.GetProcessName());
+					try
+					{
+						var ProcessName = Proc.GetProcessName();
+						Logger.LogDebug("  {Arg0}", ProcessName);
+					}
+					catch (InvalidOperationException Ex)
+					{
+						Logger.LogDebug("Exception accessing the Process name:\n{Exception}", Ex.ToString());
+					}
 				}
 				if (CommandUtils.IsBuildMachine)
 				{
@@ -184,9 +203,10 @@ namespace AutomationTool
 				}
 				foreach (var Proc in ProcessesToKill)
 				{
-					var ProcName = Proc.GetProcessName();
+					var ProcName = string.Empty;
 					try
 					{
+						ProcName = Proc.GetProcessName();
 						if (!Proc.HasExited)
 						{
 							Logger.LogDebug("Killing process: {ProcName}", ProcName);
