Fixed IAS upload to Zen on iOS

#rnx

diff --git a/Engine/Source/Programs/AutomationTool/Scripts/CopyBuildToStagingDirectory.Automation.cs b/Engine/Source/Programs/AutomationTool/Scripts/CopyBuildToStagingDirectory.Automation.cs
--- a/Engine/Source/Programs/AutomationTool/Scripts/CopyBuildToStagingDirectory.Automation.cs
+++ b/Engine/Source/Programs/AutomationTool/Scripts/CopyBuildToStagingDirectory.Automation.cs
@@ -4020,7 +4020,13 @@ namespace AutomationScripts
 						}
 					}
 
-					string PakPath = CombinePaths(SC.StageDirectory.FullName, SC.ShortProjectName, "Content", "Paks");
+					string PakPathPrefix = CombinePaths(SC.ShortProjectName, "Content", "Paks");
+					if (SC.StageTargetPlatform.PlatformType == UnrealTargetPlatform.IOS)
+					{
+						PakPathPrefix = CombinePaths("cookeddata", PakPathPrefix);
+						PakPathPrefix = PakPathPrefix.ToLower();
+					}
+					string PakPath = CombinePaths(SC.StageDirectory.FullName, PakPathPrefix);
 					string CloudDir = CombinePaths(SC.StageDirectory.FullName, "Cloud");
 					string ConfigFilePath = CombinePaths(CloudDir, "IoStoreOnDemand.ini");
 
