Improvements to Gauntlet's heartbeat system
- We're now keeping track of all client's heartbeat instead of just the first one.
- Test Node's LogCategories are being cached when the Test Starts instead of retrieving them on every tick.
- Formatting fixes in Gauntlet.UnrealTestNode.cs automatically done by VS

diff --git a/Engine/Source/Programs/AutomationTool/Gauntlet/Unreal/Base/Gauntlet.UnrealTestNode.cs b/Engine/Source/Programs/AutomationTool/Gauntlet/Unreal/Base/Gauntlet.UnrealTestNode.cs
--- a/Engine/Source/Programs/AutomationTool/Gauntlet/Unreal/Base/Gauntlet.UnrealTestNode.cs
+++ b/Engine/Source/Programs/AutomationTool/Gauntlet/Unreal/Base/Gauntlet.UnrealTestNode.cs
@@ -11,6 +11,7 @@ using EpicGames.Core;
 using Microsoft.Extensions.Logging;
 using AutomationUtils.Matchers;
 using static Gauntlet.HordeReport;
+using static Gauntlet.UnrealSessionInstance;
 
 namespace Gauntlet
 {
@@ -80,7 +81,7 @@ namespace Gauntlet
 		public override IEnumerable<string> GetWarnings()
 		{
 			IEnumerable<string> WarningList = TestNodeEvents.Where(E => E.IsWarning).Select(E => E.Summary);
-			
+
 			if (RoleResults != null)
 			{
 				WarningList = WarningList.Union(RoleResults.SelectMany(R => R.Events.Where(E => E.IsWarning)).Select(E => E.Summary));
@@ -96,13 +97,13 @@ namespace Gauntlet
 		public override IEnumerable<string> GetErrors()
 		{
 			IEnumerable<string> ErrorList = TestNodeEvents.Where(E => E.IsError).Select(E => E.Summary);
-			
+
 			if (RoleResults != null)
 			{
 				ErrorList = ErrorList.Union(RoleResults.SelectMany(R => R.Events.Where(E => E.IsError)).Select(E => E.Summary));
 			}
 
-			return ErrorList.ToArray();			
+			return ErrorList.ToArray();
 		}
 
 		/// <summary>
@@ -272,19 +273,22 @@ namespace Gauntlet
 		}
 
 		/// <summary>
-		/// Used to track how much of our server log has been written out
+		/// Used to track how much of each running app's log has been written out
 		/// </summary>
-		private int LastServerLogCount = 0;
+		private Dictionary<RoleInstance, int> RoleToLastLogCount = new Dictionary<RoleInstance, int>();
 
 		/// <summary>
-		/// Used to track how much of our client log has been written out
+		/// Returns the RoleToLastLogCount value from a given RoleInstance key.
+		/// If the key does not exist it will be created before returning it.
 		/// </summary>
-		private int LastClientLogCount = 0;
-
-		/// <summary>
-		/// Used to track how much of our editor log has been written out
-		/// </summary>
-		private int LastEditorLogCount = 0;
+		private int GetLastLogCountFromRole(RoleInstance Role)
+		{
+			if(!RoleToLastLogCount.ContainsKey(Role))
+			{
+				RoleToLastLogCount.Add(Role, 0);
+			}
+			return RoleToLastLogCount[Role];
+		}
 
 		private int CurrentPass;
 
@@ -353,6 +357,11 @@ namespace Gauntlet
 			SampleCommandlines.Add(new KeyValuePair<string, string>(Commandline, Description));
 		}
 
+		/// <summary>
+		/// List of log event categories to track with the heartbeat system.
+		/// </summary>
+		private List<string> HeartbeatLogCategories = new List<string>();
+
 		/// <summary>
 		/// Constructor. A context of the correct type is required
 		/// </summary>
@@ -363,9 +372,6 @@ namespace Gauntlet
 
 			UnrealTestResult = TestResult.Invalid;
 			TimeToWaitForProcesses = 5;
-			LastServerLogCount = 0;
-			LastClientLogCount = 0;
-			LastEditorLogCount = 0;
 			CurrentPass = 0;
 			NumPasses = 0;
 			TestVersion = new Version("1.0.0");
@@ -375,7 +381,7 @@ namespace Gauntlet
 			LogWarningsAndErrorsAfterSummary = false;
 		}
 
-		 ~UnrealTestNode()
+		~UnrealTestNode()
 		{
 			Dispose(false);
 		}
@@ -415,6 +421,7 @@ namespace Gauntlet
 
 			return string.Format("{0} ({1})", Name, GetMainRoleContextString());
 		}
+
 		public string GetMainRoleContextString()
 		{
 			if (Context == null)
@@ -607,7 +614,7 @@ namespace Gauntlet
 					UnrealTargetPlatform SessionPlatform = TestRole.PlatformOverride ?? RoleContext.Platform;
 
 					// Apply all role configurations
-					foreach(IUnrealRoleConfiguration RoleConfiguration in TestRole.RoleConfigurations)
+					foreach (IUnrealRoleConfiguration RoleConfiguration in TestRole.RoleConfigurations)
 					{
 						RoleConfiguration.ApplyConfigToRole(TestRole);
 					}
@@ -622,7 +629,7 @@ namespace Gauntlet
 					SessionRole.InstallOnly = TestRole.InstallOnly;
 					SessionRole.DeferredLaunch = TestRole.DeferredLaunch;
 					SessionRole.CommandLineParams = TestRole.CommandLineParams;
- 					SessionRole.RoleModifier = TestRole.RoleType;
+					SessionRole.RoleModifier = TestRole.RoleType;
 					SessionRole.Constraint = UseContextConstraint ? Context.Constraint : new UnrealDeviceTargetConstraint(SessionPlatform);
 					Log.Verbose("Created SessionRole {Role} from RoleContext {Context} (RoleType={Type})", SessionRole, RoleContext, TypesToRoles.Key);
 
@@ -645,8 +652,8 @@ namespace Gauntlet
 						}
 
 						// add controllers
-						SessionRole.CommandLineParams.Add("gauntlet", 
-							TestRole.Controllers.Count > 0 ? string.Join(",", TestRole.Controllers) : null);				
+						SessionRole.CommandLineParams.Add("gauntlet",
+							TestRole.Controllers.Count > 0 ? string.Join(",", TestRole.Controllers) : null);
 
 						if (PassThroughArgs.Count() > 0)
 						{
@@ -663,7 +670,7 @@ namespace Gauntlet
 					}
 
 					// copy over relevant settings from test role
-                    SessionRole.FilesToCopy = TestRole.FilesToCopy;
+					SessionRole.FilesToCopy = TestRole.FilesToCopy;
 					SessionRole.AdditionalArtifactDirectories = TestRole.AdditionalArtifactDirectories;
 					SessionRole.ConfigureDevice = TestRole.ConfigureDevice;
 					SessionRole.MapOverride = TestRole.MapOverride;
@@ -776,9 +783,6 @@ namespace Gauntlet
 			UnrealTestResult = TestResult.Invalid;
 			CurrentPass = Pass;
 			NumPasses = InNumPasses;
-			LastServerLogCount = 0;
-			LastClientLogCount = 0;
-			LastEditorLogCount = 0;
 			LastHeartbeatTime = DateTime.MinValue;
 			LastActiveHeartbeatTime = DateTime.MinValue;
 
@@ -793,7 +797,7 @@ namespace Gauntlet
 
 			// Launch the test
 			TestInstance = UnrealApp.LaunchSession();
-			if(TestInstance == null)
+			if (TestInstance == null)
 			{
 				return false;
 			}
@@ -820,6 +824,15 @@ namespace Gauntlet
 			MaxRetries = Config.MaxRetries;
 			MarkTestStarted();
 
+			// Caching what log categories heartbeat should track
+			HeartbeatLogCategories.Add("Gauntlet");
+			{
+				// get the categories used to monitor process (this needs rethought).
+				IEnumerable<string> HeartbeatCategories = GetHeartbeatLogCategories().Union(GetCachedConfiguration().LogCategoriesForEvents);
+				HeartbeatLogCategories.AddRange(HeartbeatCategories);
+			}
+			HeartbeatLogCategories = HeartbeatLogCategories.Distinct().ToList();
+
 			return true;
 		}
 
@@ -909,14 +922,14 @@ namespace Gauntlet
 				InArg = InArg.Split("=", 2).First().ToLower();
 				return !ArgsToNotDisplay.Any(str => InArg == str);
 			}
-			
+
 			string WrapParameterInQuotes(string InString)
 			{
 				var positionOfEqual = InString.IndexOf("=");
 				if (positionOfEqual > 0)
 				{
 					var ParamStrName = InString.Substring(0, positionOfEqual);
-					var ParamStrValue = InString.Substring(positionOfEqual+1);
+					var ParamStrValue = InString.Substring(positionOfEqual + 1);
 					int n;
 					if (int.TryParse(ParamStrValue, out n))
 					{
@@ -979,9 +992,12 @@ namespace Gauntlet
 			SessionArtifacts = Enumerable.Empty<UnrealRoleArtifacts>();
 			RoleResults = Enumerable.Empty<UnrealRoleResult>();
 
-			LastServerLogCount = 0;
-			LastClientLogCount = 0;
-			LastEditorLogCount = 0;
+			//Reset last log count for all apps
+			foreach(RoleInstance RoleKey in RoleToLastLogCount.Keys)
+			{
+				RoleToLastLogCount[RoleKey] = 0;
+			}
+
 			LastHeartbeatTime = DateTime.MinValue;
 			LastActiveHeartbeatTime = DateTime.MinValue;
 
@@ -1011,7 +1027,7 @@ namespace Gauntlet
 		{
 			// Retrieve and log heartbeat
 			LogHeartbeat(InInstance);
-			
+
 			// Detect missed heartbeats and fail the test
 			CheckHeartbeat();
 		}
@@ -1120,7 +1136,7 @@ namespace Gauntlet
 			catch (Exception Ex)
 			{
 				CreateReportFailed = true;
-				Message = Ex.Message;				
+				Message = Ex.Message;
 				Log.Warning("Failed to save completion report. {Exception}", Ex);
 			}
 
@@ -1139,8 +1155,8 @@ namespace Gauntlet
 			try
 			{
 				if (Report != null)
-				{ 
-					SubmitToDashboard(Report); 
+				{
+					SubmitToDashboard(Report);
 				}
 
 			}
@@ -1321,8 +1337,9 @@ namespace Gauntlet
 		/// <summary>
 		/// Optional override for the test report name, useful when a single test node has different testing modes
 		/// </summary>
-		protected virtual string HordeReportTestName { 
-			get 
+		protected virtual string HordeReportTestName
+		{
+			get
 			{
 				string ReportName = Name;
 				if (ReportName.Split('.').Length > 1)
@@ -1390,7 +1407,7 @@ namespace Gauntlet
 			{
 				SimpleTestReport SimpleReport = CreateSimpleReportForHorde_legacy(Result);
 				// Set or Attach
-				if(HordeTestReport == null)
+				if (HordeTestReport == null)
 				{
 					HordeTestReport = SimpleReport;
 				}
@@ -1407,7 +1424,7 @@ namespace Gauntlet
 				foreach (UnrealRoleResult RoleResult in RoleResults)
 				{
 					string LogPath = RoleResult.Artifacts.LogPath;
-					if(!string.IsNullOrEmpty(LogPath) && File.Exists(LogPath))
+					if (!string.IsNullOrEmpty(LogPath) && File.Exists(LogPath))
 					{
 						string LogName = FileUtils.ConvertPathToUri(Path.GetRelativePath(Path.GetFullPath(ArtifactPath), Path.GetFullPath(LogPath)));
 						if (HordeTestReport is BaseHordeReport Report)
@@ -1460,7 +1477,7 @@ namespace Gauntlet
 		/// <returns>ITestReport</returns>
 		public virtual ITestReport CreateUnrealEngineTestPassReport(string UnrealAutomatedTestReportPath, string ReportURL)
 		{
-			if(IsSetToRetry()) { return null; }
+			if (IsSetToRetry()) { return null; }
 			string JsonReportPath = Path.Combine(UnrealAutomatedTestReportPath, "index.json");
 			if (File.Exists(JsonReportPath))
 			{
@@ -1544,18 +1561,19 @@ namespace Gauntlet
 		/// <param name="Roles"></param>
 		protected virtual void SetReportMetadata(ITestReport Report, IEnumerable<UnrealTestRole> Roles)
 		{
-			var AllRoleTypes = Roles.Select(R =>  R.Type);
+			var AllRoleTypes = Roles.Select(R => R.Type);
 			var AllRoleContexts = Roles.Select(R => Context.GetRoleContext(R.Type));
 			Report.SetMetadata("Platform", string.Join("+", AllRoleContexts.Select(R => R.Platform.ToString()).Distinct().OrderBy(P => P)));
 			Report.SetMetadata("BuildTarget", string.Join("+", AllRoleTypes.Select(R => R.ToString()).Distinct().OrderBy(B => B)));
 			Report.SetMetadata("Configuration", string.Join("+", AllRoleContexts.Select(R => R.Configuration.ToString()).Distinct().OrderBy(C => C)));
 			Report.SetMetadata("Project", Context.BuildInfo.ProjectName);
 			// Additional metadata passed through the command line arguments
-			foreach( var Meta in Globals.Params.ParseValues("Metadata") )
+			foreach (var Meta in Globals.Params.ParseValues("Metadata"))
 			{
 				var Entry = Meta.Split(":", 2);
-				if (Entry.Count() > 1) {
-					Report.SetMetadata(Entry[0], Entry[1]); 
+				if (Entry.Count() > 1)
+				{
+					Report.SetMetadata(Entry[0], Entry[1]);
 				}
 			}
 		}
@@ -1734,7 +1752,7 @@ namespace Gauntlet
 					ExitReason = "Server exited while clients were running";
 				}
 			}*/
-			
+
 			// The process is gone but we don't know why. This is likely bad and signifies an unhandled or undiagnosed error
 			ExitReason = "app exited with code 0";
 			ExitCode = -1;
@@ -1771,7 +1789,7 @@ namespace Gauntlet
 
 			HashSet<string> MonitoredCategorySet = new HashSet<string>();
 
-			foreach(string Category in GetCachedConfiguration().LogCategoriesForEvents)
+			foreach (string Category in GetCachedConfiguration().LogCategoriesForEvents)
 			{
 				MonitoredCategorySet.Add(Category);
 			}
@@ -1788,7 +1806,7 @@ namespace Gauntlet
 			{
 				bool IsMonitored = MonitoredCategorySet.Contains(Entry.Category);
 
-				if (Entry.Level == UnrealLog.LogLevel.Warning && 
+				if (Entry.Level == UnrealLog.LogLevel.Warning &&
 					(IsMonitored || TrackAllWarnings))
 				{
 					EventList.Add(new UnrealTestEvent(EventSeverity.Warning, Entry.ToString(), Enumerable.Empty<string>()));
@@ -1817,7 +1835,7 @@ namespace Gauntlet
 			string ExitReason;
 
 			UnrealLog LogSummary = CreateLogSummaryFromArtifact(InRoleArtifacts);
-	
+
 			// Give ourselves (and derived classes) a chance to analyze what happened
 			UnrealProcessResult ProcessResult = GetExitCodeAndReason(InReason, LogSummary, InRoleArtifacts, out ExitReason, out ExitCode);
 
@@ -1831,7 +1849,7 @@ namespace Gauntlet
 				ExitCode = -1;
 			}
 
-			return new UnrealRoleResult(ProcessResult, ExitCode, ExitReason, LogSummary, InRoleArtifacts, EventList); 
+			return new UnrealRoleResult(ProcessResult, ExitCode, ExitReason, LogSummary, InRoleArtifacts, EventList);
 		}
 
 		/// <summary>
@@ -1860,7 +1878,7 @@ namespace Gauntlet
 		private void LogHeartbeat(UnrealSessionInstance InInstance)
 		{
 			if (CachedConfig == null ||
-			    GetTestStatus() != TestStatus.InProgress)
+				GetTestStatus() != TestStatus.InProgress)
 			{
 				return;
 			}
@@ -1871,14 +1889,6 @@ namespace Gauntlet
 				return;
 			}
 
-			List<string> LogCategories = new List<string>() { "Gauntlet" };
-			{
-				// get the categories used to monitor process (this needs rethought).
-				IEnumerable<string> HeartbeatCategories = GetHeartbeatLogCategories().Union(GetCachedConfiguration().LogCategoriesForEvents);
-				LogCategories.AddRange(HeartbeatCategories);
-			}
-			LogCategories = LogCategories.Distinct().ToList();
-
 			void LogHeartbeatCategories(ref int LastLogCount, IAppInstance App, string AppPrefix, bool bUpdateHeartbeatTime)
 			{
 				if (App != null && !bDisableHeartbeatLogging)
@@ -1905,7 +1915,7 @@ namespace Gauntlet
 					UnrealLogStreamParser Parser = new UnrealLogStreamParser();
 					LastLogCount += Parser.ReadStream(StdOut, LastLogCount);
 
-					foreach (string TestLine in Parser.GetLogFromShortNameChannels(LogCategories))
+					foreach (string TestLine in Parser.GetLogFromShortNameChannels(HeartbeatLogCategories))
 					{
 						Log.Info("{App}: {Message}", AppPrefix, TestLine);
 
@@ -1925,21 +1935,28 @@ namespace Gauntlet
 				}
 			}
 
-			if (InInstance.ServerApp != null)
-			{
-				bool bUpdateHeartbeat = InInstance.ClientApps == null;
-				LogHeartbeatCategories(ref LastServerLogCount, InInstance.ServerApp, "Server", bUpdateHeartbeat);
-			}
-			if (InInstance.ClientApps.Length > 0)
+			foreach (RoleInstance Role in InInstance.RunningRoles)
 			{
 				bool bUpdateHeartbeat = true;
-				LogHeartbeatCategories(ref LastClientLogCount, InInstance.ClientApps.First(), "Client", bUpdateHeartbeat);
-			}
+				string RoleName = Role.Role.RoleType.ToString();
 
-			if (InInstance.EditorApp != null)
-			{
-				bool bUpdateHeartbeat = false;
-				LogHeartbeatCategories(ref LastEditorLogCount, InInstance.EditorApp, "Editor", bUpdateHeartbeat);
+				if (Role.Role.RoleType.IsServer())
+				{
+					bUpdateHeartbeat = InInstance.ClientApps == null;
+				} 
+				else if (Role.Role.RoleType.IsEditor())
+				{
+					bUpdateHeartbeat = false;
+				} 
+				else if (Role.Role.RoleType.IsClient())
+				{
+					RoleName += " " + Role.AppInstance.Device.Name;
+				}
+
+
+				int RoleLastLogCount = GetLastLogCountFromRole(Role);
+				LogHeartbeatCategories(ref RoleLastLogCount, Role.AppInstance, RoleName, bUpdateHeartbeat);
+				RoleToLastLogCount[Role] = RoleLastLogCount;
 			}
 
 			LastHeartbeatLogTime = DateTime.Now;
@@ -1947,7 +1964,7 @@ namespace Gauntlet
 
 		private void CheckHeartbeat()
 		{
-			if (CachedConfig == null 
+			if (CachedConfig == null
 				|| CachedConfig.DisableHeartbeatTimeout
 				|| CachedConfig.HeartbeatOptions.bExpectHeartbeats == false
 				|| GetTestStatus() != TestStatus.InProgress)
@@ -2001,16 +2018,16 @@ namespace Gauntlet
 		/// <returns></returns>
 		protected virtual string GetRoleResultHash(UnrealRoleResult InResult)
 		{
-			const int MaxCallstackLines = 10;			
+			const int MaxCallstackLines = 10;
 
 			UnrealLog LogSummary = InResult.LogSummary;
 
 			string TotalString = "";
 
 			//Func<int, string> ComputeHash = (string Str) => { return Hasher.ComputeHash(Encoding.UTF8.GetBytes(Str)); };
-			
+
 			if (LogSummary.FatalError != null)
-			{				
+			{
 				TotalString += string.Join("\n", LogSummary.FatalError.Callstack.Take(MaxCallstackLines));
 				TotalString += "\n";
 			}
@@ -2084,7 +2101,7 @@ namespace Gauntlet
 
 			// log command line up here for visibility
 
-			Log.Info( string.Join("\n", new string[] {
+			Log.Info(string.Join("\n", new string[] {
 				string.Format(" * CommandLine: {0}", RoleArtifacts.AppInstance.CommandLine),
 				string.Format(" * Log: {0}", RoleArtifacts.LogPath),
 				string.Format(" * SavedDir: {0}", RoleArtifacts.ArtifactPath),
@@ -2242,7 +2259,7 @@ namespace Gauntlet
 		{
 
 			AddProcessResultEventsFromTestNode();
-			
+
 			MarkdownBuilder MB = new MarkdownBuilder();
 
 			if (TestNodeEvents.Count == 0)
@@ -2256,7 +2273,7 @@ namespace Gauntlet
 			IEnumerable<UnrealTestEvent> Errors = TestNodeEvents.Where(E => E.Severity == EventSeverity.Error);
 			IEnumerable<UnrealTestEvent> Ensures = TestNodeEvents.Where(E => E.IsEnsure);
 			IEnumerable<UnrealTestEvent> Warnings = TestNodeEvents.Where(E => E.Severity == EventSeverity.Warning && !E.IsEnsure);
-			IEnumerable<UnrealTestEvent> Infos = TestNodeEvents.Where(E => E.Severity == EventSeverity.Info );
+			IEnumerable<UnrealTestEvent> Infos = TestNodeEvents.Where(E => E.Severity == EventSeverity.Info);
 
 			foreach (UnrealTestEvent Event in Fatals)
 			{
@@ -2393,7 +2410,7 @@ namespace Gauntlet
 				return Enumerable.Empty<UnrealRoleResult>();
 			}
 
-			return RoleResults.Where(R => R.ProcessResult != UnrealProcessResult.ExitOk);	
+			return RoleResults.Where(R => R.ProcessResult != UnrealProcessResult.ExitOk);
 		}
 
 		/// <summary>
@@ -2452,7 +2469,7 @@ namespace Gauntlet
 
 			// Good/Bad news upfront
 			string Prefix = TestFailed ? "Error: " : "";
-			string WarningStatement = (HasWarnings && !TestFailed)  ? " With Warnings" : "";
+			string WarningStatement = (HasWarnings && !TestFailed) ? " With Warnings" : "";
 			string ResultString = string.Format(" ### {0}{1} {2}{3} ***\n", Prefix, this.Name, GetTestResult(), WarningStatement);
 			Log.Info(ResultString);
 
@@ -2480,7 +2497,7 @@ namespace Gauntlet
 				Errors > 0 ? string.Format(" * Log Errors: {0}", Errors) : null,
 				Warnings > 0 ? string.Format(" * Log Warnings: {0}", Warnings) : null,
 				string.Format(" * Result: {0}", GetTestResult())
-			}.Where(L=>!string.IsNullOrEmpty(L))));
+			}.Where(L => !string.IsNullOrEmpty(L))));
 
 			if (TestFailed && GetRolesThatFailed().Where(R => R.ProcessResult == UnrealProcessResult.Unknown).Any())
 			{
