Adding optional shader model support for RHI
#rnx

diff --git a/Engine/Source/Programs/AutomationTool/Gauntlet/Unreal/Automation/UE.Automation.cs b/Engine/Source/Programs/AutomationTool/Gauntlet/Unreal/Automation/UE.Automation.cs
--- a/Engine/Source/Programs/AutomationTool/Gauntlet/Unreal/Automation/UE.Automation.cs
+++ b/Engine/Source/Programs/AutomationTool/Gauntlet/Unreal/Automation/UE.Automation.cs
@@ -19,7 +19,7 @@ namespace UE
 	public class AutomationTestConfig : UnrealTestConfiguration
 	{
 		/// <summary>
-		/// Run with specify RHI
+		/// Run with specific RHI
 		/// </summary>
 		[AutoParam]
 		public string RHI = "";
@@ -42,6 +42,11 @@ namespace UE
 			{
 				metal
 			}
+			public enum ShaderModel
+			{
+				sm5,
+				sm6
+			}
 		}
 
 		/// <summary>
@@ -271,40 +276,66 @@ namespace UE
 				{
 					AppConfig.CommandLine += " -dpcvars=r.RayTracing=0";
 				}
-
-				// Options specific to windows
-				if (ConfigRole.Platform != null && ((UnrealTargetPlatform)ConfigRole.Platform).IsInGroup(UnrealPlatformGroup.Windows))
+				
+				if (ConfigRole.Platform != null)
 				{
-					if (AttachRenderDoc && !RayTracing)
-					{
-						AppConfig.CommandLine += " -attachRenderDoc";
-					}
+					Type RHIType = null;
+					UnrealTargetPlatform TargetPlatform = (UnrealTargetPlatform)ConfigRole.Platform;
 
-					if (PreferNvidia)
+					// Set our options by platform
+					if (TargetPlatform.IsInGroup(UnrealPlatformGroup.Windows))
 					{
-						AppConfig.CommandLine += " -preferNvidia";
-					}
+						RHIType = typeof(ValidRHI.Win64);
+						if (AttachRenderDoc && !RayTracing)
+						{
+							AppConfig.CommandLine += " -attachRenderDoc";
+						}
 
-					if (!string.IsNullOrEmpty(RHI))
-					{
-						if (Enum.IsDefined(typeof(ValidRHI.Win64), RHI.ToLower()))
+						if (PreferNvidia)
 						{
-							AppConfig.CommandLine += string.Format(" -{0}", RHI);
+							AppConfig.CommandLine += " -preferNvidia";
 						}
-						else
+
+						if (D3DDebug)
 						{
-							throw new AutomationException(string.Format("Unknown RHI target '{0}' for Win64", RHI));
+							AppConfig.CommandLine += " -d3ddebug";
 						}
-					}
 
-					if (D3DDebug)
+						if (StompMalloc)
+						{
+							AppConfig.CommandLine += " -stompmalloc";
+						}
+					}
+					else if (TargetPlatform.IsInGroup(UnrealPlatformGroup.Apple))
+					{
+						RHIType = typeof(ValidRHI.Mac);
+					}
+					else if (TargetPlatform.IsInGroup(UnrealPlatformGroup.Linux))
 					{
-						AppConfig.CommandLine += " -d3ddebug";
+						RHIType = typeof(ValidRHI.Linux);
 					}
 
-					if (StompMalloc)
+					// RHI can specify both RHI and ShaderModel (SM) for Windows, Mac, and Linux platforms
+					if (!string.IsNullOrEmpty(RHI))
 					{
-						AppConfig.CommandLine += " -stompmalloc";
+						if (RHIType == null)
+						{
+							throw new AutomationException(string.Format("Unknown target platform '{0}'", TargetPlatform));
+						}
+
+						// The RHI can include the Shader Model version to go along with it and should be split out e.g d3d12-sm6
+						string[] RHIArguments = RHI.Split('-');
+						foreach (string Argument in RHIArguments)
+						{
+							if (Enum.IsDefined(RHIType, Argument.ToLower()) || Enum.IsDefined(typeof(ValidRHI.ShaderModel), Argument.ToLower()))
+							{
+								AppConfig.CommandLine += string.Format(" -{0}", Argument);
+							}
+							else
+							{
+								throw new AutomationException(string.Format("Unknown RHI target or Shader Model '{0}' for {1}", Argument, RHIType.Name));
+							}
+						}
 					}
 				}
 			}
@@ -972,7 +1003,7 @@ namespace UE
 							if (Report != null)
 							{
 								var MainRolePlatform = Context.GetRoleContext(Config.GetMainRequiredRole().Type).Platform;
-								Report.SetMetadata("RHI", string.IsNullOrEmpty(Config.RHI) || !MainRolePlatform.IsInGroup(UnrealPlatformGroup.Windows) ? "default" : Config.RHI.ToLower());
+								Report.SetMetadata("RHI", string.IsNullOrEmpty(Config.RHI) || !MainRolePlatform.IsInGroup(UnrealPlatformGroup.Desktop) ? "default" : Config.RHI.ToLower());
 							}
 						}
 					}
