Add Horde custom installer wix extension to dot net check exceptions
#rnx

diff --git a/Engine/Source/Programs/AutomationTool/Scripts/CheckCsprojDotNetVersion.cs b/Engine/Source/Programs/AutomationTool/Scripts/CheckCsprojDotNetVersion.cs
--- a/Engine/Source/Programs/AutomationTool/Scripts/CheckCsprojDotNetVersion.cs
+++ b/Engine/Source/Programs/AutomationTool/Scripts/CheckCsprojDotNetVersion.cs
@@ -41,7 +41,8 @@ class CheckCsprojDotNetVersion : BuildCommand
 				CsProj.ContainsName("UnrealVS", EngineDir) ||
 				CsProj.ContainsName("DatasmithRevitExporter", EngineDir) ||
 				CsProj.ContainsName("DatasmithNavisworksExporter", EngineDir) ||
-				CsProj.ContainsName("CSVTools", EngineDir))
+				CsProj.ContainsName("CSVTools", EngineDir) ||
+				CsProj.ContainsName("InstallerCustomActions", EngineDir))
 			{
 				continue;
 			}
