Fix warning on Horde installer custom action
#rnx

diff --git a/Engine/Source/Programs/AutomationTool/Scripts/CheckCsprojDotNetVersion.cs b/Engine/Source/Programs/AutomationTool/Scripts/CheckCsprojDotNetVersion.cs
--- a/Engine/Source/Programs/AutomationTool/Scripts/CheckCsprojDotNetVersion.cs
+++ b/Engine/Source/Programs/AutomationTool/Scripts/CheckCsprojDotNetVersion.cs
@@ -42,7 +42,7 @@ class CheckCsprojDotNetVersion : BuildCommand
 				CsProj.ContainsName("DatasmithRevitExporter", EngineDir) ||
 				CsProj.ContainsName("DatasmithNavisworksExporter", EngineDir) ||
 				CsProj.ContainsName("CSVTools", EngineDir) ||
-				CsProj.ContainsName("InstallerCustomActions", EngineDir))
+				(CsProj.ContainsName("CustomActions", EngineDir) && CsProj.GetFileName().Equals("InstallerCustomActions.csproj")))
 			{
 				continue;
 			}
