[Gauntlet] Avoid logging twice the log line in verbse

#rnx

diff --git a/Engine/Source/Programs/AutomationTool/Gauntlet/Unreal/Automation/UE.CookOnTheFly.cs b/Engine/Source/Programs/AutomationTool/Gauntlet/Unreal/Automation/UE.CookOnTheFly.cs
--- a/Engine/Source/Programs/AutomationTool/Gauntlet/Unreal/Automation/UE.CookOnTheFly.cs
+++ b/Engine/Source/Programs/AutomationTool/Gauntlet/Unreal/Automation/UE.CookOnTheFly.cs
@@ -225,16 +225,6 @@ namespace UE
 			}
 
 			EditorLogParser.ReadStream();
-			if (Globals.Params.ParseParam("Verbose"))
-			{
-				IEnumerable<string> EditorCookEntries = EditorLogParser.GetLogFromShortNameChannels(LogCategories);
-
-				if (EditorCookEntries.Any())
-				{
-					EditorCookEntries.ToList().ForEach(S => Log.Info("[Editor] " + S));
-				}
-			}
-
 			if (!ServerStarted)
 			{
 				string CompletionString = GetStartedCookServerString();
@@ -340,15 +330,6 @@ namespace UE
 				}
 
 				ClientLogParser.ReadStream();
-				if (Globals.Params.ParseParam("Verbose"))
-				{
-					IEnumerable<string> ClientCookEntries = ClientLogParser.GetLogFromShortNameChannels(LogCategories);
-					if (ClientCookEntries.Any())
-					{
-						ClientCookEntries.ToList().ForEach(S => Log.Info("[Client] " + S));
-					}
-				}
-
 				if (ClientCookOnTheFlyCheckers.CheckAll() && CookingProcessStarted)
 				{
 					Log.Info("Found '{0}', '{1}', '{2}'. The CookOnTheFly log channel is active. The cooking process is taking place.", GetGameStartString(), GetCookingProcessString(), GetReceivedPackagesCookedString());
