BuildGraph: Add ulimit setting for DockerBuild task

Default to 100k open file descriptors as default

diff --git a/Engine/Source/Programs/AutomationTool/BuildGraph/Tasks/DockerBuildTask.cs b/Engine/Source/Programs/AutomationTool/BuildGraph/Tasks/DockerBuildTask.cs
--- a/Engine/Source/Programs/AutomationTool/BuildGraph/Tasks/DockerBuildTask.cs
+++ b/Engine/Source/Programs/AutomationTool/BuildGraph/Tasks/DockerBuildTask.cs
@@ -47,6 +47,12 @@ namespace AutomationTool.Tasks
 		/// </summary>
 		[TaskParameter(Optional = true)]
 		public bool UseBuildKit { get; set; }
+		
+		/// <summary>
+		/// Set ulimit in build
+		/// </summary>
+		[TaskParameter(Optional = true)]
+		public string Ulimit { get; set; } = "nofile=100000:100000";
 
 		/// <summary>
 		/// Type of progress output (--progress)
@@ -181,6 +187,10 @@ namespace AutomationTool.Tasks
 				{
 					arguments.Append($" --progress={_parameters.ProgressOutput}");
 				}
+				if (_parameters.Ulimit != null && _parameters.Ulimit.Length > 1)
+				{
+					arguments.Append($" --ulimit={_parameters.Ulimit}");
+				}
 				if (_parameters.Arguments != null)
 				{
 					arguments.Append($" {_parameters.Arguments}");
