[Gauntlet] fix partial log file when running client locally

#rnx

diff --git a/Engine/Source/Programs/AutomationTool/Gauntlet/Framework/Base/Gauntlet.TargetDeviceDesktopCommon.cs b/Engine/Source/Programs/AutomationTool/Gauntlet/Framework/Base/Gauntlet.TargetDeviceDesktopCommon.cs
--- a/Engine/Source/Programs/AutomationTool/Gauntlet/Framework/Base/Gauntlet.TargetDeviceDesktopCommon.cs
+++ b/Engine/Source/Programs/AutomationTool/Gauntlet/Framework/Base/Gauntlet.TargetDeviceDesktopCommon.cs
@@ -401,7 +401,7 @@ namespace Gauntlet
 				// -abslog makes sure Unreal dynamically update the log window when using -log
 				if (!CommandUtils.IsBuildMachine && (!string.IsNullOrEmpty(LogFile) || AppConfig.CommandLine.Contains("-log")))
 				{
-					string LogFolder = string.IsNullOrEmpty(LogFile) ? Path.Combine(ArtifactPath, "Logs") : Path.GetDirectoryName(LogFile);
+					string LogFolder = string.IsNullOrEmpty(LogFile) ? Path.Combine(Device.LocalCachePath, "Logs") : Path.GetDirectoryName(LogFile);
 
 					if (!Directory.Exists(LogFolder))
 					{
@@ -515,6 +515,7 @@ namespace Gauntlet
 			}
 			else
 			{
+				Log.Warning("Log file '{filepath}' is missing at the time of making a copy. The buffer will be used instead but will most likely lack the beginning of the file.", ProcessLogFile);
 				StreamWriter Writer = ProcessUtils.CreateWriterForProcessLog(FilePath, CommandLine);
 				foreach(string Line in ProcessLogOutput)
 				{
