[LLT] reuse cached stdout and utilize IAppInstance Log buffer interface

#rnx

diff --git a/Engine/Source/Programs/AutomationTool/LowLevelTests/Tests/Gauntlet.LowLevelTestNode.cs b/Engine/Source/Programs/AutomationTool/LowLevelTests/Tests/Gauntlet.LowLevelTestNode.cs
--- a/Engine/Source/Programs/AutomationTool/LowLevelTests/Tests/Gauntlet.LowLevelTestNode.cs
+++ b/Engine/Source/Programs/AutomationTool/LowLevelTests/Tests/Gauntlet.LowLevelTestNode.cs
@@ -26,7 +26,7 @@ namespace LowLevelTests
 
 		public LowLevelTestsSession LowLevelTestsApp { get; private set; }
 
-		private int LastStdoutSeekPos = 0;
+		private ILogStreamReader LogReader = null;
 		private string[] CurrentProcessedLines;
 
 		public LowLevelTests(LowLevelTestContext InContext)
@@ -221,12 +221,12 @@ namespace LowLevelTests
 				// No reports from Android tests yet. Since adb shell doesn't forward exit code, we look for it in the log output.
 				if (Context.Options.Platform == UnrealTargetPlatform.Android)
 				{
-					Match AndroidExitCodeLog = Regex.Match(TestInstance.StdOut, @"Tests finished with exit code (\d+)");
+					Match AndroidExitCodeLog = Regex.Match(StdOut, @"Tests finished with exit code (\d+)");
 					if (AndroidExitCodeLog.Success)
 					{
 						ExitCodeOverride = int.Parse(Regex.Match(AndroidExitCodeLog.Value, @"\d+").Value);
 					}
-					else if (TestInstance.StdOut.Contains("beginning of crash"))
+					else if (StdOut.Contains("beginning of crash"))
 					{
 						Log.Info("Crash occurred during test.");
 						ExitCodeOverride = -1;
@@ -356,16 +356,12 @@ namespace LowLevelTests
 
 		private void ParseLowLevelTestsLog()
 		{
-			// Parse new lines from Stdout, if any
-			if (LastStdoutSeekPos < TestInstance.StdOut.Length)
+			if (LogReader == null)
 			{
-				CurrentProcessedLines = TestInstance.StdOut
-					.Substring(LastStdoutSeekPos)
-					.Split("\n")
-					.Where(Line => !string.IsNullOrWhiteSpace(Line))
-					.ToArray();
-				LastStdoutSeekPos = TestInstance.StdOut.Length - 1;
+				LogReader = TestInstance?.GetLogBufferReader();
 			}
+			// Parse new lines from Stdout, if any
+			CurrentProcessedLines = LogReader?.EnumerateNextLines().Where(Line => !string.IsNullOrWhiteSpace(Line)).ToArray();
 		}
 
 		private void PrintLogIfCaptureOutput()
@@ -374,7 +370,7 @@ namespace LowLevelTests
 			{
 				foreach (string OutputLine in CurrentProcessedLines)
 				{
-					Console.WriteLine(OutputLine.TrimEnd(Environment.NewLine.ToCharArray()));
+					Console.WriteLine(OutputLine);
 				}
 			}
 		}
