LaunchMultiServer UAT command improvements
Supports -cook option to launch cooked builds
Added -directclients option that will launch one client per server and connect directly to it.
Supports -clientcmdline option to pass additional command-line options to the clients.

diff --git a/Engine/Source/Programs/AutomationTool/Scripts/LaunchMultiServer.Automation.cs b/Engine/Source/Programs/AutomationTool/Scripts/LaunchMultiServer.Automation.cs
--- a/Engine/Source/Programs/AutomationTool/Scripts/LaunchMultiServer.Automation.cs
+++ b/Engine/Source/Programs/AutomationTool/Scripts/LaunchMultiServer.Automation.cs
@@ -4,6 +4,7 @@ using System;
 using System.IO;
 using UnrealBuildBase;
 using UnrealBuildTool;
+using AutomationScripts;
 using EpicGames.Core;
 using Microsoft.Extensions.Logging;
 using System.Collections.Generic;
@@ -23,18 +24,34 @@ namespace AutomationTool
 		}
 	}
 
+	// MultiServerDefinition plus additional information only used by this command
+	class MultiServerInstanceInfo
+	{
+		public MultiServerDefinition Definition = null;
+		public int GameListenPort = 0;
+
+		public MultiServerInstanceInfo(MultiServerDefinition definition)
+		{
+			Definition = definition;
+		}
+	}
+
 	[Help("Launches multiple server processes for a project using the MultiServerReplication plugin.")]
 	[Help("project=<project>", "Project to open. Will search current path and paths in ueprojectdirs.")]
 	[Help("map=<MapName>", "Map to load on startup.")]
 	public class LaunchMultiServer : BuildCommand, IProjectParamsHelpers
 	{
+		private enum ProcessLaunchType
+		{
+			Server,
+			DirectClient
+		}
+
 		public override ExitCode Execute()
 		{
 			Logger.LogInformation("********** RUN MULTISERVER COMMAND STARTED **********");
 			var StartTime = DateTime.UtcNow;
 
-			var MapToRun = ParseParamValue("map", "");
-
 			// Parse server configuration from ini files
 			ConfigHierarchy ProjectGameConfig = ConfigCache.ReadHierarchy(ConfigHierarchyType.Game, DirectoryReference.FromFile(ProjectPath), UnrealTargetPlatform.Win64);
 
@@ -45,15 +62,18 @@ namespace AutomationTool
 			if (ConfigServerDefs != null && ConfigServerDefs.Length > 0)
 			{
 				// Use a dictionary to store ServerDefs by unqiue ID
-				Dictionary<string, MultiServerDefinition> UniqueServerDefs = new Dictionary<string, MultiServerDefinition>();
+				Dictionary<string, MultiServerInstanceInfo> UniqueServerInfos = new Dictionary<string, MultiServerInstanceInfo>();
+
+				var GameListenPort = 7777;
 
 				// Check for duplicates and build peer list
 				string PeerList = "";
 				for (int i = 0; i < ConfigServerDefs.Length; i++)
 				{
-					if (UniqueServerDefs.TryAdd(ConfigServerDefs[i].LocalId, ConfigServerDefs[i]))
+					if (UniqueServerInfos.TryAdd(ConfigServerDefs[i].LocalId, new MultiServerInstanceInfo(ConfigServerDefs[i])))
 					{
 						PeerList += String.Format("127.0.0.1:{0}{1}", ConfigServerDefs[i].ListenPort, i < (ConfigServerDefs.Length - 1) ? "," : "");
+						UniqueServerInfos[ConfigServerDefs[i].LocalId].GameListenPort = GameListenPort++;
 					}
 					else
 					{
@@ -61,35 +81,14 @@ namespace AutomationTool
 					}
 				}
 
-				var CommonArgs = ProjectPath.FullName + " " + MapToRun + " -server -log -newconsole -MultiServerPeers=" + PeerList;
+				StartProcessesForServerInfos(UniqueServerInfos, ProcessLaunchType.Server, " -MultiServerPeers=" + PeerList);
 
-				var ServerApp = CombinePaths(CmdEnv.LocalRoot, "Engine/Binaries/Win64/UnrealEditor.exe");
-				PushDir(Path.GetDirectoryName(ServerApp));
-
-				try
+				// If -directclients is specified, launch one client per server to connect directly.
+				if (ParseParam("directclients"))
 				{
-					foreach (var ServerDef in UniqueServerDefs.Values)
-					{
-						Logger.LogInformation("Starting MultiServer instance with LocalId {0} and port {1}", ServerDef.LocalId, ServerDef.ListenPort);
-
-						var Args = CommonArgs + String.Format(" -MultiServerLocalId={0} -MultiServerListenPort={1}", ServerDef.LocalId, ServerDef.ListenPort);
-						var ServerProcess = Run(ServerApp, Args, null, ERunOptions.Default | ERunOptions.NoWaitForExit | ERunOptions.NoStdOutRedirect);
+					Logger.LogInformation("Starting direct client instances connecting to MultiServer instances");
 
-						if (ServerProcess != null)
-						{
-							// Remove started process so it won't be killed on UAT exit.
-							// Essentially forces the -NoKill command-line option behavior for these.
-							ProcessManager.RemoveProcess(ServerProcess);
-						}
-					}
-				}
-				catch
-				{
-					throw;
-				}
-				finally
-				{
-					PopDir();
+					StartProcessesForServerInfos(UniqueServerInfos, ProcessLaunchType.DirectClient, "");
 				}
 			}
 			else
@@ -103,6 +102,78 @@ namespace AutomationTool
 			return ExitCode.Success;
 		}
 
+		private void StartProcessesForServerInfos(Dictionary<string, MultiServerInstanceInfo> ServerInfos, ProcessLaunchType LaunchType, string AdditionalParameters)
+		{
+			var Params = new ProjectParams
+				(
+					Command: this,
+					RawProjectPath: ProjectPath,
+					DedicatedServer: LaunchType == ProcessLaunchType.Server,
+					Client: LaunchType == ProcessLaunchType.DirectClient
+				);
+
+			var DeployContexts = Project.CreateDeploymentContext(Params, LaunchType == ProcessLaunchType.Server);
+
+			if (DeployContexts.Count == 0)
+			{
+				throw new AutomationException("No DeployContexts for launching a process.");
+			}
+
+			var DeployContext = DeployContexts[0];
+
+			var App = CombinePaths(CmdEnv.LocalRoot, "Engine/Binaries/Win64/UnrealEditor.exe");
+			if (Params.Cook)
+			{
+				List<FileReference> Exes = DeployContext.StageTargetPlatform.GetExecutableNames(DeployContext);
+				App = Exes[0].FullName;
+			}
+
+			var CommonArgs = (Params.Cook ? "" : DeployContext.ProjectArgForCommandLines + " ");
+			if (LaunchType == ProcessLaunchType.Server)
+			{
+				CommonArgs += Params.MapToRun + " -server -log -newconsole";
+			}
+
+			PushDir(Path.GetDirectoryName(App));
+
+			try
+			{
+				foreach (var ServerInfo in ServerInfos.Values)
+				{
+					Logger.LogInformation("Starting MultiServer instance of type {0} for LocalId {1}, beacon port {2}, and game port {3}", LaunchType, ServerInfo.Definition.LocalId, ServerInfo.Definition.ListenPort, ServerInfo.GameListenPort);
+
+					var Args = CommonArgs;
+
+					if (LaunchType == ProcessLaunchType.Server)
+					{
+						Args += String.Format(" -MultiServerLocalId={0} -MultiServerListenPort={1} -port={2} ", ServerInfo.Definition.LocalId, ServerInfo.Definition.ListenPort, ServerInfo.GameListenPort);
+					}
+					else if (LaunchType == ProcessLaunchType.DirectClient)
+					{
+						Args += String.Format(" 127.0.0.1:{0} ", ServerInfo.GameListenPort) + (Params.Cook ? "" : "-game ") + Params.ClientCommandline;
+					}
+					Args += AdditionalParameters;
+
+					var NewProcess = Run(App, Args, null, ERunOptions.Default | ERunOptions.NoWaitForExit | ERunOptions.NoStdOutRedirect);
+
+					if (NewProcess != null)
+					{
+						// Remove started process so it won't be killed on UAT exit.
+						// Essentially forces the -NoKill command-line option behavior for these.
+						ProcessManager.RemoveProcess(NewProcess);
+					}
+				}
+			}
+			catch
+			{
+				throw;
+			}
+			finally
+			{
+				PopDir();
+			}
+		}
+
 		private FileReference ProjectFullPath;
 		public virtual FileReference ProjectPath
 		{
