UE-223759: add retries and fix warnings for socket exceptions on Android platform

diff --git a/Engine/Source/Programs/AutomationTool/Android/AndroidFileClient.cs b/Engine/Source/Programs/AutomationTool/Android/AndroidFileClient.cs
--- a/Engine/Source/Programs/AutomationTool/Android/AndroidFileClient.cs
+++ b/Engine/Source/Programs/AutomationTool/Android/AndroidFileClient.cs
@@ -1990,6 +1990,11 @@ namespace AutomationTool
 			Stats_Stopwatch = null;
 			if (ClientSocket != null)
 			{
+				if (!ClientSocket.Connected)
+				{
+					return;
+				}
+
 				try
 				{
 					int bytesSent = SocketSend(false, CommandPacket(Command_Close, 0));
@@ -2176,15 +2181,45 @@ namespace AutomationTool
 			return Result;
 		}
 
+		private const int RetriesMax = 5;
+		private const int RetryMillisecondsDelay = 5000;
+
 		public bool FileWrite(string SourcePath, string DestPath, int bLog = 0)
 		{
-			Boolean Result = false;
-
 			if (ClientSocket == null)
 			{
-				return Result;
+				return false;
 			}
 
+			// with no network errors, return the result of filesystem IO process
+			// with network error, try reconnect
+			// with reconnect error, retry reconnect
+			// with reconnect success, try filesystem IO process
+			bool HasNetError = false;
+			for (int TryNum = 0; TryNum < RetriesMax; ++TryNum)
+			{
+				if (HasNetError)
+				{
+					Logger.LogInformation($"Retry {TryNum + 1}/{RetriesMax} writing file {DestPath} because of network error");
+					Thread.Sleep(RetryMillisecondsDelay);
+					HasNetError = !OpenConnection();
+					continue;
+				}
+
+				(bool Result, HasNetError) = FileWriteInternal(SourcePath, DestPath, bLog);
+				if (!HasNetError)
+				{
+					return Result;
+				}
+			}
+
+			return false;
+		}
+
+		private (bool Result, bool HasNetworkError) FileWriteInternal(string SourcePath, string DestPath, int bLog = 0)
+		{
+			bool Result = false;
+			bool HasNetError = false;
 			try
 			{
 				using (FileStream fileStream = new FileStream(SourcePath, FileMode.Open, FileAccess.Read, FileShare.None, 4096, FileOptions.SequentialScan))
@@ -2269,14 +2304,19 @@ namespace AutomationTool
 			catch (IOException)
 			{
 				// file not found
-				return false;
+			}
+			catch (SocketException e)
+			{
+				HasNetError = true;
+				Logger.LogWarning("{bLog}> Network Exception: {Arg1}", bLog, e.ToString());
+				CloseConnection();
 			}
 			catch (Exception e)
 			{
 				Logger.LogError("{bLog}> Unexpected Exception: {Arg1}", bLog, e.ToString());
 				CloseConnection();
 			}
-			return Result;
+			return (Result, HasNetError);
 		}
 
 		public bool FileWriteCompressed(string SourcePath, string DestPath, int bLog = 0)
