[Gauntlet] Fix false positive when Editor canceled initialization

#rnx

diff --git a/Engine/Source/Programs/AutomationTool/Gauntlet/Framework/Base/Gauntlet.TargetDeviceDesktopCommon.cs b/Engine/Source/Programs/AutomationTool/Gauntlet/Framework/Base/Gauntlet.TargetDeviceDesktopCommon.cs
--- a/Engine/Source/Programs/AutomationTool/Gauntlet/Framework/Base/Gauntlet.TargetDeviceDesktopCommon.cs
+++ b/Engine/Source/Programs/AutomationTool/Gauntlet/Framework/Base/Gauntlet.TargetDeviceDesktopCommon.cs
@@ -508,7 +508,7 @@ namespace Gauntlet
 				return ProcessResult.WriteOutputToFile(FilePath) != null;
 			}
 
-			if (File.Exists(ProcessLogFile) && HasExited)
+			if (File.Exists(ProcessLogFile))
 			{
 				ProcessUtils.CheckProcessLogReachedSizeLimit(new FileReference(ProcessLogFile));
 				File.Copy(ProcessLogFile, FilePath, true);
diff --git a/Engine/Source/Programs/AutomationTool/Gauntlet/Unreal/Utils/Gauntlet.UnrealLogParser.cs b/Engine/Source/Programs/AutomationTool/Gauntlet/Unreal/Utils/Gauntlet.UnrealLogParser.cs
--- a/Engine/Source/Programs/AutomationTool/Gauntlet/Unreal/Utils/Gauntlet.UnrealLogParser.cs
+++ b/Engine/Source/Programs/AutomationTool/Gauntlet/Unreal/Utils/Gauntlet.UnrealLogParser.cs
@@ -189,8 +189,7 @@ namespace Gauntlet
 		private bool HasEngineInitialized()
 		{
 			// Search for Engine initialized pattern.
-			// Since long log can have the early part truncated, look for test exit code or requested exit and assumed the engine was initialized to get this far.
-			return _parser.GetAllMatches(EngineInitializedPattern).Any() || HasTestExitCode || RequestedExit;
+			return _parser.GetAllMatches(EngineInitializedPattern).Any();
 		}
 
 		/// <summary>
