[Gauntlet] enable sanitize log output when targetting local virtual device

#rnx

diff --git a/Engine/Source/Programs/AutomationTool/Gauntlet/Unreal/BuildSource/Gauntlet.UnrealBuildSource.cs b/Engine/Source/Programs/AutomationTool/Gauntlet/Unreal/BuildSource/Gauntlet.UnrealBuildSource.cs
--- a/Engine/Source/Programs/AutomationTool/Gauntlet/Unreal/BuildSource/Gauntlet.UnrealBuildSource.cs
+++ b/Engine/Source/Programs/AutomationTool/Gauntlet/Unreal/BuildSource/Gauntlet.UnrealBuildSource.cs
@@ -600,8 +600,8 @@ namespace Gauntlet
 
 			// Detect json log line output
 			if (Config.CommandLineParams.HasParam("JsonStdOut")
-				|| (Role.Platform == BuildHostPlatform.Current.Platform
-					&& !string.IsNullOrEmpty(Environment.GetEnvironmentVariable("UE_LOG_JSON_TO_STDOUT"))))
+				|| (!string.IsNullOrEmpty(Environment.GetEnvironmentVariable("UE_LOG_JSON_TO_STDOUT"))
+					&& (Role.Platform == BuildHostPlatform.Current.Platform || CanRunPlatformVirtualized(Role.Platform))))
 			{
 				Config.FilterLoggingDelegate = (string M, bool IsErr) => UnrealLogParser.SanitizeJsonOutputLine(M);
 			}
@@ -619,6 +619,12 @@ namespace Gauntlet
 			return Config;
 		}
 
+		private bool CanRunPlatformVirtualized(UnrealTargetPlatform? Platform)
+		{
+			return Utils.InterfaceHelpers.FindImplementations<IVirtualLocalDevice>()
+				.Any(F => F.GetPlatform() == Platform && F.CanRunVirtualFromPlatform(BuildHostPlatform.Current.Platform));
+		}
+
 		/// <summary>
 		/// Remove all duplicate flags and combine any execcmd strings that might be floating around in the commandline.
 		/// </summary>
