Fixing bug introduced by supporting RHI/SM for various platforms. Fix only allows for RHI/SM for Desktop platforms

#rnx

diff --git a/Engine/Source/Programs/AutomationTool/Gauntlet/Unreal/Automation/UE.Automation.cs b/Engine/Source/Programs/AutomationTool/Gauntlet/Unreal/Automation/UE.Automation.cs
--- a/Engine/Source/Programs/AutomationTool/Gauntlet/Unreal/Automation/UE.Automation.cs
+++ b/Engine/Source/Programs/AutomationTool/Gauntlet/Unreal/Automation/UE.Automation.cs
@@ -277,7 +277,8 @@ namespace UE
 					AppConfig.CommandLine += " -dpcvars=r.RayTracing=0";
 				}
 				
-				if (ConfigRole.Platform != null)
+				// Options specific to Desktop platforms
+				if (ConfigRole.Platform != null && ((UnrealTargetPlatform)ConfigRole.Platform).IsInGroup(UnrealPlatformGroup.Desktop))
 				{
 					Type RHIType = null;
 					UnrealTargetPlatform TargetPlatform = (UnrealTargetPlatform)ConfigRole.Platform;
@@ -320,7 +321,7 @@ namespace UE
 					{
 						if (RHIType == null)
 						{
-							throw new AutomationException(string.Format("Unknown target platform '{0}'", TargetPlatform));
+							throw new AutomationException(string.Format("Unknown target platform '{0}' for RHI '{1}'", TargetPlatform, RHI));
 						}
 
 						// The RHI can include the Shader Model version to go along with it and should be split out e.g d3d12-sm6
