Pass the timeout setting down.

diff --git a/Engine/Source/Programs/AutomationTool/Gauntlet/Framework/RpcFramework/Gauntlet.RpcExecutor.cs b/Engine/Source/Programs/AutomationTool/Gauntlet/Framework/RpcFramework/Gauntlet.RpcExecutor.cs
--- a/Engine/Source/Programs/AutomationTool/Gauntlet/Framework/RpcFramework/Gauntlet.RpcExecutor.cs
+++ b/Engine/Source/Programs/AutomationTool/Gauntlet/Framework/RpcFramework/Gauntlet.RpcExecutor.cs
@@ -744,6 +744,7 @@ namespace Gauntlet
 			// We might timeout calling CallRpc on a client that's closing - we do not want this to be a fail;
 			try
 			{
+				ExecutorClient.Timeout = new TimeSpan(0, 0, 0, 0, Timeout);
 				HttpResponseMessage RpcResponse = ExecutorClient.Send(RpcRequest);
 				return RpcResponse;
 			}
@@ -799,7 +800,8 @@ namespace Gauntlet
         	// We might timeout calling CallRpc on a client that's closing - we do not want this to be a fail;
         	try
         	{
-        		HttpResponseMessage RpcResponse = ExecutorClient.Send(RpcRequest);
+				ExecutorClient.Timeout = new TimeSpan(0, 0, 0, 0, Timeout);
+				HttpResponseMessage RpcResponse = ExecutorClient.Send(RpcRequest);
         		return RpcResponse;
         	}
         	catch (Exception e)
