Undo CL: 36312059

Pass the timeout setting down.

diff --git a/Engine/Source/Programs/AutomationTool/Gauntlet/Framework/RpcFramework/Gauntlet.RpcExecutor.cs b/Engine/Source/Programs/AutomationTool/Gauntlet/Framework/RpcFramework/Gauntlet.RpcExecutor.cs
--- a/Engine/Source/Programs/AutomationTool/Gauntlet/Framework/RpcFramework/Gauntlet.RpcExecutor.cs
+++ b/Engine/Source/Programs/AutomationTool/Gauntlet/Framework/RpcFramework/Gauntlet.RpcExecutor.cs
@@ -744,7 +744,6 @@ namespace Gauntlet
 			// We might timeout calling CallRpc on a client that's closing - we do not want this to be a fail;
 			try
 			{
-				ExecutorClient.Timeout = new TimeSpan(0, 0, 0, 0, Timeout);
 				HttpResponseMessage RpcResponse = ExecutorClient.Send(RpcRequest);
 				return RpcResponse;
 			}
@@ -800,8 +799,7 @@ namespace Gauntlet
         	// We might timeout calling CallRpc on a client that's closing - we do not want this to be a fail;
         	try
         	{
-				ExecutorClient.Timeout = new TimeSpan(0, 0, 0, 0, Timeout);
-				HttpResponseMessage RpcResponse = ExecutorClient.Send(RpcRequest);
+        		HttpResponseMessage RpcResponse = ExecutorClient.Send(RpcRequest);
         		return RpcResponse;
         	}
         	catch (Exception e)
