Improvements to LaunchMultiServer UAT script:
Supports -servercmdline option to pass additional arguments to the servers
Supports -notimeouts option to pass -notimeouts to all servers and clients
By default spawns clients in 1280x720 windows

diff --git a/Engine/Source/Programs/AutomationTool/Scripts/LaunchMultiServer.Automation.cs b/Engine/Source/Programs/AutomationTool/Scripts/LaunchMultiServer.Automation.cs
--- a/Engine/Source/Programs/AutomationTool/Scripts/LaunchMultiServer.Automation.cs
+++ b/Engine/Source/Programs/AutomationTool/Scripts/LaunchMultiServer.Automation.cs
@@ -131,7 +131,7 @@ namespace AutomationTool
 			var CommonArgs = (Params.Cook ? "" : DeployContext.ProjectArgForCommandLines + " ");
 			if (LaunchType == ProcessLaunchType.Server)
 			{
-				CommonArgs += Params.MapToRun + " -server -log -newconsole";
+				CommonArgs += Params.MapToRun + " -server -log -newconsole ";
 			}
 
 			PushDir(Path.GetDirectoryName(App));
@@ -146,11 +146,23 @@ namespace AutomationTool
 
 					if (LaunchType == ProcessLaunchType.Server)
 					{
-						Args += String.Format(" -MultiServerLocalId={0} -MultiServerListenPort={1} -port={2} ", ServerInfo.Definition.LocalId, ServerInfo.Definition.ListenPort, ServerInfo.GameListenPort);
+						Args += String.Format("-MultiServerLocalId={0} -MultiServerListenPort={1} -port={2} ", ServerInfo.Definition.LocalId, ServerInfo.Definition.ListenPort, ServerInfo.GameListenPort);
+						Args += Params.ServerCommandline + " ";
 					}
 					else if (LaunchType == ProcessLaunchType.DirectClient)
 					{
-						Args += String.Format(" 127.0.0.1:{0} ", ServerInfo.GameListenPort) + (Params.Cook ? "" : "-game ") + Params.ClientCommandline;
+						Args += String.Format("127.0.0.1:{0} ", ServerInfo.GameListenPort) + (Params.Cook ? "" : "-game ");
+						Args += Params.ClientCommandline + " ";
+						if (!Params.ClientCommandline.Contains("-res"))
+						{
+							// Small window by default
+							Args += "-res=1280x720 ";
+						}
+					}
+
+					if (ParseParam("notimeouts"))
+					{
+						Args += "-notimeouts ";
 					}
 					Args += AdditionalParameters;
 
