Update UAT to have the same default value for n.VerifyPeer as UE. Fixes cacert.pem not being staged when n.VerifyPeer is not explicitly set to true in Engine.ini

diff --git a/Engine/Source/Programs/AutomationTool/Scripts/CopyBuildToStagingDirectory.Automation.cs b/Engine/Source/Programs/AutomationTool/Scripts/CopyBuildToStagingDirectory.Automation.cs
--- a/Engine/Source/Programs/AutomationTool/Scripts/CopyBuildToStagingDirectory.Automation.cs
+++ b/Engine/Source/Programs/AutomationTool/Scripts/CopyBuildToStagingDirectory.Automation.cs
@@ -1669,7 +1669,7 @@ namespace AutomationScripts
 						SC.StageCrashReporterFiles(StagedFileType.UFS, DirectoryReference.Combine(SC.EngineRoot, "Programs", "CrashReportClient", "Config"), StageFilesSearch.AllDirectories);
 					}
 					// check if the game will be verifying ssl connections - if not, we can skip staging files that won't be needed
-					bool bStageSSLCertificates = false;
+					bool bStageSSLCertificates = true;
 					ConfigHierarchy PlatformEngineConfig = ConfigCache.ReadHierarchy(ConfigHierarchyType.Engine, DirectoryReference.FromFile(Params.RawProjectPath), SC.StageTargetPlatform.IniPlatformType, SC.CustomConfig);
 					if (PlatformEngineConfig != null)
 					{
