Temporally add verbose output to Symstore to understand symbol publishing issue.

diff --git a/Engine/Source/Programs/AutomationTool/Win/WinPlatform.Automation.cs b/Engine/Source/Programs/AutomationTool/Win/WinPlatform.Automation.cs
--- a/Engine/Source/Programs/AutomationTool/Win/WinPlatform.Automation.cs
+++ b/Engine/Source/Programs/AutomationTool/Win/WinPlatform.Automation.cs
@@ -591,7 +591,7 @@ public class Win64Platform : Platform
 				// Copy everything to the temp symstore
 				ProcessStartInfo StartInfo = new ProcessStartInfo();
 				StartInfo.FileName = SymStoreExe.FullName;
-				StartInfo.Arguments = string.Format("add /f \"@{0}\" /s \"{1}\" /t \"{2}\"", TempFileName, TempSymStoreDir, Product);
+				StartInfo.Arguments = string.Format("add /f \"@{0}\" /s \"{1}\" /t \"{2}\" /o", TempFileName, TempSymStoreDir, Product);
 				StartInfo.UseShellExecute = false;
 				StartInfo.CreateNoWindow = true;
 				if (Utils.RunLocalProcessAndLogOutput(StartInfo, Log.Logger) != 0)
