[Gauntlet] Keep stream writer for the log file opened from FileLogger instance

#rnx

diff --git a/Engine/Source/Programs/AutomationTool/Gauntlet/Framework/Gauntlet.Utils.cs b/Engine/Source/Programs/AutomationTool/Gauntlet/Framework/Gauntlet.Utils.cs
--- a/Engine/Source/Programs/AutomationTool/Gauntlet/Framework/Gauntlet.Utils.cs
+++ b/Engine/Source/Programs/AutomationTool/Gauntlet/Framework/Gauntlet.Utils.cs
@@ -2612,15 +2612,28 @@ namespace Gauntlet
 		private CircularLogBuffer _buffer;
 		private string _filepath;
 		private string _commandline;
+		private StreamWriter _writer;
 
-		public FileLogger(string FilePath, string CommandLine)
+		/// <summary>
+		/// Instantiate Logger to file with interface to get a log reader
+		/// </summary>
+		/// <param name="FilePath">Target log file path</param>
+		/// <param name="CommandLine">Command line that producing the log output</param>
+		/// <param name="bKeepStreamOpen">If true, the stream writer is kept open until closed through a call to StreamClose(). The logger can still append lines when the stream is closed.</param>
+		public FileLogger(string FilePath, string CommandLine, bool bKeepStreamOpen = true)
 		{
 			_buffer = ProcessUtils.CreateLogBuffer();
 			_filepath = FilePath;
 			_commandline = CommandLine;
+			_writer = null;
 			try
 			{
-				ProcessUtils.CreateWriterForProcessLog(FilePath, CommandLine).Close();
+				_writer = ProcessUtils.CreateWriterForProcessLog(FilePath, CommandLine);
+				if (!bKeepStreamOpen)
+				{
+					_writer.Close();
+					_writer = null;
+				}
 			}
 			catch (IOException Ex)
 			{
@@ -2629,13 +2642,35 @@ namespace Gauntlet
 			}
 		}
 
+		~FileLogger()
+		{
+			CloseStream();
+		}
+
+		public void CloseStream()
+		{
+			if (_writer != null)
+			{
+				_writer.Close();
+				_writer = null;
+			}
+		}
+
 		public void AppendLine(string Data)
 		{
 			if (!string.IsNullOrEmpty(_filepath))
 			{
 				try
 				{
-					File.AppendAllText(_filepath, Data + '\n');
+					if (_writer == null)
+					{
+						File.AppendAllText(_filepath, Data + '\n');
+					}
+					else
+					{
+						_writer.Write(Data + '\n');
+						_writer.Flush();
+					}
 				}
 				catch (Exception Ex)
 				{
@@ -2887,6 +2922,7 @@ namespace Gauntlet
 						bHasExited = LocalProcess.HasExited;
 						if (bHasExited)
 						{
+							Logger?.CloseStream();
 							ExitCode = LocalProcess.ExitCode;
 						}
 					}
@@ -2912,6 +2948,7 @@ namespace Gauntlet
 		~LongProcessResult()
 		{
 			DisposeProcess();
+			Logger?.CloseStream();
 		}
 
 		public string GetProcessName()
@@ -3014,6 +3051,7 @@ namespace Gauntlet
 			{
 				Log.Info("Outputs did not close in time after process {ProcessName} exited.", LocalProcess.ProcessName);
 			}
+			Logger?.CloseStream();
 		}
 	}
 }
