[LLT] Avoid using StdOut directly, use WriteOutputToFile to copy the log and instanciate the Log buffer reader at the start of the test

#rnx

diff --git a/Engine/Source/Programs/AutomationTool/Gauntlet/Framework/Base/Gauntlet.AppInstance.cs b/Engine/Source/Programs/AutomationTool/Gauntlet/Framework/Base/Gauntlet.AppInstance.cs
--- a/Engine/Source/Programs/AutomationTool/Gauntlet/Framework/Base/Gauntlet.AppInstance.cs
+++ b/Engine/Source/Programs/AutomationTool/Gauntlet/Framework/Base/Gauntlet.AppInstance.cs
@@ -79,13 +79,6 @@ namespace Gauntlet
 
 	}
 
-	public interface IWithUnfilteredStdOut
-	{
-		string UnfilteredStdOut { get; }
-	}
-
-
-
 	/// <summary>
 	/// Interface used by IAppInstance if they support Suspend/Resume
 	/// </summary>
diff --git a/Engine/Source/Programs/AutomationTool/LowLevelTests/Tests/Gauntlet.LowLevelTestNode.cs b/Engine/Source/Programs/AutomationTool/LowLevelTests/Tests/Gauntlet.LowLevelTestNode.cs
--- a/Engine/Source/Programs/AutomationTool/LowLevelTests/Tests/Gauntlet.LowLevelTestNode.cs
+++ b/Engine/Source/Programs/AutomationTool/LowLevelTests/Tests/Gauntlet.LowLevelTestNode.cs
@@ -111,6 +111,7 @@ namespace LowLevelTests
 			if (TestInstance != null)
 			{
 				MarkTestStarted();
+				LogReader = TestInstance.GetLogBufferReader();
 			}
 
 			return TestInstance != null;
@@ -180,37 +181,19 @@ namespace LowLevelTests
 					TestInstance.Kill();
 				}
 
-				string StdOut;
-				if (TestInstance is IWithUnfilteredStdOut)
-				{
-					StdOut = ((IWithUnfilteredStdOut)TestInstance).UnfilteredStdOut;
-				}
-				else
-				{
-					StdOut = TestInstance.StdOut;
-				}
-
+				// Save log artifact
+				const string ClientLogFile = "ClientOutput.log";
+				string ClientOutputLog = Path.Combine(ArtifactPath, ClientLogFile);
 				string LogDir = Path.Combine(Unreal.EngineDirectory.FullName, "Programs", "AutomationTool", "Saved", "Logs");
-
-				if (StdOut == null || string.IsNullOrEmpty(StdOut.Trim()))
+				if (!TestInstance.WriteOutputToFile(ClientOutputLog))
 				{
 					Log.Warning("No StdOut returned from low level test app.");
 				}
-				else // Save log artifact
+				else
 				{
-					const string ClientLogFile = "ClientOutput.log";
-					string ClientOutputLog = Path.Combine(ArtifactPath, ClientLogFile);
-
-					using (var ClientOutputWriter = File.CreateText(ClientOutputLog))
-					{
-						ClientOutputWriter.Write(StdOut);
-					}
-
+					// Copy to UAT artifacts
 					string DestClientLogFile = Path.Combine(LogDir, ClientLogFile);
-					if (DestClientLogFile != ClientOutputLog)
-					{
-						File.Copy(ClientOutputLog, DestClientLogFile, true);
-					}
+					TestInstance.WriteOutputToFile(DestClientLogFile);
 				}
 
 				bool? ReportCopied = null;
@@ -221,20 +204,25 @@ namespace LowLevelTests
 				// No reports from Android tests yet. Since adb shell doesn't forward exit code, we look for it in the log output.
 				if (Context.Options.Platform == UnrealTargetPlatform.Android)
 				{
-					Match AndroidExitCodeLog = Regex.Match(StdOut, @"Tests finished with exit code (\d+)");
-					if (AndroidExitCodeLog.Success)
-					{
-						ExitCodeOverride = int.Parse(Regex.Match(AndroidExitCodeLog.Value, @"\d+").Value);
-					}
-					else if (StdOut.Contains("beginning of crash"))
+					ILogStreamReader AndroidLogReader = TestInstance.GetLogReader();
+					string ExitCodeLine = AndroidLogReader.EnumerateNextLines().Where(Line => Line.Contains("Tests finished with exit code")).FirstOrDefault();
+					if (!string.IsNullOrEmpty(ExitCodeLine))
 					{
-						Log.Info("Crash occurred during test.");
-						ExitCodeOverride = -1;
+						ExitCodeOverride = int.Parse(Regex.Match(ExitCodeLine, @"\d+").Value);
 					}
 					else
 					{
-						Log.Error("Could not find exit code in Android log, assumming failure.");
 						ExitCodeOverride = -1;
+						AndroidLogReader.SetLineIndex(0); // Reset reader
+						string CrashLine = AndroidLogReader.EnumerateNextLines().Where(Line => Line.Contains("beginning of crash")).FirstOrDefault();
+						if (!string.IsNullOrEmpty(CrashLine))
+						{
+							Log.Info("Crash occurred during test.");
+						}
+						else
+						{
+							Log.Error("Could not find exit code in Android log, assuming failure.");
+						}
 					}
 				}
 				else if (!string.IsNullOrEmpty(Context.Options.ReportType))
@@ -278,7 +266,7 @@ namespace LowLevelTests
 				else if (ReportCopied.HasValue && !ReportCopied.Value)
 				{
 					LowLevelTestResult = TestResult.Failed;
-					ExitReason = "Uabled to read test report";
+					ExitReason = "Unable to read test report";
 				}
 				else if (ReportPath != null)
 				{
@@ -356,11 +344,7 @@ namespace LowLevelTests
 
 		private void ParseLowLevelTestsLog()
 		{
-			if (LogReader == null)
-			{
-				LogReader = TestInstance?.GetLogBufferReader();
-			}
-			// Parse new lines from Stdout, if any
+			// Parse new lines from log, if any
 			CurrentProcessedLines = LogReader?.EnumerateNextLines().Where(Line => !string.IsNullOrWhiteSpace(Line)).ToArray();
 		}
 
