Add automation script to create a LiveLinkHub installed build.
You can pass -set:WithVendorPlugins to generate a build with vendor plugins.

Add parameter to the BuildPlugin command to allow outputting the build artifacts in a subdir named after the plugin.

diff --git a/Engine/Source/Programs/AutomationTool/LiveLinkHub/LiveLinkHub.Automation.cs b/Engine/Source/Programs/AutomationTool/LiveLinkHub/LiveLinkHub.Automation.cs
--- a/Engine/Source/Programs/AutomationTool/LiveLinkHub/LiveLinkHub.Automation.cs
+++ b/Engine/Source/Programs/AutomationTool/LiveLinkHub/LiveLinkHub.Automation.cs
@@ -1,29 +1,34 @@
 // Copyright Epic Games, Inc. All Rights Reserved.
 
 using EpicGames.Core;
+using Microsoft.Extensions.Logging;
+using System.Collections.Generic;
 using System.Linq;
 
 
+
 namespace AutomationTool
 {
 	class MakeLiveLinkHubEditor : MakeCookedEditor
 	{
 		protected override ProjectParams MakeParams(string DLCName, string BasedOnReleaseVersion)
 		{
-			return new ProjectParams(
+			ProjectParams Params = new ProjectParams(
 				Command: this
 				, RawProjectPath: ProjectFile
-
 				, NoBootstrapExe: true
 				, DLCName: DLCName
 				, BasedOnReleaseVersion: BasedOnReleaseVersion
 				, DedicatedServer: bIsCookedCooker
+				, SkipBuildClient: true
 				, NoClient: bIsCookedCooker
 				, OptionalContent: true
 				, ClientCookedTargets: new ParamList<string>() // Prevent AutodetectSettings from looking for a game target
 				, EditorTargets: new ParamList<string>("LiveLinkHub")
 				, UbtArgs: "-SingleModulePlatform"
 			);
+
+			return Params;
 		}
 
 		protected override void ModifyParams(ProjectParams BuildParams)
@@ -80,6 +85,19 @@ namespace AutomationTool
 
 			// Remove any files from NonUFS that were also added to Debug
 			SC.FilesToStage.NonUFSFiles = SC.FilesToStage.NonUFSFiles.Where(x => !SC.FilesToStage.NonUFSDebugFiles.ContainsKey(x.Key)).ToDictionary(x => x.Key, x => x.Value);
+			
+			// Make sure there are no restricted folders in the output
+			HashSet<StagedFileReference> RestrictedFiles = new HashSet<StagedFileReference>();
+
+			foreach (string RestrictedName in SC.RestrictedFolderNames)
+			{
+				RestrictedFiles.UnionWith(SC.FilesToStage.NonUFSDebugFiles.Keys.Where(x => x.ContainsName(RestrictedName)));
+			}
+
+			foreach (StagedFileReference StagedFileReference in RestrictedFiles)
+			{
+				SC.FilesToStage.NonUFSDebugFiles.Remove(StagedFileReference);
+			}
 		}
 	}
 }
diff --git a/Engine/Source/Programs/AutomationTool/Scripts/BuildPluginCommand.Automation.cs b/Engine/Source/Programs/AutomationTool/Scripts/BuildPluginCommand.Automation.cs
--- a/Engine/Source/Programs/AutomationTool/Scripts/BuildPluginCommand.Automation.cs
+++ b/Engine/Source/Programs/AutomationTool/Scripts/BuildPluginCommand.Automation.cs
@@ -20,6 +20,7 @@ using static AutomationTool.CommandUtils;
 [Help("HostPlatforms", "Specify a list of host platforms to build, separated by '+' characters (eg. -HostPlatforms=Win32+Win64). Default is the current host platforms")]
 [Help("TargetPlatforms", "Specify a list of target platforms to build, separated by '+' characters (eg. -TargetPlatforms=Win32+Win64). Default is all the Rocket target platforms.")]
 [Help("Package", "The path which the build artifacts should be packaged to, ready for distribution.")]
+[Help("PackageAppendPluginSubdir", "Package the build artifacts in a subfolder named after the plugin.")]
 [Help("StrictIncludes", "Disables precompiled headers and unity build in order to check all source files have self-contained headers.")]
 [Help("EngineDir=<RootDirectory>", "Root Directory of the engine that will be used to build plugin(s) (optional)")]
 [Help("Unversioned", "Do not embed the current engine version into the descriptor")]
@@ -65,6 +66,14 @@ public sealed class BuildPlugin : BuildCommand
 
 		// Make sure the packaging directory is valid
 		DirectoryReference PackageDir = new DirectoryReference(PackageParam);
+		
+		// If specified, package the build artifacts in a subfolder named after the initial plugin directory. (ie. MyPlugin/Manifest.uplugin will be packaged under {PackageDir}/MyPlufin)
+		if (ParseParam("PackageAppendPluginSubdir"))
+		{
+			string PluginSubfolder = PluginFile.Directory.GetDirectoryName();
+			PackageDir = DirectoryReference.Combine(PackageDir, PluginSubfolder);
+		}
+		
 		if (PluginFile.IsUnderDirectory(PackageDir))
 		{
 			throw new AutomationException("Packaged plugin output directory must be different to source");
