[Gauntlet] add support for ILogStreamReader in UnrealSnapshotParser

#rnx

diff --git a/Engine/Source/Programs/AutomationTool/Gauntlet/Unreal/Game/Samples/ElementalDemoTest.cs b/Engine/Source/Programs/AutomationTool/Gauntlet/Unreal/Game/Samples/ElementalDemoTest.cs
--- a/Engine/Source/Programs/AutomationTool/Gauntlet/Unreal/Game/Samples/ElementalDemoTest.cs
+++ b/Engine/Source/Programs/AutomationTool/Gauntlet/Unreal/Game/Samples/ElementalDemoTest.cs
@@ -37,7 +37,7 @@ namespace UnrealGame
 		{
 			UnrealRoleArtifacts ClientArtifacts = RoleResults.Where(R => R.Artifacts.SessionRole.RoleType == UnrealTargetRole.Client).Select(R => R.Artifacts).FirstOrDefault();
 
-			var SnapshotSummary = new UnrealSnapshotSummary<UnrealHealthSnapshot>(ClientArtifacts.AppInstance.StdOut);
+			var SnapshotSummary = new UnrealSnapshotSummary<UnrealHealthSnapshot>(ClientArtifacts.AppInstance.GetLogReader());
 
 			Log.Info("Elemental Performance Report");
 			Log.Info(SnapshotSummary.ToString());
diff --git a/Engine/Source/Programs/AutomationTool/Gauntlet/Unreal/Utils/Gauntlet.UnrealSnapshotParser.cs b/Engine/Source/Programs/AutomationTool/Gauntlet/Unreal/Utils/Gauntlet.UnrealSnapshotParser.cs
--- a/Engine/Source/Programs/AutomationTool/Gauntlet/Unreal/Utils/Gauntlet.UnrealSnapshotParser.cs
+++ b/Engine/Source/Programs/AutomationTool/Gauntlet/Unreal/Utils/Gauntlet.UnrealSnapshotParser.cs
@@ -241,7 +241,7 @@ namespace Gauntlet
 
 		public IEnumerable<TSnapshotClass> Snapshots;
 
-		public UnrealSnapshotSummary(string LogContents, string SnapshotTitles="")
+		public UnrealSnapshotSummary(ILogStreamReader LogReader, string SnapshotTitles="")
 		{
 			SampleCount = 0;
 			SessionTime = 0.0f;
@@ -260,12 +260,16 @@ namespace Gauntlet
 			DrawnPrims = new List<int>();
 			UnbuiltHLODs = new List<int>();
 
-			CreateFromLog(LogContents, SnapshotTitles);
+			CreateFromLog(LogReader, SnapshotTitles);
 		}
 
-		protected virtual void CreateFromLog(string LogContents, string InTitle)
+		public UnrealSnapshotSummary(string InContent, string SnapshotTitles = "")
+			: this(new DynamicStringReader(() => InContent), SnapshotTitles)
+		{ }
+
+		protected virtual void CreateFromLog(ILogStreamReader LogReader, string InTitle)
 		{
-			UnrealLogParser Parser = new UnrealLogParser(LogContents);
+			UnrealLogParser Parser = new UnrealLogParser(LogReader);
 
 			if (string.IsNullOrEmpty(InTitle))
 			{
