Add EpicAccountId property to Gauntlet.EpicAccount class

diff --git a/Engine/Source/Programs/AutomationTool/Gauntlet/Framework/Account/Gauntlet.Account.cs b/Engine/Source/Programs/AutomationTool/Gauntlet/Framework/Account/Gauntlet.Account.cs
--- a/Engine/Source/Programs/AutomationTool/Gauntlet/Framework/Account/Gauntlet.Account.cs
+++ b/Engine/Source/Programs/AutomationTool/Gauntlet/Framework/Account/Gauntlet.Account.cs
@@ -33,12 +33,19 @@ namespace Gauntlet
 
 		public string Password { get; protected set; }
 
+		public string EpicAccountId { get; protected set; }
+
 		public EpicAccount(string InUsername, string InPassword)
 		{
 			Username = InUsername;
 			Password = InPassword;
 		}
 
+		public EpicAccount(string InUserName, string InPassword, string InAccountId) : this(InUserName, InPassword)
+		{
+			EpicAccountId = InAccountId;
+		}
+
 		public override void ApplyToConfig(UnrealAppConfig AppConfig)
 		{
 			// add our credentials to the command line
