WinPlatform.Automation: Disallow using x86 symstore

#rnx

diff --git a/Engine/Source/Programs/AutomationTool/Win/WinPlatform.Automation.cs b/Engine/Source/Programs/AutomationTool/Win/WinPlatform.Automation.cs
--- a/Engine/Source/Programs/AutomationTool/Win/WinPlatform.Automation.cs
+++ b/Engine/Source/Programs/AutomationTool/Win/WinPlatform.Automation.cs
@@ -402,9 +402,9 @@ public class Win64Platform : Platform
 	}
 
     /// <summary>
-    /// Try to get the SYMSTORE.EXE path from the given Windows SDK version
+    /// Try to get the symstore.exe path from the given Windows SDK version
     /// </summary>
-    /// <returns>Path to SYMSTORE.EXE</returns>
+    /// <returns>Path to symstore.exe</returns>
 	[SupportedOSPlatform("windows")]
     private static FileReference GetSymStoreExe()
     {
@@ -416,15 +416,8 @@ public class Win64Platform : Platform
 
 			if (DirectoryReference.Exists(WindowsKitsDebuggersDirAutoSdk))
 			{
-				// Defaulting to the x86 because of a known issue with the latest x64 version
-				// x64 version gets the errorcode STATUS_ENTRYPOINT_NOT_FOUND on some configurations
-				FileReference SymStoreExe32 = FileReference.Combine(WindowsKitsDebuggersDirAutoSdk, "x86", "SymStore.exe");
-				if (FileReference.Exists(SymStoreExe32))
-				{
-					return SymStoreExe32;
-				}
 
-				FileReference SymStoreExe64 = FileReference.Combine(WindowsKitsDebuggersDirAutoSdk, "x64", "SymStore.exe");
+				FileReference SymStoreExe64 = FileReference.Combine(WindowsKitsDebuggersDirAutoSdk, "x64", "symstore.exe");
 				if (FileReference.Exists(SymStoreExe64))
 				{
 					return SymStoreExe64;
@@ -435,19 +428,13 @@ public class Win64Platform : Platform
 		List<KeyValuePair<string, DirectoryReference>> WindowsSdkDirs = WindowsExports.GetWindowsSdkDirs();
 		foreach (DirectoryReference WindowsSdkDir in WindowsSdkDirs.Select(x => x.Value))
 		{
-			FileReference SymStoreExe64 = FileReference.Combine(WindowsSdkDir, "Debuggers", "x64", "SymStore.exe");
+			FileReference SymStoreExe64 = FileReference.Combine(WindowsSdkDir, "Debuggers", "x64", "symstore.exe");
 			if (FileReference.Exists(SymStoreExe64))
 			{
 				return SymStoreExe64;
 			}
-
-			FileReference SymStoreExe32 = FileReference.Combine(WindowsSdkDir, "Debuggers", "x86", "SymStore.exe");
-			if (FileReference.Exists(SymStoreExe32))
-			{
-				return SymStoreExe32;
-			}
 		}
-		throw new AutomationException("Unable to find a Windows SDK installation containing PDBSTR.EXE");
+		throw new AutomationException("Unable to find a Windows SDK installation containing Debuggers/x64/symstore.exe");
     }
 
 	[SupportedOSPlatform("windows")]
@@ -554,6 +541,7 @@ public class Win64Platform : Platform
 
 		// Get the SYMSTORE.EXE path, using the latest SDK version we can find.
 		FileReference SymStoreExe = GetSymStoreExe();
+		Logger.LogInformation("Using '{Path}' Version {Version}", SymStoreExe, FileVersionInfo.GetVersionInfo(SymStoreExe.FullName).FileVersion);
 
 		List<FileReference> FilesToAdd = Files.Where(x => x.HasExtension(".pdb") || x.HasExtension(".exe") || x.HasExtension(".dll")).ToList();
 		if(FilesToAdd.Count > 0)
