WinPlatform.Automation: Disallow using x86 pdbcopy and pdbstr

#rnx

diff --git a/Engine/Source/Programs/AutomationTool/Win/WinPlatform.Automation.cs b/Engine/Source/Programs/AutomationTool/Win/WinPlatform.Automation.cs
--- a/Engine/Source/Programs/AutomationTool/Win/WinPlatform.Automation.cs
+++ b/Engine/Source/Programs/AutomationTool/Win/WinPlatform.Automation.cs
@@ -448,16 +448,7 @@ public class Win64Platform : Platform
 
 			if (DirectoryReference.Exists(WindowsKitsDebuggersDirAutoSdk))
 			{
-				// Defaulting to the x86 because of a known issue with the latest x64 version
-				// x64 version gets the errorcode STATUS_ENTRYPOINT_NOT_FOUND on some configurations
-				FileReference PdbCopyExe32 = FileReference.Combine(WindowsKitsDebuggersDirAutoSdk, "x86", "PdbCopy.exe");
-				if (FileReference.Exists(PdbCopyExe32))
-				{
-					OutLocation = PdbCopyExe32;
-					return true;
-				}
-
-				FileReference PdbCopyExe64 = FileReference.Combine(WindowsKitsDebuggersDirAutoSdk, "x64", "PdbCopy.exe");
+				FileReference PdbCopyExe64 = FileReference.Combine(WindowsKitsDebuggersDirAutoSdk, "x64", "pdbcopy.exe");
 				if (FileReference.Exists(PdbCopyExe64))
 				{
 					OutLocation = PdbCopyExe64;
@@ -470,30 +461,14 @@ public class Win64Platform : Platform
 		List<KeyValuePair<string, DirectoryReference>> WindowsSdkDirs = WindowsExports.GetWindowsSdkDirs();
 		foreach (DirectoryReference WindowsSdkDir in WindowsSdkDirs.Select(x => x.Value))
 		{
-			FileReference PdbCopyExe = FileReference.Combine(WindowsSdkDir, "Debuggers", "x64", "PdbCopy.exe");
-			if (FileReference.Exists(PdbCopyExe))
+			FileReference PdbCopyExe64 = FileReference.Combine(WindowsSdkDir, "Debuggers", "x64", "pdbcopy.exe");
+			if (FileReference.Exists(PdbCopyExe64))
 			{
-				OutLocation = PdbCopyExe;
+				OutLocation = PdbCopyExe64;
 				return true;
 			}
 		}
 
-		// Look for an installation of the MSBuild 14
-		FileReference LocationMsBuild14 = FileReference.Combine(DirectoryReference.GetSpecialFolder(Environment.SpecialFolder.ProgramFilesX86), "MSBuild", "Microsoft", "VisualStudio", "v14.0", "AppxPackage", "PDBCopy.exe");
-		if(FileReference.Exists(LocationMsBuild14))
-		{
-			OutLocation = LocationMsBuild14;
-			return true;
-		}
-
-		// Look for an installation of the MSBuild 12
-		FileReference LocationMsBuild12 = FileReference.Combine(DirectoryReference.GetSpecialFolder(Environment.SpecialFolder.ProgramFilesX86), "MSBuild", "Microsoft", "VisualStudio", "v12.0", "AppxPackage", "PDBCopy.exe");
-		if(FileReference.Exists(LocationMsBuild12))
-		{
-			OutLocation = LocationMsBuild12;
-			return true;
-		}
-
 		// Otherwise fail
 		OutLocation = null;
 		return false;
@@ -514,7 +489,7 @@ public class Win64Platform : Platform
 		FileReference PdbCopyLocation;
 		if(!TryGetPdbCopyLocation(out PdbCopyLocation))
 		{
-			throw new AutomationException("Unable to find installation of PDBCOPY.EXE, which is required to strip symbols. This tool is included as part of the 'Windows Debugging Tools' component of the Windows 10 SDK (https://developer.microsoft.com/en-us/windows/downloads/windows-10-sdk).");
+			throw new AutomationException("Unable to find installation of pdbcopy.exe, which is required to strip symbols. This tool is included as part of the 'Windows Debugging Tools' component of the Windows 10 SDK (https://developer.microsoft.com/en-us/windows/downloads/windows-10-sdk).");
 		}
 
 		ProcessStartInfo StartInfo = new ProcessStartInfo();
@@ -539,7 +514,7 @@ public class Win64Platform : Platform
     {
 		Logger.LogInformation("Publishing symbols to \"{SymbolStoreDirectory}\" (source indexing: {bIndexSources})", SymbolStoreDirectory, bIndexSources);
 
-		// Get the SYMSTORE.EXE path, using the latest SDK version we can find.
+		// Get the symstore.exe path, using the latest SDK version we can find.
 		FileReference SymStoreExe = GetSymStoreExe();
 		Logger.LogInformation("Using '{Path}' Version {Version}", SymStoreExe, FileVersionInfo.GetVersionInfo(SymStoreExe.FullName).FileVersion);
 
@@ -774,7 +749,7 @@ public class Win64Platform : Platform
 			throw new AutomationException($"Failed to query the source code information for '{DepotFilter}'.");
 		}
 
-		// Get the PDBSTR.EXE path, using the latest SDK version we can find.
+		// Get the pdbstr.exe path, using the latest SDK version we can find.
 		FileReference PdbStrExe = GetPdbStrExe();
 
 		// Get the path to the generated SRCSRV.INI file
@@ -823,7 +798,7 @@ public class Win64Platform : Platform
 	/// <summary>
 	/// Executes the PdbStr tool.
 	/// </summary>
-	/// <param name="PdbStrExe">Path to PdbStr.exe</param>
+	/// <param name="PdbStrExe">Path to pdbstr.exe</param>
 	/// <param name="PdbFile">The PDB file to embed source information for</param>
 	/// <param name="SrcSrvIni">Ini file containing settings to embed</param>
 	/// <param name="State">The current loop state</param>
@@ -870,9 +845,9 @@ public class Win64Platform : Platform
 	}
 
 	/// <summary>
-	/// Try to get the PDBSTR.EXE path from the Windows SDK
+	/// Try to get the pdbstr.exe path from the Windows SDK
 	/// </summary>
-	/// <returns>Path to PDBSTR.EXE</returns>
+	/// <returns>Path to pdbstr.exe</returns>
 	[SupportedOSPlatform("windows")]
 	static FileReference GetPdbStrExe()
 	{
@@ -886,15 +861,7 @@ public class Win64Platform : Platform
 
 			if (DirectoryReference.Exists(WindowsKitsDebuggersDirAutoSdk))
 			{
-				// Defaulting to the x86 because of a known issue with the latest x64 version
-				// x64 version gets the errorcode STATUS_ENTRYPOINT_NOT_FOUND on some configurations
-				FileReference CheckPdbStrExe32 = FileReference.Combine(WindowsKitsDebuggersDirAutoSdk, "x86", "SrcSrv", "PdbStr.exe");
-				if (FileReference.Exists(CheckPdbStrExe32))
-				{
-					return CheckPdbStrExe32;
-				}
-
-				FileReference CheckPdbStrExe64 = FileReference.Combine(WindowsKitsDebuggersDirAutoSdk, "x64", "SrcSrv", "PdbStr.exe");
+				FileReference CheckPdbStrExe64 = FileReference.Combine(WindowsKitsDebuggersDirAutoSdk, "x64", "srcsrv", "pdbstr.exe");
 				if (FileReference.Exists(CheckPdbStrExe64))
 				{
 					return CheckPdbStrExe64;
@@ -904,19 +871,14 @@ public class Win64Platform : Platform
 
 		foreach (DirectoryReference WindowsSdkDir in WindowsSdkDirs.Select(x => x.Value))
 		{
-			FileReference CheckPdbStrExe64 = FileReference.Combine(WindowsSdkDir, "Debuggers", "x64", "SrcSrv", "PdbStr.exe");
+			FileReference CheckPdbStrExe64 = FileReference.Combine(WindowsSdkDir, "Debuggers", "x64", "srcsrv", "pdbstr.exe");
 			if (FileReference.Exists(CheckPdbStrExe64))
 			{
 				return CheckPdbStrExe64;
 			}
-
-			FileReference CheckPdbStrExe32 = FileReference.Combine(WindowsSdkDir, "Debuggers", "x86", "SrcSrv", "PdbStr.exe");
-			if (FileReference.Exists(CheckPdbStrExe32))
-			{
-				return CheckPdbStrExe32;
-			}
 		}
-		throw new AutomationException("Unable to find a Windows SDK installation containing PDBSTR.EXE");
+
+        throw new AutomationException("Unable to find a Windows SDK installation containing Debuggers/x64/srcsrv/pdbstr.exe");
 	}
 
 	public override string[] SymbolServerDirectoryStructure
