- Added variables to the "whatsnew" parameter ([REDACTED_HANDLE], [REDACTED_HANDLE], [REDACTED_HANDLE], [REDACTED_HANDLE])

diff --git a/Engine/Source/Programs/AutomationTool/Apple/ExportAppleBuild.Automation.cs b/Engine/Source/Programs/AutomationTool/Apple/ExportAppleBuild.Automation.cs
--- a/Engine/Source/Programs/AutomationTool/Apple/ExportAppleBuild.Automation.cs
+++ b/Engine/Source/Programs/AutomationTool/Apple/ExportAppleBuild.Automation.cs
@@ -142,10 +142,16 @@ class ExportAppleBuild : BuildCommand
 		if (!string.IsNullOrEmpty(WhatsNew))
 		{
 			// get some info from the .xcarchive that we just uploaded
-			string BundleId, AppVersion, BundleVersion;
-			GetInfoFromXcArchive(XcArchive, out BundleId, out AppVersion, out BundleVersion);
+			string BundleId, AppVersion, BundleVersion, CommandLine;
+			GetInfoFromXcArchive(XcArchive, out BundleId, out AppVersion, out BundleVersion, out CommandLine);
 
-			Logger.LogInformation("Setting What's New for {BundleId}, app version{Version}, build version {BundleVersion}", BundleId, AppVersion, BundleVersion);
+			Logger.LogInformation("Setting What's New for {BundleId}, app version{Version}, build version {BundleVersion}, commandline {CommandLine}, final whatsnew = {WhatsNew}", BundleId, AppVersion, BundleVersion, CommandLine, WhatsNew);
+
+			// replace some variables
+			WhatsNew = WhatsNew.Replace("@appversion", AppVersion);
+			WhatsNew = WhatsNew.Replace("@buildversion", BundleVersion);
+			WhatsNew = WhatsNew.Replace("@bundleid", BundleId);
+			WhatsNew = WhatsNew.Replace("@uecommandline", CommandLine);
 
 			UpdateWhatsNew(JWTToken, BundleId, AppVersion, BundleVersion, WhatsNew);
 		}
@@ -270,7 +276,7 @@ class ExportAppleBuild : BuildCommand
 	{
 		// get some info from the .xcarchive
 		string BundleId, AppVersion, BundleVersion;
-		GetInfoFromXcArchive(XcArchive, out BundleId, out AppVersion, out BundleVersion);
+		GetInfoFromXcArchive(XcArchive, out BundleId, out AppVersion, out BundleVersion, out _);
 
 
 Console.WriteLine("Checking for version >= {0}", BundleVersion);
@@ -504,7 +510,7 @@ Console.WriteLine("Checking for version >= {0}", BundleVersion);
 		}
 	}
 
-	static void GetInfoFromXcArchive(DirectoryReference XcArchive, out string BundleId, out string AppVersion, out string BundleVersion)
+	static void GetInfoFromXcArchive(DirectoryReference XcArchive, out string BundleId, out string AppVersion, out string BundleVersion, out string CommandLine)
 	{
 		BundleId = null;
 		AppVersion = null;
@@ -555,6 +561,22 @@ Console.WriteLine("Checking for version >= {0}", BundleVersion);
 		{
 			throw new Exception($"Failed to get bundle ID and versions from {Plist}");
 		}
+
+		CommandLine = "";
+		// now find the uecommandline.txt - we don't know the name of the .app inside the .xcarchive, but there should only be one
+		DirectoryReference ProductsDir = DirectoryReference.Combine(XcArchive, "Products", "Applications");
+		foreach (DirectoryReference AppDir in DirectoryReference.EnumerateDirectories(ProductsDir, "*.app"))
+		{
+			FileReference CmdLineFile = FileReference.Combine(AppDir, "uecommandline.txt");
+			if (FileReference.Exists(CmdLineFile))
+			{
+				if (CommandLine != "")
+				{
+					throw new Exception("Multiple .app's in the .xcarchive. This is unexpected");
+				}
+				CommandLine = File.ReadAllText(CmdLineFile.FullName);
+			}
+		}
 	}
 
 	static string CreateJWTToken(FileReference ProjectFile)
