Update sha1 code signing to use /fd and /td args

Args are now required with new signtool shipped with newer win sdks

#rnx

diff --git a/Engine/Source/Programs/AutomationTool/AutomationUtils/CommandUtils.cs b/Engine/Source/Programs/AutomationTool/AutomationUtils/CommandUtils.cs
--- a/Engine/Source/Programs/AutomationTool/AutomationUtils/CommandUtils.cs
+++ b/Engine/Source/Programs/AutomationTool/AutomationUtils/CommandUtils.cs
@@ -3071,7 +3071,7 @@ namespace AutomationTool
 				StringBuilder FileList = new StringBuilder();
 				if (SignatureType == SignatureType.SHA1)
 				{
-					CommandLine.Append($"sign{SpecificStoreArg} /a /n \"{SigningIdentity}\" /v {MultipleSignatureArg} {DescriptionArg}");
+					CommandLine.Append($"sign{SpecificStoreArg} /a /fd sha1 /td sha1 /n \"{SigningIdentity}\" /v {MultipleSignatureArg} {DescriptionArg}");
 				}
 				else if (SignatureType == SignatureType.SHA256)
 				{
