[Gauntlet] Support fake log reader in NullClient

#rnx

diff --git a/Engine/Source/Programs/AutomationTool/Gauntlet/Platform/Null/Gauntlet.TargetDeviceNull.cs b/Engine/Source/Programs/AutomationTool/Gauntlet/Platform/Null/Gauntlet.TargetDeviceNull.cs
--- a/Engine/Source/Programs/AutomationTool/Gauntlet/Platform/Null/Gauntlet.TargetDeviceNull.cs
+++ b/Engine/Source/Programs/AutomationTool/Gauntlet/Platform/Null/Gauntlet.TargetDeviceNull.cs
@@ -37,9 +37,9 @@ namespace Gauntlet
 
 		public string StdOut { get { return ""; } }
 
-		public ILogStreamReader GetLogReader() { return null; }
+		public ILogStreamReader GetLogReader() { return new DynamicStringReader(() => ""); }
 
-		public ILogStreamReader GetLogBufferReader() { return null; }
+		public ILogStreamReader GetLogBufferReader() { return new DynamicStringReader(() => ""); }
 
 		public bool WriteOutputToFile(string FilePath) { return false; }
 
