Code Signing use /tr instead of /t

As we now force the algo so /tr needs to be used.

#rnx

diff --git a/Engine/Source/Programs/AutomationTool/AutomationUtils/CommandUtils.cs b/Engine/Source/Programs/AutomationTool/AutomationUtils/CommandUtils.cs
--- a/Engine/Source/Programs/AutomationTool/AutomationUtils/CommandUtils.cs
+++ b/Engine/Source/Programs/AutomationTool/AutomationUtils/CommandUtils.cs
@@ -3057,7 +3057,7 @@ namespace AutomationTool
 			string DescriptionArg = String.IsNullOrEmpty(Description) ? "" : $"/d \"{Description}\"";
 			string SpecificStoreArg = bUseMachineStoreForCertificates ? " /sm" : "";
 			string MultipleSignatureArg = AllowMultipleSignatures ? " /as" : "";
-			string SHA1TimestampArg = AllowMultipleSignatures ? " /tr" : " /t";
+			string SHA1TimestampArg = " /tr";
 			string SHA256TimestampArg = " /tr";
 
 			List<Task> SignTaskList = new List<Task>();
