Restrict ExecutorClient to be private so that we can ensure Timeout is only set before the first request is made.

diff --git a/Engine/Source/Programs/AutomationTool/Gauntlet/Framework/RpcFramework/Gauntlet.RpcExecutor.cs b/Engine/Source/Programs/AutomationTool/Gauntlet/Framework/RpcFramework/Gauntlet.RpcExecutor.cs
--- a/Engine/Source/Programs/AutomationTool/Gauntlet/Framework/RpcFramework/Gauntlet.RpcExecutor.cs
+++ b/Engine/Source/Programs/AutomationTool/Gauntlet/Framework/RpcFramework/Gauntlet.RpcExecutor.cs
@@ -117,6 +117,8 @@ namespace Gauntlet
 					if (ExecutorInstance == null)
 					{
 						ExecutorInstance = new RpcExecutor();
+						int RpcHttpTimeoutInMilliseconds = Globals.Params.ParseValue("RPCHttpTimeoutInMilliseconds", 100 * 1000);
+						ExecutorClient.Timeout = new TimeSpan(0, 0, 0, 0, RpcHttpTimeoutInMilliseconds);
 					}
 					return ExecutorInstance;
 				}
@@ -799,7 +801,7 @@ namespace Gauntlet
         	// We might timeout calling CallRpc on a client that's closing - we do not want this to be a fail;
         	try
         	{
-        		HttpResponseMessage RpcResponse = ExecutorClient.Send(RpcRequest);
+				HttpResponseMessage RpcResponse = ExecutorClient.Send(RpcRequest);
         		return RpcResponse;
         	}
         	catch (Exception e)
