Fix wrong $S(BuildDir) in UPL when using modern xcode

#rnx

diff --git a/Engine/Source/Programs/AutomationTool/IOS/IOSPlatform.Automation.cs b/Engine/Source/Programs/AutomationTool/IOS/IOSPlatform.Automation.cs
--- a/Engine/Source/Programs/AutomationTool/IOS/IOSPlatform.Automation.cs
+++ b/Engine/Source/Programs/AutomationTool/IOS/IOSPlatform.Automation.cs
@@ -1388,7 +1388,7 @@ public class IOSPlatform : ApplePlatform
 							(SC.IsCodeBasedProject ? TargetName : "UnrealGame"),
 							SC.IsCodeBasedProject ? false : Params.Client, // Code based projects will have Client in their executable name already
 							SC.ShortProjectName, DirectoryReference.Combine(SC.LocalRoot, "Engine"),
-							DirectoryReference.Combine((SC.IsCodeBasedProject ? SC.ProjectRoot : DirectoryReference.Combine(SC.LocalRoot, "Engine")), "Binaries", PlatformName, "Payload", (SC.IsCodeBasedProject ? SC.ShortProjectName : "UnrealGame") + ".app"),
+							DirectoryReference.Combine((SC.IsCodeBasedProject ? SC.ProjectRoot : DirectoryReference.Combine(SC.LocalRoot, "Engine")), "Binaries", PlatformName, AppleExports.UseModernXcode(SC.RawProjectPath) ? "" : "Payload", (SC.IsCodeBasedProject ? SC.ShortProjectName : "UnrealGame") + ".app"),
 							SC.StageExecutables[0]);
 
 					if (!AppleExports.UseModernXcode(SC.RawProjectPath))
