Fix an issue with chunked builds related to incorrect quoting of path

diff --git a/Engine/Source/Programs/AutomationTool/Scripts/CopyBuildToStagingDirectory.Automation.cs b/Engine/Source/Programs/AutomationTool/Scripts/CopyBuildToStagingDirectory.Automation.cs
--- a/Engine/Source/Programs/AutomationTool/Scripts/CopyBuildToStagingDirectory.Automation.cs
+++ b/Engine/Source/Programs/AutomationTool/Scripts/CopyBuildToStagingDirectory.Automation.cs
@@ -4214,7 +4214,7 @@ namespace AutomationScripts
 						}
 
 						string BuildRoot = MakePathSafeToUseWithCommandLine(RawDataPath);
-						string CloudDir = MakePathSafeToUseWithCommandLine(CombinePaths(ChunkInstallBasePath, "CloudDir"));
+						string CloudDir = CombinePaths(ChunkInstallBasePath, "CloudDir");
 						string ManifestDir = CombinePaths(ChunkInstallBasePath, "ManifestDir");
 						var AppID = 1; // For a chunk install this value doesn't seem to matter
 						string AppName = String.Format("{0}_{1}", SC.ShortProjectName, PakName);
@@ -4224,7 +4224,7 @@ namespace AutomationScripts
 						string DestManifestPath = CombinePaths(ManifestDir, ManifestFilename);
 						InternalUtils.SafeCreateDirectory(ManifestDir, true);
 
-						string CmdLine = String.Format("-BuildRoot={0} -CloudDir={1} -AppID={2} -AppName=\"{3}\" -BuildVersion=\"{4}\" -AppLaunch=\"{5}\"", BuildRoot, CloudDir, AppID, AppName, VersionString, AppLaunch);
+						string CmdLine = String.Format("-BuildRoot={0} -CloudDir={1} -AppID={2} -AppName=\"{3}\" -BuildVersion=\"{4}\" -AppLaunch=\"{5}\"", BuildRoot, MakePathSafeToUseWithCommandLine(CloudDir), AppID, AppName, VersionString, AppLaunch);
 						CmdLine += " -AppArgs=\"\"";
 						CmdLine += " -custom=\"bIsPatch=false\"";
 						CmdLine += String.Format(" -customint=\"ChunkID={0}\"", ChunkID);
