[Gauntlet] UnrealLogStreamParser - Use HashSet instead of List for internal tracking + Add IsAttachedToLogReader() for convenience.

#rnx

diff --git a/Engine/Source/Programs/AutomationTool/Gauntlet/Unreal/Utils/Gauntlet.UnrealLogParser.cs b/Engine/Source/Programs/AutomationTool/Gauntlet/Unreal/Utils/Gauntlet.UnrealLogParser.cs
--- a/Engine/Source/Programs/AutomationTool/Gauntlet/Unreal/Utils/Gauntlet.UnrealLogParser.cs
+++ b/Engine/Source/Programs/AutomationTool/Gauntlet/Unreal/Utils/Gauntlet.UnrealLogParser.cs
@@ -293,14 +293,14 @@ namespace Gauntlet
 	{
 		protected List<UnrealLog.LogEntry> LogEvents { get; private set; }
 
-		private List<string> UnidentifiedLogLevels { get; set; }
+		private HashSet<string> UnidentifiedLogLevels { get; set; }
 
 		private ILogStreamReader LogReader { get; set; }
 
 		public UnrealLogStreamParser()
 		{
 			LogEvents = new();
-			UnidentifiedLogLevels = new List<string>();
+			UnidentifiedLogLevels = new HashSet<string>();
 			LogReader = null;
 		}
 
@@ -319,6 +319,15 @@ namespace Gauntlet
 			LogReader = InLogReader;
 		}
 
+		/// <summary>
+		/// Return true if the internal log reader was set.
+		/// </summary>
+		/// <returns></returns>
+		public bool IsAttachedToLogReader()
+		{
+			return LogReader != null;
+		}
+
 		/// <summary>
 		/// Clear aggregated log events 
 		/// </summary>
@@ -335,7 +344,7 @@ namespace Gauntlet
 		/// <returns>The number of line parsed</returns>
 		public int ReadStream(int LineOffset = -1, bool bClearAggregatedLines = true)
 		{
-			if (LogReader == null)
+			if (!IsAttachedToLogReader())
 			{
 				throw new AutomationException("Internal Log reader is not set. Use SetLogReader() to set it.");
 			}
