Make Turnkey persist UE_SDKS_ROOT on Mac and Linux

diff --git a/Engine/Build/BatchFiles/Linux/SetupDotnet.sh b/Engine/Build/BatchFiles/Linux/SetupDotnet.sh
--- a/Engine/Build/BatchFiles/Linux/SetupDotnet.sh
+++ b/Engine/Build/BatchFiles/Linux/SetupDotnet.sh
@@ -57,4 +57,15 @@ else
 	export IS_DOTNET_INSTALLED=$IS_DOTNET_INSTALLED
 fi
 
+# this is the current assumed location for now
+# We use FUnixPlatformProcess::ApplicationSettingsDir() from c++ and
+# and for C# it uses Environment.GetFolderPath(SpecialFolder.ApplicationData)
+# for this location, so lets share this as our "place to put an AutoSDK file"
+AUTO_SDK_PATH_FILE="$HOME/.config/.autosdk"
+
+# if the file exists and we dont currently have a $UE_SDKS_ROOT set, lets setup UE_SDKS_ROOT to our files location path
+if [ -f "$AUTO_SDK_PATH_FILE" ] && [ -z "$UE_SDKS_ROOT" ]; then
+	export UE_SDKS_ROOT="$(cat $AUTO_SDK_PATH_FILE)"
+fi
+
 cd "$START_DIR"
diff --git a/Engine/Build/BatchFiles/Mac/SetupDotnet.sh b/Engine/Build/BatchFiles/Mac/SetupDotnet.sh
--- a/Engine/Build/BatchFiles/Mac/SetupDotnet.sh
+++ b/Engine/Build/BatchFiles/Mac/SetupDotnet.sh
@@ -48,4 +48,15 @@ else
 	export IS_DOTNET_INSTALLED=$IS_DOTNET_INSTALLED
 fi
 
+# this is the current assumed location for now
+# We use FUnixPlatformProcess::ApplicationSettingsDir() from c++ and
+# and for C# it uses Environment.GetFolderPath(SpecialFolder.ApplicationData)
+# for this location, so lets share this as our "place to put an AutoSDK file"
+AUTO_SDK_PATH_FILE="$HOME/.config/.autosdk"
+
+# if the file exists and we dont currently have a $UE_SDKS_ROOT set, lets setup UE_SDKS_ROOT to our files location path
+if [ -f "$AUTO_SDK_PATH_FILE" ] && [ -z "$UE_SDKS_ROOT" ]; then
+	export UE_SDKS_ROOT="$(cat $AUTO_SDK_PATH_FILE)"
+fi
+
 cd "$START_DIR"
diff --git a/Engine/Source/Programs/AutomationTool/AutomationUtils/Platform.cs b/Engine/Source/Programs/AutomationTool/AutomationUtils/Platform.cs
--- a/Engine/Source/Programs/AutomationTool/AutomationUtils/Platform.cs
+++ b/Engine/Source/Programs/AutomationTool/AutomationUtils/Platform.cs
@@ -316,6 +316,10 @@ namespace AutomationTool
 			return ExitCode == 0;
 		}
 
+		public virtual void PersistSdkRootVar()
+		{
+		}
+
 		public virtual string GetSDKCreationHelp()
 		{
 			return null;
diff --git a/Engine/Source/Programs/AutomationTool/Linux/LinuxPlatform.Automation.cs b/Engine/Source/Programs/AutomationTool/Linux/LinuxPlatform.Automation.cs
--- a/Engine/Source/Programs/AutomationTool/Linux/LinuxPlatform.Automation.cs
+++ b/Engine/Source/Programs/AutomationTool/Linux/LinuxPlatform.Automation.cs
@@ -47,6 +47,23 @@ public abstract class BaseLinuxPlatform : Platform
 		return Devices.ToArray();
 	}
 
+	public override void PersistSdkRootVar()
+	{
+		string UeSdksRoot = Environment.GetEnvironmentVariable("UE_SDKS_ROOT");
+		if (UeSdksRoot != null)
+		{
+			base.PersistSdkRootVar();
+			string AutoSdkFile = Path.Combine(Environment.GetFolderPath(Environment.SpecialFolder.ApplicationData), ".autosdk");
+			if (!File.Exists(AutoSdkFile))
+			{
+				using (StreamWriter wr = new StreamWriter(AutoSdkFile))
+				{
+					wr.WriteLine(UeSdksRoot);
+				}
+			}
+		}
+	}
+
 	public override void GetFilesToDeployOrStage(ProjectParams Params, DeploymentContext SC)
 	{
 		if (SC.bStageCrashReporter)
diff --git a/Engine/Source/Programs/AutomationTool/Mac/MacPlatform.Automation.cs b/Engine/Source/Programs/AutomationTool/Mac/MacPlatform.Automation.cs
--- a/Engine/Source/Programs/AutomationTool/Mac/MacPlatform.Automation.cs
+++ b/Engine/Source/Programs/AutomationTool/Mac/MacPlatform.Automation.cs
@@ -50,6 +50,23 @@ public class MacPlatform : ApplePlatform
 		}		
 	}
 
+	public override void PersistSdkRootVar()
+	{
+		string UeSdksRoot = Environment.GetEnvironmentVariable("UE_SDKS_ROOT");
+		if (UeSdksRoot != null)
+		{
+			base.PersistSdkRootVar();
+			string AutoSdkFile = Path.Combine(Environment.GetFolderPath(Environment.SpecialFolder.ApplicationData), ".autosdk");
+			if (!File.Exists(AutoSdkFile))
+			{
+				using (StreamWriter wr = new StreamWriter(AutoSdkFile))
+				{
+					wr.WriteLine(UeSdksRoot);
+				}
+			}
+		}
+	}
+
 	public override DeviceInfo[] GetDevices()
 	{
 		List<DeviceInfo> Devices = new List<DeviceInfo>();
diff --git a/Engine/Source/Programs/AutomationTool/Turnkey/SdkInfo.cs b/Engine/Source/Programs/AutomationTool/Turnkey/SdkInfo.cs
--- a/Engine/Source/Programs/AutomationTool/Turnkey/SdkInfo.cs
+++ b/Engine/Source/Programs/AutomationTool/Turnkey/SdkInfo.cs
@@ -1,4 +1,4 @@
-﻿// Copyright Epic Games, Inc. All Rights Reserved.
+// Copyright Epic Games, Inc. All Rights Reserved.
 
 using System;
 using System.Collections.Generic;
@@ -215,6 +215,7 @@ namespace Turnkey
 					Environment.SetEnvironmentVariable("UE_SDKS_ROOT", Response);
 					Environment.SetEnvironmentVariable("UE_SDKS_ROOT", Response, EnvironmentVariableTarget.User);
 					TurnkeyUtils.EndTrackingExternalEnvVarChanges();
+					AutomationPlatform.PersistSdkRootVar();
 
 					bSetupEnvVarAfterInstall = false;
 				}
@@ -274,6 +275,7 @@ namespace Turnkey
 				Environment.SetEnvironmentVariable("UE_SDKS_ROOT", AutoSdkDir);
 				Environment.SetEnvironmentVariable("UE_SDKS_ROOT", AutoSdkDir, EnvironmentVariableTarget.User);
 				TurnkeyUtils.EndTrackingExternalEnvVarChanges();
+				AutomationPlatform.PersistSdkRootVar();
 
 			}
 
