- Trimmed whitespace off the commnandline for [REDACTED_HANDLE] usage

diff --git a/Engine/Source/Programs/AutomationTool/Apple/ExportAppleBuild.Automation.cs b/Engine/Source/Programs/AutomationTool/Apple/ExportAppleBuild.Automation.cs
--- a/Engine/Source/Programs/AutomationTool/Apple/ExportAppleBuild.Automation.cs
+++ b/Engine/Source/Programs/AutomationTool/Apple/ExportAppleBuild.Automation.cs
@@ -152,6 +152,7 @@ class ExportAppleBuild : BuildCommand
 			WhatsNew = WhatsNew.Replace("@buildversion", BundleVersion);
 			WhatsNew = WhatsNew.Replace("@bundleid", BundleId);
 			WhatsNew = WhatsNew.Replace("@uecommandline", CommandLine);
+			WhatsNew = WhatsNew.Replace("@commandline", CommandLine);
 
 			UpdateWhatsNew(JWTToken, BundleId, AppVersion, BundleVersion, WhatsNew);
 		}
@@ -574,7 +575,7 @@ Console.WriteLine("Checking for version >= {0}", BundleVersion);
 				{
 					throw new Exception("Multiple .app's in the .xcarchive. This is unexpected");
 				}
-				CommandLine = File.ReadAllText(CmdLineFile.FullName);
+				CommandLine = File.ReadAllText(CmdLineFile.FullName).Trim();
 			}
 		}
 	}
