RPC Framework changes to reduce logging and prevent RPC Listener failures during shutdown

#rnx

diff --git a/Engine/Source/Programs/AutomationTool/Gauntlet/Framework/RpcFramework/Gauntlet.RpcExecutor.cs b/Engine/Source/Programs/AutomationTool/Gauntlet/Framework/RpcFramework/Gauntlet.RpcExecutor.cs
--- a/Engine/Source/Programs/AutomationTool/Gauntlet/Framework/RpcFramework/Gauntlet.RpcExecutor.cs
+++ b/Engine/Source/Programs/AutomationTool/Gauntlet/Framework/RpcFramework/Gauntlet.RpcExecutor.cs
@@ -228,18 +228,17 @@ namespace Gauntlet
 			{
 				if (CallPath.Equals(RouteEntry.Key))
 				{
-					
-					using (HttpListenerResponse ListenerResponse = ListenerContext.Response)
+					try
 					{
-						if (ListenerResponse == null)
+						using (HttpListenerResponse ListenerResponse = ListenerContext.Response)
 						{
-							return;
-						}
-						ListenerResponse.Headers.Set("Content-Type", "text/plain");
+							if (ListenerResponse == null)
+							{
+								return;
+							}
+							ListenerResponse.Headers.Set("Content-Type", "text/plain");
 
-						using (Stream ResponseStream = ListenerResponse.OutputStream)
-						{
-							try
+							using (Stream ResponseStream = ListenerResponse.OutputStream)
 							{
 								if (ResponseStream == null)
 								{
@@ -257,29 +256,53 @@ namespace Gauntlet
 								ListenerResponse.ContentLength64 = ResponseBuffer.Length;
 								ResponseStream.Write(ResponseBuffer, 0, ResponseBuffer.Length);
 							}
-							catch (HttpListenerException)
-							{
-								Log.Verbose(string.Format("Discarded response for call to {0} because we are in the process of shutting down.", CallPath.Path));
-							}
 						}
 					}
-					
+					catch (Exception Error)
+					{
+						if (bHasReceivedTerminationCall)
+						{
+							Log.SuspendECErrorParsing();
+							Log.Info("Prevented Shutdown failure. Discarded response for call to {CallPath} because we are in the process of shutting down. Error: {ErrorMessage}", CallPath.Path, Error.Message);
+							Log.ResumeECErrorParsing();
+						}
+						else
+						{
+							throw new(Error.Message);
+						}
+					}
 					return;
 				}
 			}
-			using (HttpListenerResponse ListenerResponse = ListenerContext.Response)
-			{
-				ListenerResponse.Headers.Set("Content-Type", "text/plain");
 
-				using (Stream ResponseStream = ListenerResponse.OutputStream)
+			try 
+			{ 
+				using (HttpListenerResponse ListenerResponse = ListenerContext.Response)
 				{
-					ListenerContext.Response.StatusCode = (int)HttpStatusCode.NotFound;
-					byte[] ResponseBuffer = Encoding.UTF8.GetBytes(CallResponse.ResponseBody);
-					ListenerResponse.ContentLength64 = ResponseBuffer.Length;
-					ResponseStream.Write(ResponseBuffer, 0, ResponseBuffer.Length);
+					ListenerResponse.Headers.Set("Content-Type", "text/plain");
+
+					using (Stream ResponseStream = ListenerResponse.OutputStream)
+					{
+						ListenerContext.Response.StatusCode = (int)HttpStatusCode.NotFound;
+						byte[] ResponseBuffer = Encoding.UTF8.GetBytes(CallResponse.ResponseBody);
+						ListenerResponse.ContentLength64 = ResponseBuffer.Length;
+						ResponseStream.Write(ResponseBuffer, 0, ResponseBuffer.Length);
+					}
+				}
+			}
+			catch (Exception Error)
+			{
+				if (bHasReceivedTerminationCall)
+				{
+					Log.SuspendECErrorParsing();
+					Log.Info("Prevented Shutdown failure. Discarded response for call to {CallPath} because we are in the process of shutting down. Error: {ErrorMessage}", CallPath.Path, Error.Message);
+					Log.ResumeECErrorParsing();
+				}
+				else
+				{
+					throw new(Error.Message);
 				}
 			}
-
 		}
 		public List<GauntletIncomingRpcMessage> GetNewMessages(bool bClearMessageQueue = true)
 		{
