LaunchMultiServer UAT script now uses SaveWinPos to save window positions of clients.

diff --git a/Engine/Source/Programs/AutomationTool/Scripts/LaunchMultiServer.Automation.cs b/Engine/Source/Programs/AutomationTool/Scripts/LaunchMultiServer.Automation.cs
--- a/Engine/Source/Programs/AutomationTool/Scripts/LaunchMultiServer.Automation.cs
+++ b/Engine/Source/Programs/AutomationTool/Scripts/LaunchMultiServer.Automation.cs
@@ -138,6 +138,7 @@ namespace AutomationTool
 
 			try
 			{
+				int ServerNum = 0;
 				foreach (var ServerInfo in ServerInfos.Values)
 				{
 					Logger.LogInformation("Starting MultiServer instance of type {0} for LocalId {1}, beacon port {2}, and game port {3}", LaunchType, ServerInfo.Definition.LocalId, ServerInfo.Definition.ListenPort, ServerInfo.GameListenPort);
@@ -153,11 +154,9 @@ namespace AutomationTool
 					{
 						Args += String.Format("127.0.0.1:{0} ", ServerInfo.GameListenPort) + (Params.Cook ? "" : "-game ");
 						Args += Params.ClientCommandline + " ";
-						if (!Params.ClientCommandline.Contains("-res"))
-						{
-							// Small window by default
-							Args += "-res=1280x720 ";
-						}
+						
+						// Save a separate window position for each client
+						Args += "-windowed -SaveWinPos=" + (ServerNum + 1) + " ";
 					}
 
 					if (ParseParam("notimeouts"))
@@ -174,6 +173,8 @@ namespace AutomationTool
 						// Essentially forces the -NoKill command-line option behavior for these.
 						ProcessManager.RemoveProcess(NewProcess);
 					}
+
+					ServerNum++;
 				}
 			}
 			catch
