Create a Gauntlet test to cover the CookByTheBook process

1. The TC1 https://testrail.it.epicgames.com/index.php?/cases/view/35379108&group_by=cases:section_id&group_id=2058461&group_order=asc&display_deleted_cases=0
2. Zenstore is disabled
3. Log messages about success and cooking errors are tracked
4. Created for the Lyra project

#rnx

diff --git a/Engine/Source/Programs/AutomationTool/Gauntlet/Unreal/Automation/UE.CookByTheBook.cs b/Engine/Source/Programs/AutomationTool/Gauntlet/Unreal/Automation/UE.CookByTheBook.cs
new file mode 100644
--- /dev/null
+++ b/Engine/Source/Programs/AutomationTool/Gauntlet/Unreal/Automation/UE.CookByTheBook.cs
@@ -0,0 +1,428 @@
+// Copyright Epic Games, Inc. All Rights Reserved.
+
+using Gauntlet;
+using System;
+using System.Collections.Generic;
+using System.IO;
+using System.Linq;
+using System.Text;
+using UnrealBuildTool;
+using static UE.CookOnTheFly;
+using Log = Gauntlet.Log;
+
+namespace UE
+{
+	public class CookByTheBook : UnrealTestNode<UnrealTestConfiguration>
+	{
+		private const string CookingCompleteKey = "The cooking process is complete";
+		private const string CookingStartedPattern = "CookSettings for Memory";
+		private const string CookingInProgressPattern = @"Cooked\s+packages\s+\d+\s+Packages\s+Remain\s+\d+\s+Total\s+\d+";
+		private const string CookingCompletePattern = "Cook by the book total time in tick";
+		private const string ClientStartedPattern = "Starting Game";
+
+		private UnrealLogStreamParser EditorLogParser;
+		private UnrealLogStreamParser ClientLogParser;
+		private Checker Checker;
+		private string[] CookErrors;
+
+		public CookByTheBook(UnrealTestContext InContext) : base(InContext)
+		{
+			EditorLogParser = new UnrealLogStreamParser();
+			ClientLogParser = new UnrealLogStreamParser();
+			CookErrors = Array.Empty<string>();
+
+			Checker = new Checker();
+			Checker.Validate("The cooking process is started", IsCookingStarted);
+			Checker.Validate("The cooking process is in progress", IsCookingInProgress);
+			Checker.Validate(CookingCompleteKey, IsCookingComplete);
+			Checker.Validate("The cooking content is placed correctly", IsCookedContentPlacedCorrectly);
+			Checker.Validate("The client is started", IsClientStarted);
+
+			CleanContentDir();
+		}
+
+		public override UnrealTestConfiguration GetConfiguration()
+		{
+			UnrealTestConfiguration Config = base.GetConfiguration();
+
+			Config.AllRolesExit = true;
+
+			SetEditorRole(Config);
+			SetClientRole(Config);
+
+			return Config;
+		}
+
+		public override void TickTest()
+		{
+			base.TickTest();
+
+			ReadEditorLogStream();
+			ReadClientLogStream();
+
+			CookErrors = FindCookErrors();
+
+			if (CookErrors.Any())
+			{
+				Log.Error("There are cook errors. Ending test");
+				CompleteTest(TestResult.Failed);
+			}
+
+			if (Checker.CheckAll())
+			{
+				Log.Info("The cooking process was successful, the cooking content is correct, the client was launched. Ending test");
+				CompleteTest(TestResult.Passed);
+			}
+		}
+
+		protected override UnrealProcessResult GetExitCodeAndReason(StopReason InReason, UnrealLog InLog, UnrealRoleArtifacts InArtifacts,
+			out string ExitReason, out int ExitCode)
+		{
+			if (InArtifacts.SessionRole.RoleType == UnrealTargetRole.Editor)
+			{
+				InLog.EngineInitializedPattern = CookingStartedPattern;
+			}
+
+			string ErrorDescription = GetExitReason();
+
+			if (!string.IsNullOrEmpty(ErrorDescription))
+			{
+				ExitCode = -1;
+				ExitReason = ErrorDescription;
+				return UnrealProcessResult.TestFailure;
+			}
+
+			return base.GetExitCodeAndReason(InReason, InLog, InArtifacts, out ExitReason, out ExitCode);
+		}
+
+		private void SetEditorRole(UnrealTestConfiguration Config)
+		{
+			UnrealTestRole EditorRole = Config.RequireRole(UnrealTargetRole.Editor);
+			string TargetPlatformName = GetTargetPlatformName();
+
+			EditorRole.CommandLine += $@" -run=cook -log -targetplatform={TargetPlatformName} -ini:Game:[/Script/UnrealEd.ProjectPackagingSettings]:bUseZenStore=false";
+		}
+
+		private static void SetClientRole(UnrealTestConfiguration Config)
+		{
+			UnrealTestRole ClientRole = Config.RequireRole(UnrealTargetRole.Client);
+			ClientRole.DeferredLaunch = true;
+
+			ClientRole.CommandLine += " -log";
+		}
+
+		private void LaunchDeferredClient()
+		{
+			UnrealSessionInstance SessionInstance = TestSessions.GetValueOrDefault(Name)?.SessionInstance;
+			UnrealSessionInstance.RoleInstance DeferredRole = SessionInstance?.DeferredRoles?.FirstOrDefault();
+
+			if (DeferredRole is null || !SessionInstance.LaunchDeferredRole(DeferredRole.Role))
+			{
+				Log.Error("The client hasn't been started");
+				CompleteTest(TestResult.Failed);
+			}
+		}
+
+		private void ReadEditorLogStream()
+		{
+			if (TestInstance.EditorApp is null)
+			{
+				return;
+			}
+
+			EditorLogParser.ReadStream(TestInstance.EditorApp.GetLogBufferReader());
+		}
+
+		private void ReadClientLogStream()
+		{
+			IAppInstance ClientApp = TestInstance.ClientApps.FirstOrDefault();
+
+			if (ClientApp is null)
+			{
+				return;
+			}
+
+			ClientLogParser.ReadStream(ClientApp.GetLogBufferReader());
+		}
+
+		private bool IsCookingStarted()
+		{
+			return EditorLogParser.GetLogLinesContaining(CookingStartedPattern).Any();
+		}
+
+		private bool IsCookingInProgress()
+		{
+			return EditorLogParser.GetLogLinesMatchingPattern(CookingInProgressPattern).Any();
+		}
+
+		private bool IsCookingComplete()
+		{
+			bool IsCookingComplete = EditorLogParser.GetLogLinesContaining(CookingCompletePattern).Any();
+
+			if (IsCookingComplete)
+			{
+				LaunchDeferredClient();
+			}
+
+			return IsCookingComplete;
+		}
+
+		private bool IsCookedContentPlacedCorrectly()
+		{
+			if (!Checker.HasValidated(CookingCompleteKey))
+			{
+				return false;
+			}
+
+			ContentFolder[] ActualCookedContent = GetActualCookedContent();
+			ContentFolder[] ExpectedCookedContent = GetExpectedCookedContent();
+			bool IsCookedContentPlacedCorrectly = ExpectedCookedContent.SequenceEqual(ActualCookedContent, 
+				new ContentFolderEqualityComparer());
+
+			LogCookedContent(ActualCookedContent);
+
+			if (!IsCookedContentPlacedCorrectly)
+			{
+				Log.Error("The cooked content is not valid");
+				CompleteTest(TestResult.Failed);
+			}
+
+			return IsCookedContentPlacedCorrectly;
+		}
+
+		private ContentFolder[] GetActualCookedContent()
+		{
+			string SavedCookedPlatformPath = GetSavedCookedPlatformPath();
+			ContentFolder[] CookedContent = GetSubDirNames(SavedCookedPlatformPath);
+
+			foreach (ContentFolder Folder in CookedContent)
+			{
+				string FolderPath = Path.Combine(SavedCookedPlatformPath, Folder.Name);
+				Folder.SubFolders = GetSubDirNames(FolderPath);
+				Folder.Files = GetDirFileNames(FolderPath);
+			}
+
+			return CookedContent;
+		}
+
+		private ContentFolder[] GetExpectedCookedContent()
+		{
+			List<string> EngineFiles = ["GlobalShaderCache-PCD3D_SM5.bin", "GlobalShaderCache-VULKAN_SM6.bin"];
+
+			if (GetTargetPlatform() == UnrealTargetPlatform.Win64)
+			{
+				EngineFiles.Add("GlobalShaderCache-PCD3D_SM6.bin");
+			}
+
+			ContentFolder[] CookedContent =
+			[
+				new ContentFolder("Engine")
+				{
+					SubFolders = new List<ContentFolder>
+					{
+						new("Content"),
+						new("Plugins")
+					},
+					Files = EngineFiles.ToArray()
+				},
+				new ContentFolder(Context.Options.Project)
+				{
+					SubFolders = new List<ContentFolder>
+					{
+						new("Content"),
+						new("Metadata"),
+						new("Plugins")
+					},
+					Files = new List<string>{ "AssetRegistry.bin" }
+				}
+			];
+
+			return CookedContent;
+		}
+
+		private void LogCookedContent(ContentFolder[] ContentFolders)
+		{
+			foreach (ContentFolder Folder in ContentFolders)
+			{
+				Log.Info(Folder.Name);
+
+				foreach (ContentFolder SubFolder in Folder.SubFolders)
+				{
+					Log.Info($"- {SubFolder.Name}");
+				}
+
+				foreach (string File in Folder.Files)
+				{
+					Log.Info($"- {File}");
+				}
+			}
+		}
+
+		private bool IsClientStarted()
+		{
+			return ClientLogParser.GetLogLinesContaining(ClientStartedPattern).Any();
+		}
+
+		private string[] FindCookErrors()
+		{
+			string[] CookErrors = EditorLogParser
+				.GetEventsFromChannels(["LogCook"])
+				.Where(E => E.Level == UnrealLog.LogLevel.Error)
+				.Select(E => E.Message)
+				.ToArray();
+
+			return CookErrors;
+		}
+
+		private void CleanContentDir()
+		{
+			string SavedCookedPath = GetSavedCookedPath();
+
+			if (!Directory.Exists(SavedCookedPath))
+			{
+				return;
+			}
+
+			Directory.Delete(SavedCookedPath, true);
+		}
+
+		private string GetExitReason()
+		{
+			StringBuilder ErrorBuilder = new StringBuilder();
+
+			if (CookErrors.Any())
+			{
+				ErrorBuilder.Append($"There are cook errors: {string.Join("\n", CookErrors)}");
+			}
+
+			string[] FailedEvents = GetFailedEvents();
+
+			if (FailedEvents.Any())
+			{
+				ErrorBuilder.Append($"The expected events in the logs have not occurred: {string.Join("\n", FailedEvents)}");
+			}
+
+			return ErrorBuilder.ToString();
+		}
+
+		private string[] GetFailedEvents()
+		{
+			string[] FailedEvents = Checker.GetValidations().Keys.ToArray();
+
+			return FailedEvents;
+		}
+
+		private string GetSavedCookedPlatformPath()
+		{
+			string SavedCookedPath = GetSavedCookedPath();
+			string TargetPlatformName = GetTargetPlatformName();
+			string SavedCookedPlatformPath = Path.Combine(SavedCookedPath, TargetPlatformName);
+
+			return SavedCookedPlatformPath;
+		}
+
+		private string GetSavedCookedPath()
+		{
+			string NormalizedProjectFilePath = Context.Options.ProjectPath.ToNormalizedPath();
+			string ProjectPath = Path.GetDirectoryName(NormalizedProjectFilePath) ?? string.Empty;
+			string SavedCookedPath = Path.Combine(ProjectPath, "Saved", "Cooked");
+
+			return SavedCookedPath;
+		}
+
+		private UnrealTargetPlatform GetTargetPlatform()
+		{
+			UnrealTargetPlatform TargetPlatform = Context.GetRoleContext(UnrealTargetRole.Editor).Platform;
+
+			return TargetPlatform;
+		}
+
+		private string GetTargetPlatformName()
+		{
+			UnrealTargetPlatform TargetPlatform = GetTargetPlatform();
+			string TargetPlatformName = UnrealHelpers.GetPlatformName(TargetPlatform, UnrealTargetRole.Editor, false);
+
+			return TargetPlatformName;
+		}
+
+		private static ContentFolder[] GetSubDirNames(string Dir)
+		{
+			ContentFolder[] SubDirNames = Array.Empty<ContentFolder>();
+
+			try
+			{
+				SubDirNames = Directory.GetDirectories(Dir).Select(D => new ContentFolder(new DirectoryInfo(D).Name)).ToArray();
+			}
+			catch(Exception E)
+			{
+				Log.Warning(E.Message);
+			}
+
+			return SubDirNames;
+		}
+
+		private static string[] GetDirFileNames(string Dir)
+		{
+			string[] DirFileNames = Array.Empty<string>();
+
+			try
+			{
+				DirFileNames = Directory.GetFiles(Dir).Select(Path.GetFileName).ToArray();
+			}
+			catch (Exception E)
+			{
+				Log.Warning(E.Message);
+			}
+
+			return DirFileNames;
+		}
+
+		private void CompleteTest(TestResult TestResult)
+		{
+			MarkTestComplete();
+			SetUnrealTestResult(TestResult);
+		}
+
+		private class ContentFolder
+		{
+			public string Name { get; }
+			public IList<ContentFolder> SubFolders { get; set; }
+			public IList<string> Files { get; set; }
+
+			public ContentFolder(string InName)
+			{
+				Name = InName;
+			}
+		}
+
+		private sealed class ContentFolderEqualityComparer : IEqualityComparer<ContentFolder>
+		{
+			public bool Equals(ContentFolder x, ContentFolder y)
+			{
+				if (ReferenceEquals(x, y))
+				{
+					return true;
+				}
+
+				if (ReferenceEquals(x, null) || ReferenceEquals(y, null))
+				{
+					return false;
+				}
+
+				if (x.GetType() != y.GetType())
+				{
+					return false;
+				}
+
+				return x.Name.Equals(y.Name, StringComparison.InvariantCultureIgnoreCase) &&
+				       !x.Files.Except(y.Files, StringComparer.InvariantCultureIgnoreCase).Any() &&
+				       !x.SubFolders.Select(f => f.Name).Except(y.SubFolders.Select(f => f.Name), StringComparer.InvariantCultureIgnoreCase).Any();
+			}
+
+			public int GetHashCode(ContentFolder obj)
+			{
+				return 0; // The same hash for two objects indicates a hash collision, the Equals will be called 
+			}
+		}
+	}
+}
\ No newline at end of file
diff --git a/Engine/Source/Programs/AutomationTool/Gauntlet/Unreal/Automation/UE.CookOnTheFly.cs b/Engine/Source/Programs/AutomationTool/Gauntlet/Unreal/Automation/UE.CookOnTheFly.cs
--- a/Engine/Source/Programs/AutomationTool/Gauntlet/Unreal/Automation/UE.CookOnTheFly.cs
+++ b/Engine/Source/Programs/AutomationTool/Gauntlet/Unreal/Automation/UE.CookOnTheFly.cs
@@ -75,6 +75,24 @@ namespace UE
 
 				return Aggregate;
 			}
+
+			/// <summary>
+			/// Check the specific callback, return true if it passed or is missing.
+			/// </summary>
+			/// <returns></returns>
+			public bool HasValidated(string ToFound)
+			{
+				return !ToValidate.ContainsKey(ToFound);
+			}
+
+			/// <summary>
+			/// Provides all callbacks
+			/// </summary>
+			/// <returns></returns>
+			public Dictionary<string, Func<bool>> GetValidations()
+			{
+				return ToValidate;
+			}
 		}
 
 		public CookOnTheFly(Gauntlet.UnrealTestContext InContext)
