Horde: Add an option to dedupe artifacts against the previous instance.

#rnx

diff --git a/Engine/Source/Programs/AutomationTool/BuildGraph/Tasks/CreateArtifactTask.cs b/Engine/Source/Programs/AutomationTool/BuildGraph/Tasks/CreateArtifactTask.cs
--- a/Engine/Source/Programs/AutomationTool/BuildGraph/Tasks/CreateArtifactTask.cs
+++ b/Engine/Source/Programs/AutomationTool/BuildGraph/Tasks/CreateArtifactTask.cs
@@ -86,6 +86,12 @@ namespace AutomationTool.Tasks
 		/// </summary>
 		[TaskParameter(Optional = true)]
 		public bool Symbols { get; set; }
+
+		/// <summary>
+		/// Whether to reuse blobs from the previous artifact
+		/// </summary>
+		[TaskParameter(Optional = true)]
+		public bool Dedupe { get; set; }
 	}
 
 	/// <summary>
@@ -196,7 +202,27 @@ namespace AutomationTool.Tasks
 			IHashedBlobRef<DirectoryNode> rootRef;
 			await using (IBlobWriter blobWriter = artifact.CreateBlobWriter())
 			{
-				rootRef = await blobWriter.WriteFilesAsync(baseDir, files);
+				if (_parameters.Dedupe)
+				{
+					await using DedupeBlobWriter dedupeWriter = new DedupeBlobWriter(blobWriter);
+
+					IArtifact? prevArtifact = await hordeClient.Artifacts.FindAsync(streamId, maxCommitId: commitId, name: name, type: type).FirstOrDefaultAsync(x => x.CommitId != commitId);
+					if (prevArtifact != null && prevArtifact.NamespaceId == artifact.NamespaceId)
+					{
+						Logger.LogInformation("Caching blobs from previous artifact ({ArtifactId})", prevArtifact.Id);
+						await dedupeWriter.AddToCacheAsync(prevArtifact.Content);
+					}
+
+					rootRef = await dedupeWriter.WriteFilesAsync(baseDir, files);
+
+					StorageStats stats = dedupeWriter.GetStats();
+					Logger.LogInformation("Stats:");
+					stats.Print(Logger);
+				}
+				else
+				{
+					rootRef = await blobWriter.WriteFilesAsync(baseDir, files);
+				}
 			}
 
 			// Tag any uploaded symbols
@@ -239,7 +265,7 @@ namespace AutomationTool.Tasks
 			}
 
 			IBlobRef<DirectoryNode> directoryRef = rootDir;
-			for (int idx = 0;; idx++)
+			for (int idx = 0; ; idx++)
 			{
 				DirectoryNode directory = await directoryRef.ReadBlobAsync(cancellationToken);
 				if (idx + 1 < fragments.Length)
diff --git a/Engine/Source/Programs/Shared/EpicGames.Horde/Storage/DedupeBlobWriter.cs b/Engine/Source/Programs/Shared/EpicGames.Horde/Storage/DedupeBlobWriter.cs
--- a/Engine/Source/Programs/Shared/EpicGames.Horde/Storage/DedupeBlobWriter.cs
+++ b/Engine/Source/Programs/Shared/EpicGames.Horde/Storage/DedupeBlobWriter.cs
@@ -216,12 +216,16 @@ namespace EpicGames.Horde.Storage
 		/// <param name="dedupeWriter">Dedupe writer to operate on</param>
 		/// <param name="directoryNodeRef">Reference to the directory to add</param>
 		/// <param name="cancellationToken">Cancellation token for the operation</param>
-		public static async Task AddToCacheAsync(this DedupeBlobWriter dedupeWriter, IHashedBlobRef<DirectoryNode> directoryNodeRef, CancellationToken cancellationToken = default)
+		public static async Task AddToCacheAsync(this DedupeBlobWriter dedupeWriter, IBlobRef<DirectoryNode> directoryNodeRef, CancellationToken cancellationToken = default)
 		{
-			using BlobData blobData = await directoryNodeRef.ReadBlobDataAsync(cancellationToken);
-			dedupeWriter.AddToCache(blobData.Type, directoryNodeRef);
+			DirectoryNode directoryNode;
+			using (BlobData blobData = await directoryNodeRef.ReadBlobDataAsync(cancellationToken))
+			{
+				IHashedBlobRef hashedBlobRef = HashedBlobRef.Create(IoHash.Compute(blobData.Data.Span), directoryNodeRef);
+				dedupeWriter.AddToCache(blobData.Type, hashedBlobRef);
+				directoryNode = BlobSerializer.Deserialize<DirectoryNode>(blobData);
+			}
 
-			DirectoryNode directoryNode = BlobSerializer.Deserialize<DirectoryNode>(blobData);
 			foreach (DirectoryEntry directoryEntry in directoryNode.Directories)
 			{
 				await AddToCacheAsync(dedupeWriter, directoryEntry.Handle, cancellationToken);
