Horde: Update artifact test job to enable dedupe.

#rnx

diff --git a/Engine/Source/Programs/AutomationTool/BuildGraph/Tasks/RandomDataTask.cs b/Engine/Source/Programs/AutomationTool/BuildGraph/Tasks/RandomDataTask.cs
--- a/Engine/Source/Programs/AutomationTool/BuildGraph/Tasks/RandomDataTask.cs
+++ b/Engine/Source/Programs/AutomationTool/BuildGraph/Tasks/RandomDataTask.cs
@@ -15,6 +15,12 @@ namespace AutomationTool.Tasks
 	/// </summary>
 	public class RandomDataTaskParameters
 	{
+		/// <summary>
+		/// Seed for the data generation
+		/// </summary>
+		[TaskParameter(Optional = true)]
+		public int Seed { get; set; } = -1;
+
 		/// <summary>
 		/// The size of each file.
 		/// </summary>
@@ -74,12 +80,23 @@ namespace AutomationTool.Tasks
 			DirectoryReference outputDir = ResolveDirectory(_parameters.OutputDir);
 			DirectoryReference.CreateDirectory(outputDir);
 
+			Random random;
+			if (_parameters.Seed >= 0)
+			{
+				random = new Random(_parameters.Seed);
+			}
+			else
+			{
+				random = new Random();
+			}
+
 			byte[] buffer = Array.Empty<byte>();
 			for (int idx = 0; idx < _parameters.Count; idx++)
 			{
 				if (idx == 0 || _parameters.Different)
 				{
-					buffer = RandomNumberGenerator.GetBytes(_parameters.Size);
+					buffer = new byte[_parameters.Size];
+					random.NextBytes(buffer);
 				}
 
 				FileReference file = FileReference.Combine(outputDir, $"test-{_parameters.Size}-{idx}.dat");
