Additional logging for deduping artifacts.

#rnx

diff --git a/Engine/Source/Programs/AutomationTool/BuildGraph/Tasks/CreateArtifactTask.cs b/Engine/Source/Programs/AutomationTool/BuildGraph/Tasks/CreateArtifactTask.cs
--- a/Engine/Source/Programs/AutomationTool/BuildGraph/Tasks/CreateArtifactTask.cs
+++ b/Engine/Source/Programs/AutomationTool/BuildGraph/Tasks/CreateArtifactTask.cs
@@ -207,7 +207,15 @@ namespace AutomationTool.Tasks
 					await using DedupeBlobWriter dedupeWriter = new DedupeBlobWriter(blobWriter);
 
 					IArtifact? prevArtifact = await hordeClient.Artifacts.FindAsync(streamId, maxCommitId: commitId, name: name, type: type).FirstOrDefaultAsync(x => x.CommitId != commitId);
-					if (prevArtifact != null && prevArtifact.NamespaceId == artifact.NamespaceId)
+					if (prevArtifact == null)
+					{
+						Logger.LogInformation("Could not find a previous artifact; dedupe will be disabled");
+					}
+					else if (prevArtifact.NamespaceId != artifact.NamespaceId)
+					{
+						Logger.LogInformation("Previous artifact {ArtifactId} is from namespace {PrevNamespaceId}, not {NamespaceId}. Will disable dedupe.", artifact.Id, prevArtifact.NamespaceId, artifact.NamespaceId);
+					}
+					else
 					{
 						Logger.LogInformation("Caching blobs from previous artifact ({ArtifactId})", prevArtifact.Id);
 						await dedupeWriter.AddToCacheAsync(prevArtifact.Content);
