Live Link Hub: Update splash screen.

#rnx

diff --git a/Engine/Source/Programs/AutomationTool/LiveLinkHub/LiveLinkHub.Automation.cs b/Engine/Source/Programs/AutomationTool/LiveLinkHub/LiveLinkHub.Automation.cs
--- a/Engine/Source/Programs/AutomationTool/LiveLinkHub/LiveLinkHub.Automation.cs
+++ b/Engine/Source/Programs/AutomationTool/LiveLinkHub/LiveLinkHub.Automation.cs
@@ -77,6 +77,15 @@ namespace AutomationTool
 				new StagedFileReference($"{Context.ProjectName}/Intermediate/TargetInfo.json"),
 				new FileReference($"Engine/Source/Programs/{Context.ProjectName}/Intermediate/TargetInfo.json"));
 
+			// Copy custom splash screen
+			SC.FilesToStage.NonUFSFiles.Add(
+				new StagedFileReference($"{Context.ProjectName}/Content/Splash/EdSplash.png"),
+				new FileReference($"Engine/Source/Programs/{Context.ProjectName}/Content/Splash/EdSplash.png"));
+
+			SC.FilesToStage.NonUFSFiles.Add(
+				new StagedFileReference($"{Context.ProjectName}/Content/Splash/EdSplash.uasset"),
+				new FileReference($"Engine/Source/Programs/{Context.ProjectName}/Content/Splash/EdSplash.uasset"));
+
 			// Remove asset registry entry
 			SC.FilesToStage.UFSFiles.Remove(new StagedFileReference($"{Context.ProjectName}/EditorClientAssetRegistry.bin"));
 
