Add -WaitForAttachNoBreak and -WaitForDebuggerNoBreak command-line options to UE that behave like -WaitForAttach and -WaitForDebugger but don't automatically break after a debugger attaches.
LaunchMultiServer UAT command can now automatically attach a Visual Studio debugger to launched processes with the -attachtoservers and/or -attachtoclients options.

diff --git a/Engine/Source/Programs/AutomationTool/Scripts/AutomationScripts.Automation.csproj b/Engine/Source/Programs/AutomationTool/Scripts/AutomationScripts.Automation.csproj
--- a/Engine/Source/Programs/AutomationTool/Scripts/AutomationScripts.Automation.csproj
+++ b/Engine/Source/Programs/AutomationTool/Scripts/AutomationScripts.Automation.csproj
@@ -56,6 +56,7 @@
   </ItemGroup>
   <ItemGroup>
     <PackageReference Include="Microsoft.CSharp" Version="4.7.0" />
+    <PackageReference Include="Microsoft.VisualStudio.OLE.Interop" Version="17.11.40262" />
     <PackageReference Include="System.Data.DataSetExtensions" Version="4.5.0" />
     <PackageReference Include="System.Net.Http" Version="4.3.4" />
     <PackageReference Include="System.Drawing.Common" Version="4.7.3" />
diff --git a/Engine/Source/Programs/AutomationTool/Scripts/LaunchMultiServer.Automation.cs b/Engine/Source/Programs/AutomationTool/Scripts/LaunchMultiServer.Automation.cs
--- a/Engine/Source/Programs/AutomationTool/Scripts/LaunchMultiServer.Automation.cs
+++ b/Engine/Source/Programs/AutomationTool/Scripts/LaunchMultiServer.Automation.cs
@@ -6,11 +6,161 @@ using UnrealBuildBase;
 using UnrealBuildTool;
 using AutomationScripts;
 using EpicGames.Core;
+using System.Linq;
 using Microsoft.Extensions.Logging;
 using System.Collections.Generic;
+using System.Runtime.InteropServices;
+
+using Microsoft.VisualStudio.OLE.Interop;
+using System.Runtime.Versioning;
 
 namespace AutomationTool
 {
+	class ProcessDebugger
+	{
+		[SupportedOSPlatform("windows")]
+		[DllImport("ole32.dll")]
+		public static extern int CreateBindCtx(int reserved, out IBindCtx ppbc);
+
+		[SupportedOSPlatform("windows")]
+		[DllImport("ole32.dll")]
+		public static extern int GetRunningObjectTable(int reserved, out IRunningObjectTable prot);
+
+		public static void DebugProcess(int ProcessID)
+		{
+			if (OperatingSystem.IsWindows())
+			{
+				EnvDTE._DTE visualStudioInstance = GetVisualStudioInstance();
+
+				if (visualStudioInstance != null)
+				{
+					AttachVisualStudioToPID(visualStudioInstance, ProcessID);
+				}
+				else
+				{
+					Log.Logger.LogInformation($"Failed to find a Visual Studio Instance.");
+				}
+			}
+		}
+
+		[SupportedOSPlatform("windows")]
+		private static EnvDTE._DTE GetVisualStudioInstance()
+		{
+			EnvDTE._DTE visualStudioInstance = null;
+
+			uint numFetched = 0;
+			IRunningObjectTable runningObjectTable = null;
+			IEnumMoniker monikerEnumerator;
+			IMoniker[] monikers = new IMoniker[1];
+
+			GetRunningObjectTable(0, out runningObjectTable);
+			runningObjectTable.EnumRunning(out monikerEnumerator);
+			monikerEnumerator.Reset();
+
+			while (monikerEnumerator.Next(1, monikers, out numFetched) == 0)
+			{
+				IBindCtx ctx;
+				CreateBindCtx(0, out ctx);
+
+				string runningObjectName;
+				monikers[0].GetDisplayName(ctx, null, out runningObjectName);
+
+				object runningObjectVal;
+				runningObjectTable.GetObject(monikers[0], out runningObjectVal);
+
+				if (!(runningObjectVal is EnvDTE._DTE) || !runningObjectName.StartsWith("!VisualStudio"))
+				{
+					continue;
+				}
+
+				return (EnvDTE._DTE)runningObjectVal;
+			}
+
+			return visualStudioInstance;
+		}
+
+		[SupportedOSPlatform("windows")]
+		private static void AttachVisualStudioToPID(EnvDTE._DTE visualStudioInstance, int processID)
+		{
+			int retryCount = 0;
+			while (true)
+			{
+				try
+				{
+					var processToAttachTo = visualStudioInstance.Debugger.LocalProcesses.Cast<EnvDTE.Process>().FirstOrDefault(process => process.ProcessID == processID);
+
+					if (processToAttachTo == null)
+					{
+						Log.Logger.LogInformation("Failed to find running Process matching provided Process Name {0}", processID);
+						continue;
+					}
+					else
+					{
+						processToAttachTo.Attach();
+					}
+
+					break;
+				}
+				catch (COMException e)
+				{
+					if ((uint)e.ErrorCode == 0x8001010a || (uint)e.ErrorCode == 0x80010001)
+					{
+						if (++retryCount < 15)
+						{
+							Log.Logger.LogInformation("Attach Debugger - Got RPC Retry Later exception. Will try again ");
+							System.Threading.Thread.Sleep(20);
+							continue;
+						}
+					}
+					Log.Logger.LogInformation("Failed to attach debugger. COMException was thrown: " + e.ToString());
+					break;
+				}
+				catch (Exception e)
+				{
+					Log.Logger.LogInformation("Failed to attach debugger. Exception was thrown: " + e.ToString());
+					break;
+				}
+			}
+		}
+
+		[SupportedOSPlatform("windows")]
+		private static bool IsDebuggerAttached(EnvDTE._DTE VisualStudio, string processID)
+		{
+			bool DebuggerAttached = false;
+
+			if (VisualStudio.Debugger.DebuggedProcesses.Count != 0)
+			{
+				foreach (EnvDTE.Process debuggedProcess in VisualStudio.Debugger.DebuggedProcesses)
+				{
+					if (debuggedProcess.Name.Contains(processID))
+					{
+						DebuggerAttached = true;
+						break;
+					}
+				}
+			}
+			return DebuggerAttached;
+		}
+
+		[SupportedOSPlatform("windows")]
+		private static bool IsDebuggerAttachedToPID(EnvDTE._DTE VisualStudio, int processID)
+		{
+			bool DebuggerAttached = false;
+
+			if (VisualStudio.Debugger.DebuggedProcesses.Count != 0)
+			{
+				foreach (EnvDTE.Process debuggedProcess in VisualStudio.Debugger.DebuggedProcesses)
+				{
+					if (debuggedProcess.ProcessID == processID)
+					{
+						DebuggerAttached = true;
+						break;
+					}
+				}
+			}
+			return DebuggerAttached;
+		}
+	}
 
 	// Mirrors FMultiServerDefinition from the MultiServerReplication plugin
 	class MultiServerDefinition
@@ -163,6 +313,12 @@ namespace AutomationTool
 					{
 						Args += "-notimeouts ";
 					}
+
+					if (IsAttachingDebugger(LaunchType))
+					{
+						Args += "-WaitForDebuggerNoBreak ";
+					}
+
 					Args += AdditionalParameters;
 
 					var NewProcess = Run(App, Args, null, ERunOptions.Default | ERunOptions.NoWaitForExit | ERunOptions.NoStdOutRedirect);
@@ -174,6 +330,11 @@ namespace AutomationTool
 						ProcessManager.RemoveProcess(NewProcess);
 					}
 
+					if (IsAttachingDebugger(LaunchType))
+					{ 
+						ProcessDebugger.DebugProcess(NewProcess.ProcessObject.Id);
+					}
+
 					ServerNum++;
 				}
 			}
@@ -187,6 +348,12 @@ namespace AutomationTool
 			}
 		}
 
+		private bool IsAttachingDebugger(ProcessLaunchType LaunchType)
+		{
+			return (ParseParam("attachtoservers") && LaunchType == ProcessLaunchType.Server) ||
+				(ParseParam("attachtoclients") && LaunchType == ProcessLaunchType.DirectClient);
+		}
+
 		private FileReference ProjectFullPath;
 		public virtual FileReference ProjectPath
 		{
diff --git a/Engine/Source/Runtime/Launch/Private/Launch.cpp b/Engine/Source/Runtime/Launch/Private/Launch.cpp
--- a/Engine/Source/Runtime/Launch/Private/Launch.cpp
+++ b/Engine/Source/Runtime/Launch/Private/Launch.cpp
@@ -102,13 +102,21 @@ int32 GuardedMain( const TCHAR* CmdLine )
 #if !(UE_BUILD_SHIPPING)
 
 	// If "-waitforattach" or "-WaitForDebugger" was specified, halt startup and wait for a debugger to attach before continuing
-	if (FParse::Param(CmdLine, TEXT("waitforattach")) || FParse::Param(CmdLine, TEXT("WaitForDebugger")))
+
+	const bool bWaitForDebuggerAndBreak = FParse::Param(CmdLine, TEXT("waitforattach")) || FParse::Param(CmdLine, TEXT("WaitForDebugger")); 
+	const bool bWaitForDebugger = bWaitForDebuggerAndBreak || FParse::Param(CmdLine, TEXT("WaitForAttachNoBreak")) || FParse::Param(CmdLine, TEXT("WaitForDebuggerNoBreak"));
+	
+	if (bWaitForDebugger)
 	{
 		while (!FPlatformMisc::IsDebuggerPresent())
 		{
 			FPlatformProcess::Sleep(0.1f);
 		}
-		UE_DEBUG_BREAK();
+
+		if (bWaitForDebuggerAndBreak)
+		{
+			UE_DEBUG_BREAK();
+		}
 	}
 
 #endif
