Fix CloseStream being call during AppendLine

#rnx

diff --git a/Engine/Source/Programs/AutomationTool/Gauntlet/Framework/Gauntlet.Utils.cs b/Engine/Source/Programs/AutomationTool/Gauntlet/Framework/Gauntlet.Utils.cs
--- a/Engine/Source/Programs/AutomationTool/Gauntlet/Framework/Gauntlet.Utils.cs
+++ b/Engine/Source/Programs/AutomationTool/Gauntlet/Framework/Gauntlet.Utils.cs
@@ -2922,7 +2922,10 @@ namespace Gauntlet
 						bHasExited = LocalProcess.HasExited;
 						if (bHasExited)
 						{
-							Logger?.CloseStream();
+							lock (OutputSyncObject)
+							{
+								Logger?.CloseStream();
+							}
 							ExitCode = LocalProcess.ExitCode;
 						}
 					}
@@ -3051,7 +3054,10 @@ namespace Gauntlet
 			{
 				Log.Info("Outputs did not close in time after process {ProcessName} exited.", LocalProcess.ProcessName);
 			}
-			Logger?.CloseStream();
+			lock (OutputSyncObject)
+			{
+				Logger?.CloseStream();
+			}
 		}
 	}
 }
