Improved ParamHelp documentation for RunUnreal to clarify that passing arguments to cooked editor, editorgame, and editorserver roles requires different parameters. I.e., "-editorargs=foo" will only pass 'foo' to the Editor role and not, e.g., to CookedEditor roles.

#rnx

diff --git a/Engine/Source/Programs/AutomationTool/Gauntlet/Unreal/RunUnreal.cs b/Engine/Source/Programs/AutomationTool/Gauntlet/Unreal/RunUnreal.cs
--- a/Engine/Source/Programs/AutomationTool/Gauntlet/Unreal/RunUnreal.cs
+++ b/Engine/Source/Programs/AutomationTool/Gauntlet/Unreal/RunUnreal.cs
@@ -26,7 +26,7 @@ namespace Gauntlet
 	/// </summary>
 	[Help("Run Unreal tests using Gauntlet")]
 	[ParamHelp("Tests", "List of gauntlet tests to run", Required = true, MultiSelectSeparator = ",")]
-	[ParamHelp("ExecCmds", "List commands to execute", MultiSelectSeparator = "+")]
+	[ParamHelp("ExecCmds", "List commands to execute. NOTE: To provide ExecCmds to specific roles, use -EditorExecCmds, -ClientExecCmds, etc.", MultiSelectSeparator = "+")]
 	[ParamHelp("Build", "Reference to the build that is being tested")]
 	[ParamHelp("Configuration", "Configuration to perform tests on", Choices = new string[] { "Debug", "DebugGame", "Development", "Test", "Shipping"})]
 	[ParamHelp("Platform", "Platforms to perform tests on and their params")]
@@ -54,10 +54,13 @@ namespace Gauntlet
 	[ParamHelp("Verbose", "Verbose logging", ParamType = typeof(bool))]
 	[ParamHelp("VeryVerbose", "Very Verbose logging", ParamType = typeof(bool))]
 	[ParamHelp("HeartbeatPeriod", "Gauntlet heartbeat period", ParamType = typeof(float))]
-	[ParamHelp("Args", "Extra arguments to pass role(s)")]
+	[ParamHelp("Args", "Extra arguments to pass to all role(s)")]
 	[ParamHelp("ClientArgs", "Extra arguments passed to client role(s)")]
 	[ParamHelp("ServerArgs", "Extra arguments passed to server role(s)")]
 	[ParamHelp("EditorArgs", "Extra arguments passed to editor role(s)")]
+	[ParamHelp("CookedEditorArgs", "Extra arguments passed to cooked editor role(s)")]
+	[ParamHelp("EditorGameArgs", "Extra arguments passed to editor(s) when acting in the game role")]
+	[ParamHelp("EditorServerArgs", "Extra arguments passed to editor(s) when acting in the server role")]
 	[ParamHelp("Namespaces", "Comma-separated list of namespaces to check for tests.", MultiSelectSeparator = ",")]
 	[ParamHelp("AdditionalArgs", "Any additional arguments to pass directly to the Gauntlet command line", IsArgument = true)]
 	public class RunUnreal : BuildCommand
