Add ability to set Proxying Handler Options though Gauntlet.

diff --git a/Engine/Source/Programs/AutomationTool/Gauntlet/Framework/RpcFramework/Gauntlet.GauntletHttpClient.cs b/Engine/Source/Programs/AutomationTool/Gauntlet/Framework/RpcFramework/Gauntlet.GauntletHttpClient.cs
--- a/Engine/Source/Programs/AutomationTool/Gauntlet/Framework/RpcFramework/Gauntlet.GauntletHttpClient.cs
+++ b/Engine/Source/Programs/AutomationTool/Gauntlet/Framework/RpcFramework/Gauntlet.GauntletHttpClient.cs
@@ -7,6 +7,14 @@ namespace Gauntlet
 {
 	public class GauntletHttpClient : HttpClient
 	{
+		public GauntletHttpClient() : this(new HttpClientHandler())
+		{
+		}
+
+		public GauntletHttpClient(HttpMessageHandler Handler) : base(Handler, true)
+		{
+		}
+		
 		public async Task<HttpResponseMessage> GetRequestAsync(string Route)
 		{
 			LogOutgoingRequestURI(Route);
