Lowered RPC ListenerThread connection errors down to a warning

#rnx

diff --git a/Engine/Source/Programs/AutomationTool/Gauntlet/Framework/RpcFramework/Gauntlet.RpcExecutor.cs b/Engine/Source/Programs/AutomationTool/Gauntlet/Framework/RpcFramework/Gauntlet.RpcExecutor.cs
--- a/Engine/Source/Programs/AutomationTool/Gauntlet/Framework/RpcFramework/Gauntlet.RpcExecutor.cs
+++ b/Engine/Source/Programs/AutomationTool/Gauntlet/Framework/RpcFramework/Gauntlet.RpcExecutor.cs
@@ -268,7 +268,9 @@ namespace Gauntlet
 						}
 						else
 						{
-							throw new(Error.Message);
+							Log.SuspendECErrorParsing();
+							Log.Warning("Call to {CallPath} failed. Error: {ErrorMessage}", CallPath.Path, Error.Message);
+							Log.ResumeECErrorParsing();
 						}
 					}
 					return;
@@ -300,7 +302,9 @@ namespace Gauntlet
 				}
 				else
 				{
-					throw new(Error.Message);
+					Log.SuspendECErrorParsing();
+					Log.Warning("Call to {CallPath} failed. Error: {ErrorMessage}", CallPath.Path, Error.Message);
+					Log.ResumeECErrorParsing();
 				}
 			}
 		}
