Follow-up fix to make sure that attempting to remove the read-only attribute doesn't end up triggering another CIS warning or skip attempting to remove the artifact. On PS5, apparently we have permissions to delete files even if we have no permission to edit file attributes.

#rnx

diff --git a/Engine/Source/Programs/AutomationTool/Gauntlet/Unreal/Base/Gauntlet.UnrealSession.cs b/Engine/Source/Programs/AutomationTool/Gauntlet/Unreal/Base/Gauntlet.UnrealSession.cs
--- a/Engine/Source/Programs/AutomationTool/Gauntlet/Unreal/Base/Gauntlet.UnrealSession.cs
+++ b/Engine/Source/Programs/AutomationTool/Gauntlet/Unreal/Base/Gauntlet.UnrealSession.cs
@@ -1055,9 +1055,16 @@ namespace Gauntlet
 								SetAttributesNormal(SubDirectory);
 							}
 						};
+                        try
+                        {
+						    SetAttributesNormal(SourceDirectory);
+                        }
+                        catch (Exception Exception)
+{
+	Log.Info("Could not remove the read-only attribute from {SourceDirectory}.", Exception, SourceDirectory);
+}
 						try
 						{
-							SetAttributesNormal(SourceDirectory);
 							SourceDirectory.Delete(true);
 						}
 						catch (Exception Exception)
