Fix typo in log statement.

diff --git a/Engine/Source/Programs/AutomationTool/Gauntlet/Unreal/Base/Gauntlet.UnrealSession.cs b/Engine/Source/Programs/AutomationTool/Gauntlet/Unreal/Base/Gauntlet.UnrealSession.cs
--- a/Engine/Source/Programs/AutomationTool/Gauntlet/Unreal/Base/Gauntlet.UnrealSession.cs
+++ b/Engine/Source/Programs/AutomationTool/Gauntlet/Unreal/Base/Gauntlet.UnrealSession.cs
@@ -1055,14 +1055,14 @@ namespace Gauntlet
 								SetAttributesNormal(SubDirectory);
 							}
 						};
-                        try
-                        {
-						    SetAttributesNormal(SourceDirectory);
-                        }
-                        catch (Exception Exception)
-{
-	Log.Info("Could not remove the read-only attribute from {SourceDirectory}.", Exception, SourceDirectory);
-}
+						try
+						{
+							SetAttributesNormal(SourceDirectory);
+						}
+						catch (Exception Exception)
+						{
+							Log.Info("Could not remove the read-only attribute from {SourceDirectory}.", SourceDirectory);
+						}
 						try
 						{
 							SourceDirectory.Delete(true);
