Fix another typo for CIS.

#rnx

diff --git a/Engine/Source/Programs/AutomationTool/Gauntlet/Unreal/Base/Gauntlet.UnrealSession.cs b/Engine/Source/Programs/AutomationTool/Gauntlet/Unreal/Base/Gauntlet.UnrealSession.cs
--- a/Engine/Source/Programs/AutomationTool/Gauntlet/Unreal/Base/Gauntlet.UnrealSession.cs
+++ b/Engine/Source/Programs/AutomationTool/Gauntlet/Unreal/Base/Gauntlet.UnrealSession.cs
@@ -1059,7 +1059,7 @@ namespace Gauntlet
 						{
 							SetAttributesNormal(SourceDirectory);
 						}
-						catch (Exception Exception)
+						catch
 						{
 							Log.Info("Could not remove the read-only attribute from {SourceDirectory}.", SourceDirectory);
 						}
