Update CosmeticCapturePostprocess BuildCommand to use an API endpoint instead of direct database access.

diff --git a/Engine/Source/Programs/AutomationTool/AutomationUtils/ImageUtils.cs b/Engine/Source/Programs/AutomationTool/AutomationUtils/ImageUtils.cs
--- a/Engine/Source/Programs/AutomationTool/AutomationUtils/ImageUtils.cs
+++ b/Engine/Source/Programs/AutomationTool/AutomationUtils/ImageUtils.cs
@@ -46,5 +46,43 @@ namespace AutomationTool
 				}
 			}
 		}
+
+		/// <summary>
+		/// Resave an image file as a .png file, with the specified quality and scale.
+		/// </summary>
+		public static void ResaveImageAsPngWithScaleAndQuality(string SrcImagePath, string DstImagePath, float Scale, int PngQuality)
+		{
+			// Find the built-in jpg encoder
+			var Codec = ImageCodecInfo.GetImageDecoders().Where(d => d.FormatID == ImageFormat.Png.Guid).First();
+
+			using (var SrcImage = Image.FromFile(SrcImagePath))
+			{
+				// Create a resized destination bitmap
+				var DstRect = new Rectangle(0, 0, (int)(SrcImage.Width * Scale), (int)(SrcImage.Height * Scale));
+				using (var DstImage = new Bitmap(DstRect.Width, DstRect.Height))
+				{
+					DstImage.SetResolution(SrcImage.HorizontalResolution, SrcImage.VerticalResolution);
+
+					// Draw and scale the original image into the new bitmap.
+					using (var Grfx = Graphics.FromImage(DstImage))
+					using (var ImageAttrs = new ImageAttributes())
+					{
+						Grfx.InterpolationMode = InterpolationMode.HighQualityBicubic;
+						Grfx.CompositingQuality = CompositingQuality.HighQuality;
+						Grfx.CompositingMode = CompositingMode.SourceCopy;
+						Grfx.PixelOffsetMode = PixelOffsetMode.HighQuality;
+						Grfx.SmoothingMode = SmoothingMode.HighQuality;
+
+						ImageAttrs.SetWrapMode(WrapMode.TileFlipXY);
+						Grfx.DrawImage(SrcImage, DstRect, 0, 0, SrcImage.Width, SrcImage.Height, GraphicsUnit.Pixel, ImageAttrs);
+					}
+
+					// Save using the Jpg encoder and a custom quality level.
+					var EncParams = new EncoderParameters(1);
+					EncParams.Param[0] = new EncoderParameter(Encoder.Quality, PngQuality);
+					DstImage.Save(DstImagePath, Codec, EncParams);
+				}
+			}
+		}
 	}
 }
