[RPC] thread lock bHasReceivedTerminationCall

#rnx

diff --git a/Engine/Source/Programs/AutomationTool/Gauntlet/Framework/RpcFramework/Gauntlet.RpcExecutor.cs b/Engine/Source/Programs/AutomationTool/Gauntlet/Framework/RpcFramework/Gauntlet.RpcExecutor.cs
--- a/Engine/Source/Programs/AutomationTool/Gauntlet/Framework/RpcFramework/Gauntlet.RpcExecutor.cs
+++ b/Engine/Source/Programs/AutomationTool/Gauntlet/Framework/RpcFramework/Gauntlet.RpcExecutor.cs
@@ -156,7 +156,26 @@ namespace Gauntlet
 		protected Dictionary<GauntletRpcDelegateRegistryFilter, Delegate> MessageDelegates = new Dictionary<GauntletRpcDelegateRegistryFilter, Delegate>();
 		protected Dictionary<RouteMapping, Delegate> HttpRoutes = new Dictionary<RouteMapping, Delegate>();
 		protected object MessageQueueLock = new object();
-		bool bHasReceivedTerminationCall = false;
+
+		protected object ReceivedTerminationCallLock = new object();
+		private bool _bHasReceivedTerminationCall = false;
+		public bool bHasReceivedTerminationCall
+		{
+			get
+			{
+				lock(ReceivedTerminationCallLock)
+				{
+					return _bHasReceivedTerminationCall;
+				}
+			}
+			private set
+			{
+				lock (ReceivedTerminationCallLock)
+				{
+					_bHasReceivedTerminationCall = value;
+				}
+			}
+		}
 		public void ListenLoopThread()
 		{
 
