Remove the automatic use of -zenstorehost on the commandline (via uecommandline.txt) because it is now handled by the contents of the ue.projectstore file.

diff --git a/Engine/Source/Programs/AutomationTool/Scripts/RunProjectCommand.Automation.cs b/Engine/Source/Programs/AutomationTool/Scripts/RunProjectCommand.Automation.cs
--- a/Engine/Source/Programs/AutomationTool/Scripts/RunProjectCommand.Automation.cs
+++ b/Engine/Source/Programs/AutomationTool/Scripts/RunProjectCommand.Automation.cs
@@ -799,40 +799,13 @@ namespace AutomationScripts
 		private static string GetFileHostCommandline(ProjectParams Params, DeploymentContext SC)
 		{
 			string FileHostParams = "";
-			FileReference ZenStoreMarkerFile = Params.FindZenProjectStoreMarker(SC);
-			bool UseZenServerHost = (ZenStoreMarkerFile != null) && !Params.UsePak(SC.StageTargetPlatform);
-			if (!Params.CookOnTheFly && !Params.FileServer && !UseZenServerHost)
+			if (!Params.CookOnTheFly && !Params.FileServer)
 			{
 				return FileHostParams;
 			}
 
-			string ProjectId = ProjectUtils.GetProjectPathId(SC.RawProjectPath);
-			ushort ZenHostPort = 8558;
-			List<string> HostAddresses = null;
-			if (ZenStoreMarkerFile != null)
-			{
-				byte[] ProjectStoreData = File.ReadAllBytes(ZenStoreMarkerFile.FullName);
-				CbObject ProjectStoreObject = new CbField(ProjectStoreData).AsObject();
-				CbObject ZenServerObject = ProjectStoreObject["zenserver"].AsObject();
-				if (ZenServerObject != CbObject.Empty)
-				{
-					ZenHostPort = ZenServerObject["hostport"].AsUInt16(ZenHostPort);
-					if (!ZenServerObject["islocalhost"].AsBool())
-					{
-						string HostName = ZenServerObject["hostname"].AsString();
-						if (!string.IsNullOrEmpty(HostName))
-						{
-							HostAddresses = new List<string> { HostName };
-						}
-					}
-					ProjectId = ZenServerObject["projectid"].AsString(ProjectId);
-				}
-			}
-			if (HostAddresses == null)
-			{
-				HostAddresses = GetFileHostAddresses(SC);
-				SC.StageTargetPlatform.ModifyFileHostAddresses(HostAddresses);
-			}
+			List<string> HostAddresses = GetFileHostAddresses(SC);
+			SC.StageTargetPlatform.ModifyFileHostAddresses(HostAddresses);
 
 			if (!IsNullOrEmpty(Params.Port))
 			{
@@ -853,20 +826,7 @@ namespace AutomationScripts
 				}
 			}
 
-			if (Params.CookOnTheFly)
-			{
-				FileHostParams += "-filehostip=";
-			}
-			else if (Params.ZenStore || ZenStoreMarkerFile != null)
-			{
-				FileHostParams += "-zenstoreproject=" + ProjectId + " ";
-				FileHostParams += "-zenstoreport=" + ZenHostPort + " ";
-				FileHostParams += "-zenstorehost=";
-			}
-			else
-			{
-				FileHostParams += "-filehostip=";
-			}
+			FileHostParams += "-filehostip=";
 			FileHostParams += String.Join("+", HostAddresses);
 			return FileHostParams;
 		}
