Staged ue.projectstore file contains a list of RemapDirectories that are applied when using Zen

 - allows Zen to load remapped Platforms/<platform> uplugin files etc.

#rnx

diff --git a/Engine/Source/Programs/AutomationTool/Scripts/CopyBuildToStagingDirectory.Automation.cs b/Engine/Source/Programs/AutomationTool/Scripts/CopyBuildToStagingDirectory.Automation.cs
--- a/Engine/Source/Programs/AutomationTool/Scripts/CopyBuildToStagingDirectory.Automation.cs
+++ b/Engine/Source/Programs/AutomationTool/Scripts/CopyBuildToStagingDirectory.Automation.cs
@@ -24,6 +24,7 @@ using Microsoft.Extensions.Logging;
 
 using static AutomationTool.CommandUtils;
 using System.Text.Json.Serialization;
+using System.Text.Json.Nodes;
 
 namespace AutomationScripts
 {
@@ -1052,6 +1053,64 @@ namespace AutomationScripts
 			throw new AutomationException(String.Format("A package store manifest or project store is required when staging to IoStore. Expected to find {0} or {1}. Ensure that cooking was successful.", PackageStoreManifestFile.FullName, ProjectStoreFile.FullName));
 		}
 
+
+
+		/// <summary>
+		/// Creates an updated .projectstore file with a new field containing the stage RemapDirectories, allowing StorageServer to apply these when requesting files
+		/// </summary>
+		/// <returns>New .projectstore file on success, source .projectsource file on failure</returns>
+		private static FileReference MakeProjectStoreWithRemappings(DeploymentContext SC, FileReference SourceProjectStorePath)
+		{
+			if (!EpicGames.Core.JsonObject.TryRead(SourceProjectStorePath, out EpicGames.Core.JsonObject ProjectStore))
+			{
+				return SourceProjectStorePath;
+			}
+
+			FileReference UpdatedProjectStorePath = FileReference.Combine(SC.ProjectRoot, "Intermediate", "Staging", SC.StageTargetPlatform.PlatformType.ToString(), "remap_" + SourceProjectStorePath.GetFileName());
+			if (FileReference.Exists(UpdatedProjectStorePath) && !InternalUtils.SafeDeleteFile(UpdatedProjectStorePath.FullName))
+			{
+				return SourceProjectStorePath;
+			}
+
+			if (!InternalUtils.SafeCreateDirectory(UpdatedProjectStorePath.Directory.FullName))
+			{
+				return SourceProjectStorePath;
+			}
+
+			string MakeStorageServerPath( StagedDirectoryReference Source )
+			{
+				string Result = Source.ToString();
+				string[] PathFragments = Result.Split(new[] { '\\', '/' }, StringSplitOptions.RemoveEmptyEntries);
+				if (PathFragments.Length > 0)
+				{
+					if (PathFragments[0].Equals("Engine", StringComparison.InvariantCultureIgnoreCase))
+					{
+						PathFragments[0] = "{engine}";
+					}
+					else if (PathFragments[0].Equals(SC.ShortProjectName, StringComparison.InvariantCultureIgnoreCase))
+					{
+						PathFragments[0] = "{project}";
+					}
+					Result = '/' + string.Join('/',PathFragments) + '/';
+				}
+				return Result;
+			}
+
+			List<EpicGames.Core.JsonObject> JsonRemapDirectories = [];
+			foreach (Tuple<StagedDirectoryReference, StagedDirectoryReference> RemapDirectory in SC.RemapDirectories)
+			{
+				EpicGames.Core.JsonObject JsonRemapDirectory = new();
+				JsonRemapDirectory.AddOrSetFieldValue("from", MakeStorageServerPath(RemapDirectory.Item1));
+				JsonRemapDirectory.AddOrSetFieldValue("to",   MakeStorageServerPath(RemapDirectory.Item2));
+				JsonRemapDirectories.Add(JsonRemapDirectory);
+			}
+			ProjectStore.AddOrSetFieldValue("remapDirectories", JsonRemapDirectories.ToArray());
+
+			File.WriteAllText(UpdatedProjectStorePath.FullName, ProjectStore.ToJsonString());
+			return UpdatedProjectStorePath;
+		}
+
+
 		private static StagedFileType? GetStagingTypeForCookedFile(FileReference File)
 		{
 			// json files have never been staged
@@ -1636,7 +1695,13 @@ namespace AutomationScripts
 							{
 								if (!Params.UsePak(SC.StageTargetPlatform))
 								{
-									SC.StageFile(StagedFileType.NonUFS, CookedFile, new StagedFileReference(CookedFile.MakeRelativeTo(SC.PlatformCookDir)));
+									FileReference SrcCookedFile = CookedFile;
+									if (SC.RemapDirectories.Any())
+									{
+										SrcCookedFile = MakeProjectStoreWithRemappings(SC, CookedFile);
+									}
+
+									SC.StageFile(StagedFileType.NonUFS, SrcCookedFile, new StagedFileReference(CookedFile.MakeRelativeTo(SC.PlatformCookDir)));
 								}
 								continue;
 							}
diff --git a/Engine/Source/Runtime/StorageServerClient/Private/StorageServerPlatformFile.cpp b/Engine/Source/Runtime/StorageServerClient/Private/StorageServerPlatformFile.cpp
--- a/Engine/Source/Runtime/StorageServerClient/Private/StorageServerPlatformFile.cpp
+++ b/Engine/Source/Runtime/StorageServerClient/Private/StorageServerPlatformFile.cpp
@@ -772,6 +772,19 @@ bool FStorageServerPlatformFile::Initialize(IPlatformFile* Inner, const TCHAR* C
 					}
 					UE_LOG(LogStorageServerPlatformFile, Display, TEXT("Using settings from ue.projectstore: ServerProject='%s' and ServerPlatform='%s'"), *ServerProject, *ServerPlatform);
 				}
+
+				const TArray<TSharedPtr<FJsonValue>>* RemapDirectoriesArrayPtr = nullptr;
+				if ( ProjectStoreObject->TryGetArrayField(TEXT("remapDirectories"), RemapDirectoriesArrayPtr) )
+				{
+					for (const TSharedPtr<FJsonValue>& JsonValue : *RemapDirectoriesArrayPtr)
+					{
+						const TSharedPtr<FJsonObject>& RemapObject = JsonValue->AsObject();
+						const FString RemapFrom = RemapObject->GetStringField(TEXT("from"));
+						const FString RemapTo   = RemapObject->GetStringField(TEXT("to"));
+						RemapDirectoriesTree.FindOrAdd(RemapFrom) = RemapTo;
+					}
+					RemapDirectoriesTree.Shrink();
+				}
 			}
 		}
 	
@@ -1373,6 +1386,18 @@ bool FStorageServerPlatformFile::SendGetFileListMessage()
 	
 	Connection->FileManifestRequest([&](FIoChunkId Id, FStringView Path, int64 RawSize)
 	{
+		FString RemapPathFrom;
+		FString* RemapPathToPtr = nullptr;
+		if (RemapDirectoriesTree.TryFindClosestPath(Path, RemapPathFrom, &RemapPathToPtr))
+		{
+			bool bTrimExtraSeparator = RemapPathFrom[RemapPathFrom.Len()-1] == TEXT('/');
+
+			TStringBuilder<1024> RemappedPath;
+			RemappedPath.Append(*RemapPathToPtr);
+			RemappedPath.Append(Path.RightChop(RemapPathFrom.Len() + (bTrimExtraSeparator ? 0 : 1)));
+			Path = RemappedPath;
+		}
+
 		ServerToc.AddFile(Id, Path, RawSize);
 	});
 
diff --git a/Engine/Source/Runtime/StorageServerClient/Private/StorageServerPlatformFile.h b/Engine/Source/Runtime/StorageServerClient/Private/StorageServerPlatformFile.h
--- a/Engine/Source/Runtime/StorageServerClient/Private/StorageServerPlatformFile.h
+++ b/Engine/Source/Runtime/StorageServerClient/Private/StorageServerPlatformFile.h
@@ -5,6 +5,7 @@
 #include "CoreMinimal.h"
 #include "GenericPlatform/GenericPlatformFile.h"
 #include "Containers/StringView.h"
+#include "Containers/DirectoryTree.h"
 #include "IO/IoDispatcher.h"
 #include "IStorageServerPlatformFile.h"
 
@@ -156,6 +157,7 @@ private:
 	FString ServerProject;
 	FString ServerPlatform;
 	FString BaseURI;
+	TDirectoryTree<FString> RemapDirectoriesTree;
 	mutable TArray<FString> HostAddrs;
 	mutable uint16 HostPort = 8558;
 	bool bAllowPackageIo = true;
