Gauntlet RPCs - Add RPC to log message in target

When writing RPC tests, it is often very useful to be able to determine from the log when you have reached a certain checkpoint across all targets. Surprisingly, there is no RPC command available to just log a string in the tartget. This CL adds that and makes use of it in the iteration time tests.

#tests tested locally
#rnx

diff --git a/Engine/Source/Programs/AutomationTool/Gauntlet/Framework/RpcFramework/Gauntlet.RpcLibrary.cs b/Engine/Source/Programs/AutomationTool/Gauntlet/Framework/RpcFramework/Gauntlet.RpcLibrary.cs
--- a/Engine/Source/Programs/AutomationTool/Gauntlet/Framework/RpcFramework/Gauntlet.RpcLibrary.cs
+++ b/Engine/Source/Programs/AutomationTool/Gauntlet/Framework/RpcFramework/Gauntlet.RpcLibrary.cs
@@ -239,6 +239,20 @@ namespace Gauntlet
 			return RequestHistory;
 		}
 
+		public static bool LogMessageInTarget(RpcTarget InTarget, string Message)
+		{
+			if (InTarget == null)
+			{
+				return false;
+			}
+			Dictionary<string, object> RequestArgs = new Dictionary<string, object>
+			{
+				{ "Message", Message }
+			};
+			RpcExecutor.CallRpc(InTarget, "Log", 5000, RequestArgs);
+			return true;
+		}
+
 		public static bool RunAndWaitForCheatCommand(RpcTarget InTarget, string CheatString)
         {
         	if (InTarget == null)
