Horde: Fix incorrect URL for uploaded ref in logs.

#rnx

diff --git a/Engine/Source/Programs/AutomationTool/BuildGraph/Tasks/CreateArtifactTask.cs b/Engine/Source/Programs/AutomationTool/BuildGraph/Tasks/CreateArtifactTask.cs
--- a/Engine/Source/Programs/AutomationTool/BuildGraph/Tasks/CreateArtifactTask.cs
+++ b/Engine/Source/Programs/AutomationTool/BuildGraph/Tasks/CreateArtifactTask.cs
@@ -194,7 +194,7 @@ namespace AutomationTool.Tasks
 			// Create the new artifact
 			IHordeClient hordeClient = CommandUtils.ServiceProvider.GetRequiredService<IHordeClient>();
 			IArtifactBuilder artifact = await hordeClient.Artifacts.CreateAsync(name, type, _parameters.Description, streamId, commitId, keys, metadata);
-			Logger.LogInformation("Creating artifact {ArtifactId} '{ArtifactName}' ({ArtifactType}) with namespace {NamespaceId}, ref {RefName} ({Link})", artifact.Id, name, type, artifact.NamespaceId, artifact.RefName, $"{hordeClient.ServerUrl}/api/v1/storage/{artifact.NamespaceId}/refs/{artifact.RefName}");
+			Logger.LogInformation("Creating artifact {ArtifactId} '{ArtifactName}' ({ArtifactType}) with namespace {NamespaceId}, ref {RefName} ({Link})", artifact.Id, name, type, artifact.NamespaceId, artifact.RefName, new Uri(hordeClient.ServerUrl, "/api/v1/storage/{artifact.NamespaceId}/refs/{artifact.RefName}").ToString());
 
 			// Upload the files
 			Stopwatch timer = Stopwatch.StartNew();
