Consistently handle Python data when staging a cooked editor

* bStagePython now consistently allows or removes all Python related data from a staged cooked editor, be it from Engine, Project, or Plugin.
* Python directories and .py/.pyd files are now consistently moved to NonUFS, as Python doesn't read files via UFS.
* .pyc and __pycache__ folders are now removed, even if bStagePython is true, as we don't want to include cache data in the build.
* The Project's Content/Python folder is now staged when bStagePython is true.

#rnx

diff --git a/Engine/Source/Programs/AutomationTool/CookedEditor/CookedEditor.Automation.cs b/Engine/Source/Programs/AutomationTool/CookedEditor/CookedEditor.Automation.cs
--- a/Engine/Source/Programs/AutomationTool/CookedEditor/CookedEditor.Automation.cs
+++ b/Engine/Source/Programs/AutomationTool/CookedEditor/CookedEditor.Automation.cs
@@ -491,6 +491,16 @@ public class MakeCookedEditor : BuildCommand
 			}		
 		}
 
+		// Python scripts
+		if (Context.bStagePython)
+		{
+			DirectoryReference ProjectPythonDirectory = DirectoryReference.Combine(Context.ProjectDirectory, "Content", "Python");
+			if (DirectoryReference.Exists(ProjectPythonDirectory))
+			{
+				Context.NonUFSFilesToStage.AddRange(DirectoryReference.EnumerateFiles(ProjectPythonDirectory, "*", SearchOption.AllDirectories));
+			}
+		}
+
 		StagePlatformExtensionFiles(Params, SC, Context, Context.ProjectDirectory);
 		StagePluginFiles(Params, SC, Context, false);
 
@@ -1021,9 +1031,41 @@ public class MakeCookedEditor : BuildCommand
 		}
 
 		// move some files from UFS to NonUFS if they ended up there
-		List<string> UFSIncompatibleExtensions = new List<string> { ".py", ".pyc" };
+		List<string> UFSIncompatibleExtensions = new List<string> { ".py", ".pyc", ".pyd" };
 		Context.NonUFSFilesToStage.AddRange(Context.UFSFilesToStage.Where(x => UFSIncompatibleExtensions.Contains(x.GetExtension())));
 		Context.UFSFilesToStage.RemoveAll(x => UFSIncompatibleExtensions.Contains(x.GetExtension()));
+		Context.NonUFSFilesToStage.AddRange(Context.UFSFilesToStage.Where(x => x.ContainsName("Python", 0) || x.ContainsName("__pycache__", 0)));
+		Context.UFSFilesToStage.RemoveAll(x => x.ContainsName("Python", 0) || x.ContainsName("__pycache__", 0));
+
+		// strip Python data now that it's all been moved to NonUFSFilesToStage
+		bool bShouldStagePython = !bIsCookedCooker && Context.bStagePython;
+		if (bShouldStagePython)
+		{
+			// strip out any leftover static library (.lib or .a) or source code (.c, .cpp, .h, .hpp, .inl) files within any Python script directories,
+			// as pre-compiled modules sometimes leave these in their build alongside the resultant .dll or .pyd files
+			List<string> NonPythonFilesToRemove = new List<string> { ".lib", ".a", ".c", ".cpp", ".h", ".hpp", ".inl" };
+			Context.NonUFSFilesToStage.RemoveAll(x => x.ContainsName("Python", 0) && NonPythonFilesToRemove.Contains(x.GetExtension()));
+		}
+		else
+		{
+			// strip out any Python script directories and .py or .pyd files if we're not staging Python
+			int PreviousCount = Context.NonUFSFilesToStage.Count;
+			Context.NonUFSFilesToStage.RemoveAll(x => x.ContainsName("Python", 0));
+			Context.NonUFSFilesToStage.RemoveAll(x => x.GetExtension() == ".py" || x.GetExtension() == ".pyd");
+			Logger.LogInformation("Removed {0} Python script file(s) from NonUFSFilesToStage", PreviousCount - Context.NonUFSFilesToStage.Count);
+		}
+		{
+			// strip out any __pycache__ directories and .pyc files regardless of whether we're staging Python (cache files shouldn't be staged)
+			int PreviousCount = Context.NonUFSFilesToStage.Count;
+			Context.NonUFSFilesToStage.RemoveAll(x => x.ContainsName("__pycache__", 0));
+			Context.NonUFSFilesToStage.RemoveAll(x => x.GetExtension() == ".pyc");
+			Logger.LogInformation("Removed {0} Python cache file(s) from NonUFSFilesToStage", PreviousCount - Context.NonUFSFilesToStage.Count);
+		}
+
+		// .tps files can all live in UFS
+		// we do this after stripping Python, as the Python content folders may have .tps files within them
+		Context.UFSFilesToStage.AddRange(Context.NonUFSFilesToStage.Where(x => x.GetExtension() == ".tps"));
+		Context.NonUFSFilesToStage.RemoveAll(x => x.GetExtension() == ".tps");
 	}
 
 	protected virtual string GetReleaseTargetName(UnrealTargetPlatform Platform, TargetType ReleaseType)
