Resolve intellisense errors

#rnx

diff --git a/Engine/Source/Programs/AutomationTool/BuildGraph/BgNodeExecutor.cs b/Engine/Source/Programs/AutomationTool/BuildGraph/BgNodeExecutor.cs
--- a/Engine/Source/Programs/AutomationTool/BuildGraph/BgNodeExecutor.cs
+++ b/Engine/Source/Programs/AutomationTool/BuildGraph/BgNodeExecutor.cs
@@ -316,7 +316,7 @@ namespace AutomationTool
 			{
 				return Enum.Parse(valueType, valueText);
 			}
-			else if (valueType == typeof(Boolean))
+			else if (valueType == typeof(bool))
 			{
 				return await BgCondition.EvaluateAsync(valueText, context);
 			}
@@ -399,7 +399,7 @@ namespace AutomationTool
 						while (idx + 1 < Node.Tasks.Count && executor.Add(_boundTasks[idx + 1]))
 						{
 							idx++;
-							_boundTasks[idx].GetTraceMetadata(scope.Span, string.Format("{0}.", 1 + idx - firstIdx));
+							_boundTasks[idx].GetTraceMetadata(scope.Span, String.Format("{0}.", 1 + idx - firstIdx));
 						}
 						try
 						{
diff --git a/Engine/Source/Programs/AutomationTool/BuildGraph/BuildGraph.cs b/Engine/Source/Programs/AutomationTool/BuildGraph/BuildGraph.cs
--- a/Engine/Source/Programs/AutomationTool/BuildGraph/BuildGraph.cs
+++ b/Engine/Source/Programs/AutomationTool/BuildGraph/BuildGraph.cs
@@ -124,9 +124,9 @@ namespace AutomationTool
 		private ScriptTaskBinding(string name, Type taskClass, Type parametersClass, List<ScriptTaskParameterBinding> parameters)
 			: base(name, parameters.ConvertAll<BgScriptTaskParameter>(x => x))
 		{
-			this.TaskClass = taskClass;
-			this.ParametersClass = parametersClass;
-			this.NameToParameter = parameters.ToDictionary(x => x.Name, x => x);
+			TaskClass = taskClass;
+			ParametersClass = parametersClass;
+			NameToParameter = parameters.ToDictionary(x => x.Name, x => x);
 		}
 
 		static List<ScriptTaskParameterBinding> CreateParameters(Type parametersClass)
@@ -193,13 +193,13 @@ namespace AutomationTool
 		/// <param name="codeChange">Code changelist being built</param>
 		internal BgEnvironment(string? branch, int? change, int? codeChange)
 		{
-			this.Stream = branch ?? "Unknown";
-			this.Change = change ?? 0;
-			this.CodeChange = codeChange ?? 0;
-			this.IsBuildMachine = CommandUtils.IsBuildMachine;
+			Stream = branch ?? "Unknown";
+			Change = change ?? 0;
+			CodeChange = codeChange ?? 0;
+			IsBuildMachine = CommandUtils.IsBuildMachine;
 
 			ReadOnlyBuildVersion version = ReadOnlyBuildVersion.Current;
-			this.EngineVersion = (version.MajorVersion, version.MinorVersion, version.PatchVersion);
+			EngineVersion = (version.MajorVersion, version.MinorVersion, version.PatchVersion);
 		}
 	}
 
@@ -336,7 +336,7 @@ namespace AutomationTool
 			defaultProperties["DataDrivenPlatforms"] = String.Join(";", DataDrivenPlatformInfo.GetAllPlatformInfos().Keys);
 
 			// Look for overrides
-			if (!string.IsNullOrEmpty(branchOverride))
+			if (!String.IsNullOrEmpty(branchOverride))
 			{
 				Logger.LogInformation("Overriding default branch '{Branch}' with '{BranchOverride}'", defaultProperties["Branch"], branchOverride);
 				defaultProperties["Branch"] = branchOverride;
@@ -805,7 +805,7 @@ namespace AutomationTool
 				}
 				catch (ReflectionTypeLoadException ex)
 				{
-					Logger.LogWarning(ex, "Exception {Ex} while trying to get types from assembly {LoadedAssembly}. LoaderExceptions: {Arg2}", ex, loadedAssembly, string.Join("\n", ex.LoaderExceptions.Select(x => x?.Message)));
+					Logger.LogWarning(ex, "Exception {Ex} while trying to get types from assembly {LoadedAssembly}. LoaderExceptions: {Arg2}", ex, loadedAssembly, String.Join("\n", ex.LoaderExceptions.Select(x => x?.Message)));
 					continue;
 				}
 
@@ -841,7 +841,7 @@ namespace AutomationTool
 				}
 				catch (ReflectionTypeLoadException ex)
 				{
-					Logger.LogWarning(ex, "Exception {Ex} while trying to get types from assembly {LoadedAssembly}. LoaderExceptions: {Arg2}", ex, loadedAssembly, string.Join("\n", ex.LoaderExceptions.Select(x => x?.Message)));
+					Logger.LogWarning(ex, "Exception {Ex} while trying to get types from assembly {LoadedAssembly}. LoaderExceptions: {Arg2}", ex, loadedAssembly, String.Join("\n", ex.LoaderExceptions.Select(x => x?.Message)));
 					continue;
 				}
 
@@ -1158,7 +1158,7 @@ namespace AutomationTool
 						.ToList();
 					if (previousNodeNames.Count > 0)
 					{
-						message += $" (previous {(previousNodeNames.Count == 1 ? "step" : "steps")}: {string.Join(" + ", previousNodeNames)})";
+						message += $" (previous {(previousNodeNames.Count == 1 ? "step" : "steps")}: {String.Join(" + ", previousNodeNames)})";
 					}
 
 					modifiedFiles.Add(file.RelativePath, message);
diff --git a/Engine/Source/Programs/AutomationTool/BuildGraph/Tasks/AgeStoreTask.cs b/Engine/Source/Programs/AutomationTool/BuildGraph/Tasks/AgeStoreTask.cs
--- a/Engine/Source/Programs/AutomationTool/BuildGraph/Tasks/AgeStoreTask.cs
+++ b/Engine/Source/Programs/AutomationTool/BuildGraph/Tasks/AgeStoreTask.cs
@@ -156,11 +156,11 @@ namespace AutomationTool.Tasks
 			else
 			{
 				string[] patterns = directoryStructure[level].Split(';');
-				foreach (var pattern in patterns)
+				foreach (string pattern in patterns)
 				{
-					string replacedPattern = string.Format(pattern, filter);
+					string replacedPattern = String.Format(pattern, filter);
 
-					foreach (var childDirectory in currentDirectory.GetDirectories(replacedPattern, SearchOption.TopDirectoryOnly))
+					foreach (DirectoryInfo childDirectory in currentDirectory.GetDirectories(replacedPattern, SearchOption.TopDirectoryOnly))
 					{
 						RecurseDirectory(expireTimeUtc, childDirectory, directoryStructure, level + 1, filter, existingBuilds, deleteIndividualFiles);
 					}
@@ -190,14 +190,14 @@ namespace AutomationTool.Tasks
 				throw new AutomationException("Platform does not specify the symbol server structure. Cannot age the symbol server.");
 			}
 
-			string filter = string.IsNullOrWhiteSpace(_parameters.Filter)
-				? string.Empty
+			string filter = String.IsNullOrWhiteSpace(_parameters.Filter)
+				? String.Empty
 				: _parameters.Filter.Trim();
 
 			// Eumerate the root directory of builds for buildversions to check against
 			// Folder names in the root directory should match the name of the .version files
 			HashSet<string> existingBuilds = new HashSet<string>(StringComparer.OrdinalIgnoreCase);
-			if (!string.IsNullOrWhiteSpace(_parameters.BuildDir))
+			if (!String.IsNullOrWhiteSpace(_parameters.BuildDir))
 			{
 				DirectoryReference buildDir = new DirectoryReference(_parameters.BuildDir);
 				if (DirectoryReference.Exists(buildDir))
diff --git a/Engine/Source/Programs/AutomationTool/BuildGraph/Tasks/CompileTask.cs b/Engine/Source/Programs/AutomationTool/BuildGraph/Tasks/CompileTask.cs
--- a/Engine/Source/Programs/AutomationTool/BuildGraph/Tasks/CompileTask.cs
+++ b/Engine/Source/Programs/AutomationTool/BuildGraph/Tasks/CompileTask.cs
@@ -255,7 +255,7 @@ namespace AutomationTool.Tasks
 		/// <param name="parameters">Parameters for this task</param>
 		public CompileTask(CompileTaskParameters parameters)
 		{
-			this.Parameters = parameters;
+			Parameters = parameters;
 		}
 
 		/// <summary>
diff --git a/Engine/Source/Programs/AutomationTool/BuildGraph/Tasks/CookTask.cs b/Engine/Source/Programs/AutomationTool/BuildGraph/Tasks/CookTask.cs
--- a/Engine/Source/Programs/AutomationTool/BuildGraph/Tasks/CookTask.cs
+++ b/Engine/Source/Programs/AutomationTool/BuildGraph/Tasks/CookTask.cs
@@ -110,7 +110,7 @@ namespace AutomationTool.Tasks
 				scope.Span.SetTag("platform", _parameters.Platform);
 				string[] maps = (_parameters.Maps == null) ? null : _parameters.Maps.Split(new char[] { '+' });
 				string arguments = (_parameters.Versioned ? "" : "-Unversioned ") + _parameters.Arguments;
-				string editorExe = (string.IsNullOrWhiteSpace(_parameters.EditorExe) ? ProjectUtils.GetEditorForProject(projectFile).FullName : _parameters.EditorExe);
+				string editorExe = (System.String.IsNullOrWhiteSpace(_parameters.EditorExe) ? ProjectUtils.GetEditorForProject(projectFile).FullName : _parameters.EditorExe);
 				CommandUtils.CookCommandlet(projectFile, editorExe, maps, null, null, null, _parameters.Platform, arguments);
 			}
 
diff --git a/Engine/Source/Programs/AutomationTool/BuildGraph/Tasks/CreateArtifactTask.cs b/Engine/Source/Programs/AutomationTool/BuildGraph/Tasks/CreateArtifactTask.cs
--- a/Engine/Source/Programs/AutomationTool/BuildGraph/Tasks/CreateArtifactTask.cs
+++ b/Engine/Source/Programs/AutomationTool/BuildGraph/Tasks/CreateArtifactTask.cs
@@ -119,8 +119,8 @@ namespace AutomationTool.Tasks
 		{
 			ArtifactName name = new ArtifactName(_parameters.Name);
 			ArtifactType type = new ArtifactType(_parameters.Type);
-			List<string> keys = (_parameters.Keys ?? string.Empty).Split(';', StringSplitOptions.RemoveEmptyEntries | StringSplitOptions.TrimEntries).ToList();
-			List<string> metadata = (_parameters.Metadata ?? string.Empty).Split(';', StringSplitOptions.RemoveEmptyEntries | StringSplitOptions.TrimEntries).ToList();
+			List<string> keys = (_parameters.Keys ?? String.Empty).Split(';', StringSplitOptions.RemoveEmptyEntries | StringSplitOptions.TrimEntries).ToList();
+			List<string> metadata = (_parameters.Metadata ?? String.Empty).Split(';', StringSplitOptions.RemoveEmptyEntries | StringSplitOptions.TrimEntries).ToList();
 
 			// Add keys for the job that's executing
 			string? jobId = Environment.GetEnvironmentVariable("UE_HORDE_JOBID");
diff --git a/Engine/Source/Programs/AutomationTool/BuildGraph/Tasks/CsCompileTask.cs b/Engine/Source/Programs/AutomationTool/BuildGraph/Tasks/CsCompileTask.cs
--- a/Engine/Source/Programs/AutomationTool/BuildGraph/Tasks/CsCompileTask.cs
+++ b/Engine/Source/Programs/AutomationTool/BuildGraph/Tasks/CsCompileTask.cs
@@ -345,8 +345,8 @@ namespace AutomationTool.Tasks
 		/// </summary>
 		public CsCompileOutput(FileSet binaries, FileSet references)
 		{
-			this.Binaries = binaries;
-			this.References = references;
+			Binaries = binaries;
+			References = references;
 		}
 
 		/// <summary>
diff --git a/Engine/Source/Programs/AutomationTool/BuildGraph/Tasks/GatherBuildProductsFromFileTask.cs b/Engine/Source/Programs/AutomationTool/BuildGraph/Tasks/GatherBuildProductsFromFileTask.cs
--- a/Engine/Source/Programs/AutomationTool/BuildGraph/Tasks/GatherBuildProductsFromFileTask.cs
+++ b/Engine/Source/Programs/AutomationTool/BuildGraph/Tasks/GatherBuildProductsFromFileTask.cs
@@ -36,8 +36,8 @@ namespace AutomationTool.Tasks
 
 			try
 			{
-				var fileBuildProducts = await File.ReadAllLinesAsync(_parameters.BuildProductsFile);
-				foreach (var buildProduct in fileBuildProducts)
+				string[] fileBuildProducts = await File.ReadAllLinesAsync(_parameters.BuildProductsFile);
+				foreach (string buildProduct in fileBuildProducts)
 				{
 					Logger.LogInformation("Adding file to build products: {BuildProduct}", buildProduct);
 					buildProducts.Add(new FileReference(buildProduct));
diff --git a/Engine/Source/Programs/AutomationTool/BuildGraph/Tasks/HelmTask.cs b/Engine/Source/Programs/AutomationTool/BuildGraph/Tasks/HelmTask.cs
--- a/Engine/Source/Programs/AutomationTool/BuildGraph/Tasks/HelmTask.cs
+++ b/Engine/Source/Programs/AutomationTool/BuildGraph/Tasks/HelmTask.cs
@@ -129,7 +129,7 @@ namespace AutomationTool.Tasks
 				arguments.Add("--kubeconfig");
 				arguments.Add(_parameters.KubeConfig);
 			}
-			if (!string.IsNullOrEmpty(_parameters.Values))
+			if (!String.IsNullOrEmpty(_parameters.Values))
 			{
 				foreach (string value in SplitDelimitedList(_parameters.Values))
 				{
diff --git a/Engine/Source/Programs/AutomationTool/BuildGraph/Tasks/ModifyJsonValue.cs b/Engine/Source/Programs/AutomationTool/BuildGraph/Tasks/ModifyJsonValue.cs
--- a/Engine/Source/Programs/AutomationTool/BuildGraph/Tasks/ModifyJsonValue.cs
+++ b/Engine/Source/Programs/AutomationTool/BuildGraph/Tasks/ModifyJsonValue.cs
@@ -64,11 +64,11 @@ namespace AutomationTool.Tasks
 				return;
 			}
 			HashSet<FileReference> files = ResolveFilespec(Unreal.RootDirectory, _parameters.Files, tagNameToFileSet);
-			foreach (var jsonFile in files.Select(f => f.FullName))
+			foreach (string jsonFile in files.Select(f => f.FullName))
 			{
-				var oldContents = await File.ReadAllTextAsync(jsonFile);
-				var paramObj = fastJSON.JSON.Instance.Parse(oldContents) as IDictionary<string, object>;
-				var currObj = paramObj;
+				string oldContents = await File.ReadAllTextAsync(jsonFile);
+				IDictionary<string, object> paramObj = fastJSON.JSON.Instance.Parse(oldContents) as IDictionary<string, object>;
+				IDictionary<string, object> currObj = paramObj;
 				for (int i = 0; i < keys.Length - 1; i++)
 				{
 					if (!currObj.TryGetValue(keys[i], out object nextNode))
@@ -80,7 +80,7 @@ namespace AutomationTool.Tasks
 
 				currObj[keys[keys.Length - 1]] = _parameters.NewValue;
 
-				var newContents = JsonSerializer.Serialize(paramObj, new JsonSerializerOptions { WriteIndented = true });
+				string newContents = JsonSerializer.Serialize(paramObj, new JsonSerializerOptions { WriteIndented = true });
 				await File.WriteAllTextAsync(jsonFile, newContents, new UTF8Encoding(encoderShouldEmitUTF8Identifier: false));
 			}
 		}
diff --git a/Engine/Source/Programs/AutomationTool/BuildGraph/Tasks/NotarizeTask.cs b/Engine/Source/Programs/AutomationTool/BuildGraph/Tasks/NotarizeTask.cs
--- a/Engine/Source/Programs/AutomationTool/BuildGraph/Tasks/NotarizeTask.cs
+++ b/Engine/Source/Programs/AutomationTool/BuildGraph/Tasks/NotarizeTask.cs
@@ -100,7 +100,7 @@ namespace AutomationTool.Tasks
 
 			for (int numRetries = 0; numRetries < MaxNumRetries; numRetries++)
 			{
-				string commandLine = string.Format("notarytool submit \"{0}\" --keychain-profile \"{1}\" --wait --timeout \"{2}\"", dmg.FullName, _parameters.KeyChainID, timeoutInMilliseconds);
+				string commandLine = String.Format("notarytool submit \"{0}\" --keychain-profile \"{1}\" --wait --timeout \"{2}\"", dmg.FullName, _parameters.KeyChainID, timeoutInMilliseconds);
 				output = CommandUtils.RunAndLog("xcrun", commandLine, out exitCode);
 
 				if (exitCode == 0)
@@ -148,7 +148,7 @@ namespace AutomationTool.Tasks
 					{
 						// once we have a log file, print it out, staple, and we're done.
 						Logger.LogInformation("{Text}", GetRequestLogs(requestUuid));
-						string commandLine = string.Format("stapler staple {0}", dmg.FullName);
+						string commandLine = String.Format("stapler staple {0}", dmg.FullName);
 						output = CommandUtils.RunAndLog("xcrun", commandLine, out exitCode);
 						if (exitCode != 0)
 						{
@@ -179,7 +179,7 @@ namespace AutomationTool.Tasks
 		{
 			try
 			{
-				string logCommand = string.Format("notarytool log {0} --keychain-profile \"{1}\"", requestUuid, _parameters.KeyChainID);
+				string logCommand = String.Format("notarytool log {0} --keychain-profile \"{1}\"", requestUuid, _parameters.KeyChainID);
 				IProcessResult logResult = CommandUtils.Run("xcrun", logCommand);
 
 				string responseContent = null;
@@ -192,7 +192,7 @@ namespace AutomationTool.Tasks
 			}
 			catch (Exception ex)
 			{
-				throw new AutomationException(ex, string.Format("Couldn't complete the request, error: {0}", ex.Message));
+				throw new AutomationException(ex, String.Format("Couldn't complete the request, error: {0}", ex.Message));
 			}
 		}
 
diff --git a/Engine/Source/Programs/AutomationTool/BuildGraph/Tasks/WriteJsonValueTask.cs b/Engine/Source/Programs/AutomationTool/BuildGraph/Tasks/WriteJsonValueTask.cs
--- a/Engine/Source/Programs/AutomationTool/BuildGraph/Tasks/WriteJsonValueTask.cs
+++ b/Engine/Source/Programs/AutomationTool/BuildGraph/Tasks/WriteJsonValueTask.cs
@@ -143,10 +143,10 @@ namespace AutomationTool.Tasks
 
 				string indexStr = key.Substring(minIdx + 1, (maxIdx - 1) - (minIdx + 1)).Trim();
 
-				int index = int.MaxValue;
+				int index = Int32.MaxValue;
 				if (indexStr.Length > 0)
 				{
-					index = int.Parse(indexStr);
+					index = Int32.Parse(indexStr);
 				}
 
 				JsonArray? array = prevValue as JsonArray;
diff --git a/Engine/Source/Programs/AutomationTool/BuildGraph/Tasks/WriteTextFileTask.cs b/Engine/Source/Programs/AutomationTool/BuildGraph/Tasks/WriteTextFileTask.cs
--- a/Engine/Source/Programs/AutomationTool/BuildGraph/Tasks/WriteTextFileTask.cs
+++ b/Engine/Source/Programs/AutomationTool/BuildGraph/Tasks/WriteTextFileTask.cs
@@ -80,7 +80,7 @@ namespace AutomationTool.Tasks
 			// If any files or tagsets are provided, add them to the text output.
 			if (!String.IsNullOrEmpty(_parameters.Files))
 			{
-				if (!string.IsNullOrWhiteSpace(fileText))
+				if (!String.IsNullOrWhiteSpace(fileText))
 				{
 					fileText += Environment.NewLine;
 				}
@@ -88,7 +88,7 @@ namespace AutomationTool.Tasks
 				HashSet<FileReference> files = ResolveFilespec(Unreal.RootDirectory, _parameters.Files, tagNameToFileSet);
 				if (files.Any())
 				{
-					fileText += string.Join(Environment.NewLine, files.Select(f => f.FullName));
+					fileText += String.Join(Environment.NewLine, files.Select(f => f.FullName));
 				}
 			}
 
@@ -100,12 +100,12 @@ namespace AutomationTool.Tasks
 
 			if (_parameters.Append)
 			{
-				Logger.LogInformation("{Text}", string.Format("Appending text to file '{0}': {1}", _parameters.File, fileText));
+				Logger.LogInformation("{Text}", String.Format("Appending text to file '{0}': {1}", _parameters.File, fileText));
 				await FileReference.AppendAllTextAsync(_parameters.File, Environment.NewLine + fileText);
 			}
 			else
 			{
-				Logger.LogInformation("{Text}", string.Format("Writing text to file '{0}': {1}", _parameters.File, fileText));
+				Logger.LogInformation("{Text}", String.Format("Writing text to file '{0}': {1}", _parameters.File, fileText));
 				await FileReference.WriteAllTextAsync(_parameters.File, fileText);
 			}
 
diff --git a/Engine/Source/Programs/AutomationTool/BuildGraph/Tasks/ZenExportSnapshotTask.cs b/Engine/Source/Programs/AutomationTool/BuildGraph/Tasks/ZenExportSnapshotTask.cs
--- a/Engine/Source/Programs/AutomationTool/BuildGraph/Tasks/ZenExportSnapshotTask.cs
+++ b/Engine/Source/Programs/AutomationTool/BuildGraph/Tasks/ZenExportSnapshotTask.cs
@@ -325,7 +325,7 @@ namespace AutomationTool.Tasks
 		private void WriteExportSource(JsonWriter writer, SnapshotStorageType destinationStorageType, ExportSourceData exportSource, string name)
 		{
 			string targetPlatform = _parameters.SnapshotDescriptorPlatform;
-			if (string.IsNullOrEmpty(targetPlatform))
+			if (String.IsNullOrEmpty(targetPlatform))
 			{
 				targetPlatform = exportSource._targetPlatform;
 			}
@@ -335,20 +335,20 @@ namespace AutomationTool.Tasks
 				case SnapshotStorageType.Cloud:
 					string bucketName = _parameters.DestinationCloudBucket;
 					string projectNameAsBucketName = _parameters.Project.GetFileNameWithoutAnyExtensions().ToLowerInvariant();
-					if (string.IsNullOrEmpty(bucketName))
+					if (String.IsNullOrEmpty(bucketName))
 					{
 						bucketName = projectNameAsBucketName;
 					}
 					bucketName = SanitizeBucketName(bucketName);
 
 					string hostName = _parameters.SnapshotDescriptorCloudHost;
-					if (string.IsNullOrEmpty(hostName))
+					if (String.IsNullOrEmpty(hostName))
 					{
 						hostName = _parameters.DestinationCloudHost;
 					}
 
 					string httpVersion = _parameters.SnapshotDescriptorCloudHttpVersion;
-					if (string.IsNullOrEmpty(httpVersion))
+					if (String.IsNullOrEmpty(httpVersion))
 					{
 						httpVersion = _parameters.DestinationCloudHttpVersion;
 					}
@@ -358,7 +358,7 @@ namespace AutomationTool.Tasks
 					writer.WriteValue("type", "cloud");
 					writer.WriteValue("targetplatform", targetPlatform);
 					writer.WriteValue("host", hostName);
-					if (!string.IsNullOrEmpty(httpVersion) && !httpVersion.Equals("None", StringComparison.OrdinalIgnoreCase))
+					if (!String.IsNullOrEmpty(httpVersion) && !httpVersion.Equals("None", StringComparison.OrdinalIgnoreCase))
 					{
 						writer.WriteValue("httpversion", httpVersion);
 					}
@@ -456,7 +456,7 @@ namespace AutomationTool.Tasks
 			foreach (string platform in _parameters.Platform.Split('+'))
 			{
 				DirectoryReference platformCookedDirectory;
-				if (string.IsNullOrEmpty(_parameters.OverridePlatformCookedDir))
+				if (String.IsNullOrEmpty(_parameters.OverridePlatformCookedDir))
 				{
 					platformCookedDirectory = DirectoryReference.Combine(projectFile.Directory, "Saved", "Cooked", platform);
 				}
@@ -521,22 +521,22 @@ namespace AutomationTool.Tasks
 			switch (destinationStorageType)
 			{
 				case SnapshotStorageType.Cloud:
-					if (string.IsNullOrEmpty(_parameters.DestinationCloudHost))
+					if (String.IsNullOrEmpty(_parameters.DestinationCloudHost))
 					{
 						throw new AutomationException("Missing destination cloud host");
 					}
-					if (string.IsNullOrEmpty(_parameters.DestinationCloudNamespace))
+					if (String.IsNullOrEmpty(_parameters.DestinationCloudNamespace))
 					{
 						throw new AutomationException("Missing destination cloud namespace");
 					}
-					if (string.IsNullOrEmpty(_parameters.DestinationIdentifier))
+					if (String.IsNullOrEmpty(_parameters.DestinationIdentifier))
 					{
 						throw new AutomationException("Missing destination identifier when exporting to cloud");
 					}
 
 					string bucketName = _parameters.DestinationCloudBucket;
 					string projectNameAsBucketName = projectFile.GetFileNameWithoutAnyExtensions().ToLowerInvariant();
-					if (string.IsNullOrEmpty(bucketName))
+					if (String.IsNullOrEmpty(bucketName))
 					{
 						bucketName = projectNameAsBucketName;
 					}
@@ -544,7 +544,7 @@ namespace AutomationTool.Tasks
 
 					oplogExportCommandline.AppendFormat(" --cloud {0} --namespace {1} --bucket {2}", _parameters.DestinationCloudHost, _parameters.DestinationCloudNamespace, bucketName);
 
-					if (!string.IsNullOrEmpty(_parameters.DestinationCloudHttpVersion))
+					if (!String.IsNullOrEmpty(_parameters.DestinationCloudHttpVersion))
 					{
 						if (_parameters.DestinationCloudHttpVersion.Equals("http2-only", StringComparison.OrdinalIgnoreCase))
 						{
@@ -559,10 +559,10 @@ namespace AutomationTool.Tasks
 					exportIndex = 0;
 					foreach (ExportSourceData exportSource in exportSources)
 					{
-						string hostUrlArg = string.Format("--hosturl http://{0}:{1}", exportSource._isLocalHost ? "localhost" : exportSource._hostName, exportSource._hostPort);
+						string hostUrlArg = String.Format("--hosturl http://{0}:{1}", exportSource._isLocalHost ? "localhost" : exportSource._hostName, exportSource._hostPort);
 
-						string baseKeyArg = string.Empty;
-						if ((exportSource._snapshotBaseDescriptor != null) && !string.IsNullOrEmpty(exportSource._snapshotBaseDescriptor.Key))
+						string baseKeyArg = String.Empty;
+						if ((exportSource._snapshotBaseDescriptor != null) && !String.IsNullOrEmpty(exportSource._snapshotBaseDescriptor.Key))
 						{
 							if (exportSource._snapshotBaseDescriptor.Type == SnapshotStorageType.Cloud)
 							{
@@ -593,11 +593,11 @@ namespace AutomationTool.Tasks
 
 					break;
 				case SnapshotStorageType.Zen:
-					if (string.IsNullOrEmpty(_parameters.DestinationZenHost))
+					if (String.IsNullOrEmpty(_parameters.DestinationZenHost))
 					{
 						throw new AutomationException("Missing destination zen host");
 					}
-					if (string.IsNullOrEmpty(_parameters.DestinationIdentifier))
+					if (String.IsNullOrEmpty(_parameters.DestinationIdentifier))
 					{
 						throw new AutomationException("Missing destination identifier when exporting to zen");
 					}
@@ -613,7 +613,7 @@ namespace AutomationTool.Tasks
 					exportIndex = 0;
 					foreach (ExportSourceData exportSource in exportSources)
 					{
-						string hostUrlArg = string.Format("--hosturl http://{0}:{1}", exportSource._isLocalHost ? "localhost" : exportSource._hostName, exportSource._hostPort);
+						string hostUrlArg = String.Format("--hosturl http://{0}:{1}", exportSource._isLocalHost ? "localhost" : exportSource._hostName, exportSource._hostPort);
 
 						StringBuilder exportSingleSourceCommandline = new StringBuilder(oplogExportCommandline.Length);
 						exportSingleSourceCommandline.Append(oplogExportCommandline);
@@ -646,15 +646,15 @@ namespace AutomationTool.Tasks
 						exportSingleSourceCommandline.Append(oplogExportCommandline);
 
 						string destinationFileName = exportSource._oplogId;
-						if (!string.IsNullOrEmpty(_parameters.DestinationFileName))
+						if (!String.IsNullOrEmpty(_parameters.DestinationFileName))
 						{
 							destinationFileName = _parameters.DestinationFileName.Replace("{Platform}", exportSource._targetPlatform, StringComparison.InvariantCultureIgnoreCase);
 						}
 
-						string projectId = string.IsNullOrEmpty(exportSource._projectId) ? defaultProjectId : exportSource._projectId;
-						string baseNameArg = string.Empty;
+						string projectId = String.IsNullOrEmpty(exportSource._projectId) ? defaultProjectId : exportSource._projectId;
+						string baseNameArg = String.Empty;
 						DirectoryReference platformDestinationFileDir = new DirectoryReference(_parameters.DestinationFileDir.FullName.Replace("{Platform}", exportSource._targetPlatform, StringComparison.InvariantCultureIgnoreCase));
-						if ((exportSource._snapshotBaseDescriptor != null) && !string.IsNullOrEmpty(exportSource._snapshotBaseDescriptor.Directory) && !string.IsNullOrEmpty(exportSource._snapshotBaseDescriptor.Filename))
+						if ((exportSource._snapshotBaseDescriptor != null) && !String.IsNullOrEmpty(exportSource._snapshotBaseDescriptor.Directory) && !String.IsNullOrEmpty(exportSource._snapshotBaseDescriptor.Filename))
 						{
 							if (exportSource._snapshotBaseDescriptor.Type == SnapshotStorageType.File)
 							{
diff --git a/Engine/Source/Programs/AutomationTool/BuildGraph/Tasks/ZenImportOplogTask.cs b/Engine/Source/Programs/AutomationTool/BuildGraph/Tasks/ZenImportOplogTask.cs
--- a/Engine/Source/Programs/AutomationTool/BuildGraph/Tasks/ZenImportOplogTask.cs
+++ b/Engine/Source/Programs/AutomationTool/BuildGraph/Tasks/ZenImportOplogTask.cs
@@ -128,7 +128,7 @@ namespace AutomationTool.Tasks
 		public override Task ExecuteAsync(JobContext job, HashSet<FileReference> buildProducts, Dictionary<string, HashSet<FileReference>> tagNameToFileSet)
 		{
 			SnapshotStorageType importMethod = SnapshotStorageType.Invalid;
-			if (!string.IsNullOrEmpty(_parameters.ImportType))
+			if (!String.IsNullOrEmpty(_parameters.ImportType))
 			{
 				importMethod = (SnapshotStorageType)Enum.Parse(typeof(SnapshotStorageType), _parameters.ImportType);
 			}
@@ -157,7 +157,7 @@ namespace AutomationTool.Tasks
 				string rootDir = _parameters.RootDir;
 				string engineDir = System.IO.Path.Combine(_parameters.RootDir, "Engine");
 				string projectDir = System.IO.Path.Combine(_parameters.RootDir, _projectFile.GetFileNameWithoutAnyExtensions());
-				string hostUrlArg = string.Format("--hosturl http://{0}:{1}", _parameters.HostName, _parameters.HostPort);
+				string hostUrlArg = String.Format("--hosturl http://{0}:{1}", _parameters.HostName, _parameters.HostPort);
 				StringBuilder oplogProjectCreateCommandline = new StringBuilder();
 				oplogProjectCreateCommandline.AppendFormat("project-create -p {0} --rootdir {1} --enginedir {2} --projectdir {3} --projectfile {4} {5}",
 					_parameters.ProjectName,
@@ -235,27 +235,27 @@ namespace AutomationTool.Tasks
 
 		private void ImportFromCloud(FileReference zenExe)
 		{
-			if (string.IsNullOrEmpty(_parameters.CloudURL))
+			if (String.IsNullOrEmpty(_parameters.CloudURL))
 			{
 				throw new AutomationException("Missing destination cloud host");
 			}
-			if (string.IsNullOrEmpty(_parameters.Namespace))
+			if (String.IsNullOrEmpty(_parameters.Namespace))
 			{
 				throw new AutomationException("Missing destination cloud namespace");
 			}
-			if (string.IsNullOrEmpty(_parameters.Key))
+			if (String.IsNullOrEmpty(_parameters.Key))
 			{
 				throw new AutomationException("Missing destination cloud storage key");
 			}
 
 			string bucketName = _parameters.Bucket;
 			string projectNameAsBucketName = _projectFile.GetFileNameWithoutAnyExtensions().ToLowerInvariant();
-			if (string.IsNullOrEmpty(bucketName))
+			if (String.IsNullOrEmpty(bucketName))
 			{
 				bucketName = projectNameAsBucketName;
 			}
 
-			string hostUrlArg = string.Format("--hosturl http://{0}:{1}", _parameters.HostName, _parameters.HostPort);
+			string hostUrlArg = String.Format("--hosturl http://{0}:{1}", _parameters.HostName, _parameters.HostPort);
 			StringBuilder oplogImportCommandline = new StringBuilder();
 			oplogImportCommandline.AppendFormat("oplog-import {0} --cloud {1} --namespace {2} --bucket {3}", hostUrlArg, _parameters.CloudURL, _parameters.Namespace, bucketName);
 			oplogImportCommandline.AppendFormat(" {0}", _parameters.Key);
diff --git a/Engine/Source/Programs/AutomationTool/BuildGraph/Tasks/ZenMirrorTask.cs b/Engine/Source/Programs/AutomationTool/BuildGraph/Tasks/ZenMirrorTask.cs
--- a/Engine/Source/Programs/AutomationTool/BuildGraph/Tasks/ZenMirrorTask.cs
+++ b/Engine/Source/Programs/AutomationTool/BuildGraph/Tasks/ZenMirrorTask.cs
@@ -141,7 +141,7 @@ namespace AutomationTool.Tasks
 			{
 				throw new AutomationException("Missing project file - {0}", projectFile.FullName);
 			}
-			if (string.IsNullOrEmpty(_parameters.Platform))
+			if (String.IsNullOrEmpty(_parameters.Platform))
 			{
 				throw new AutomationException("Missing platform");
 			}
diff --git a/Engine/Source/Programs/AutomationTool/BuildGraph/TempStorage.cs b/Engine/Source/Programs/AutomationTool/BuildGraph/TempStorage.cs
--- a/Engine/Source/Programs/AutomationTool/BuildGraph/TempStorage.cs
+++ b/Engine/Source/Programs/AutomationTool/BuildGraph/TempStorage.cs
@@ -952,7 +952,7 @@ namespace AutomationTool
 					DirectoryReference sharedNodeDir = GetDirectoryForNode(_sharedDir, nodeName);
 					FileInfo[] zipFiles = manifest!.ZipFiles.Select(x => new FileInfo(FileReference.Combine(sharedNodeDir, x.Name).FullName)).ToArray();
 					string result = ParallelUnzipFiles(zipFiles, _rootDir);
-					if (!string.IsNullOrWhiteSpace(result))
+					if (!String.IsNullOrWhiteSpace(result))
 					{
 						string logPath = CommandUtils.CombinePaths(CmdEnv.LogFolder, $"Copy Manifest - {nodeName}.log");
 						Logger.LogInformation("Saving copy log to {LogPath}", logPath);
@@ -1070,7 +1070,7 @@ namespace AutomationTool
 				select new Thread((object? indexObject) =>
 				{
 					int index = (int)indexObject!;
-					FileReference zipFileName = FileReference.Combine(zipDir, string.Format("{0}{1}.zip", zipBaseName, zipInParallel ? "-" + index.ToString("00") : ""));
+					FileReference zipFileName = FileReference.Combine(zipDir, String.Format("{0}{1}.zip", zipBaseName, zipInParallel ? "-" + index.ToString("00") : ""));
 					// don't create the zip unless we have at least one file to add
 					FileReference? file;
 					if (filesToZip.TryDequeue(out file))
@@ -1108,7 +1108,7 @@ namespace AutomationTool
 
 			zipThreads.ForEach(thread => thread.Join());
 
-			if (zipFiles.Any() && !string.IsNullOrWhiteSpace(stagingDir?.FullName))
+			if (zipFiles.Any() && !String.IsNullOrWhiteSpace(stagingDir?.FullName))
 			{
 				try
 				{
@@ -1116,7 +1116,7 @@ namespace AutomationTool
 					string logPath = CommandUtils.CombinePaths(CmdEnv.LogFolder, $"Copy Files to Temp Storage.log");
 					Logger.LogInformation("Saving copy log to {LogPath}", logPath);
 					Directory.CreateDirectory(Path.GetDirectoryName(logPath)!);
-					File.WriteAllText(logPath, string.Join(Environment.NewLine, copyResult));
+					File.WriteAllText(logPath, String.Join(Environment.NewLine, copyResult));
 					Parallel.ForEach(zipFiles, (z) => z.Delete());
 					zipFiles = new ConcurrentBag<FileInfo>(zipFiles.Select(z => new FileInfo(CommandUtils.MakeRerootedFilePath(z.FullName, stagingDir.FullName, outputDir.FullName))));
 				}
@@ -1220,7 +1220,7 @@ namespace AutomationTool
 								string logPath = CommandUtils.CombinePaths(CmdEnv.LogFolder, $"Copy Manifest - {zipFile.Name}.log");
 								Logger.LogInformation("Saving copy log to {LogPath}", logPath);
 								Directory.CreateDirectory(Path.GetDirectoryName(logPath)!);
-								File.WriteAllText(logPath, string.Join(Environment.NewLine, copyResults));
+								File.WriteAllText(logPath, String.Join(Environment.NewLine, copyResults));
 								throw;
 							}
 
@@ -1240,7 +1240,7 @@ namespace AutomationTool
 					}
 				});
 
-			return string.Join(Environment.NewLine, copyResults);
+			return String.Join(Environment.NewLine, copyResults);
 		}
 
 		/// <summary>
@@ -1450,12 +1450,12 @@ namespace AutomationTool
 
 		static string GetRoboCopyExe()
 		{
-			var result = CommandUtils.CombinePaths(Environment.SystemDirectory, "robocopy.exe");
+			string result = CommandUtils.CombinePaths(Environment.SystemDirectory, "robocopy.exe");
 			if (!CommandUtils.FileExists(result))
 			{
 				// Use Regex.Replace so we can do a case-insensitive replacement of System32
-				var sysNativeDirectory = Regex.Replace(Environment.SystemDirectory, "System32", "Sysnative", RegexOptions.IgnoreCase);
-				var sysNativeExe = CommandUtils.CombinePaths(sysNativeDirectory, "robocopy.exe");
+				string sysNativeDirectory = Regex.Replace(Environment.SystemDirectory, "System32", "Sysnative", RegexOptions.IgnoreCase);
+				string sysNativeExe = CommandUtils.CombinePaths(sysNativeDirectory, "robocopy.exe");
 				if (CommandUtils.FileExists(sysNativeExe))
 				{
 					result = sysNativeExe;
diff --git a/Engine/Source/Programs/Horde/HordeServer/Accounts/AccountsController.cs b/Engine/Source/Programs/Horde/HordeServer/Accounts/AccountsController.cs
--- a/Engine/Source/Programs/Horde/HordeServer/Accounts/AccountsController.cs
+++ b/Engine/Source/Programs/Horde/HordeServer/Accounts/AccountsController.cs
@@ -1,5 +1,6 @@
 // Copyright Epic Games, Inc. All Rights Reserved.
 
+using System;
 using System.Collections.Generic;
 using System.Threading;
 using System.Threading.Tasks;
@@ -124,7 +125,7 @@ namespace HordeServer.Accounts
 				{
 					return NotFound(accountId.Value);
 				}
-				if (request.NewPassword != null && !account.ValidatePassword(request.OldPassword ?? string.Empty))
+				if (request.NewPassword != null && !account.ValidatePassword(request.OldPassword ?? String.Empty))
 				{
 					return Unauthorized($"Invalid password for user");
 				}
diff --git a/Engine/Source/Programs/Horde/Plugins/Build/HordeServer.Build/Replicators/ReplicationService.cs b/Engine/Source/Programs/Horde/Plugins/Build/HordeServer.Build/Replicators/ReplicationService.cs
--- a/Engine/Source/Programs/Horde/Plugins/Build/HordeServer.Build/Replicators/ReplicationService.cs
+++ b/Engine/Source/Programs/Horde/Plugins/Build/HordeServer.Build/Replicators/ReplicationService.cs
@@ -118,7 +118,7 @@ namespace HordeServer.Replicators
 			}
 
 			ResiliencePipeline pipeline = new ResiliencePipelineBuilder()
-				.AddRetry(new RetryStrategyOptions { MaxRetryAttempts = int.MaxValue, OnRetry = OnRetry })
+				.AddRetry(new RetryStrategyOptions { MaxRetryAttempts = Int32.MaxValue, OnRetry = OnRetry })
 				.Build();
 
 			await pipeline.ExecuteAsync(async ctx => await _replicator.RunAsync(replicatorId, buildConfig, streamConfig, replicatorConfig, ctx), cancellationToken);
diff --git a/Engine/Source/Programs/Horde/Plugins/Compute/HordeServer.Compute/Agents/Pools/PoolConfig.cs b/Engine/Source/Programs/Horde/Plugins/Compute/HordeServer.Compute/Agents/Pools/PoolConfig.cs
--- a/Engine/Source/Programs/Horde/Plugins/Compute/HordeServer.Compute/Agents/Pools/PoolConfig.cs
+++ b/Engine/Source/Programs/Horde/Plugins/Compute/HordeServer.Compute/Agents/Pools/PoolConfig.cs
@@ -285,7 +285,7 @@ namespace HordeServer.Agents.Pools
 			if (poolConfig.Color == PoolColor.Default)
 			{
 				// Get the desired color from the properties object on the pool config
-				if (poolConfig.Properties != null && poolConfig.Properties.TryGetValue("Color", out string? colorText) && uint.TryParse(colorText, out uint colorInt))
+				if (poolConfig.Properties != null && poolConfig.Properties.TryGetValue("Color", out string? colorText) && UInt32.TryParse(colorText, out uint colorInt))
 				{
 					return GetColorValue(colorInt / 600.0f, s_colorTable);
 				}
diff --git a/Engine/Source/Programs/Horde/Plugins/Compute/HordeServer.Compute/Protos/horde/messages/session.proto.cs b/Engine/Source/Programs/Horde/Plugins/Compute/HordeServer.Compute/Protos/horde/messages/session.proto.cs
--- a/Engine/Source/Programs/Horde/Plugins/Compute/HordeServer.Compute/Protos/horde/messages/session.proto.cs
+++ b/Engine/Source/Programs/Horde/Plugins/Compute/HordeServer.Compute/Protos/horde/messages/session.proto.cs
@@ -50,7 +50,7 @@ namespace HordeCommon.Rpc.Messages
 			set => ParentIdValue = value?.ToString() ?? String.Empty;
 		}
 
-		IReadOnlyDictionary<string, int>? IAgentLease.Resources => this.Resources;
+		IReadOnlyDictionary<string, int>? IAgentLease.Resources => Resources;
 
 		LeaseState IAgentLease.State => (LeaseState)State;
 	}
diff --git a/Engine/Source/Programs/Horde/Plugins/Ddc/HordeServer.Ddc/Ddc/BucketId.cs b/Engine/Source/Programs/Horde/Plugins/Ddc/HordeServer.Ddc/Ddc/BucketId.cs
--- a/Engine/Source/Programs/Horde/Plugins/Ddc/HordeServer.Ddc/Ddc/BucketId.cs
+++ b/Engine/Source/Programs/Horde/Plugins/Ddc/HordeServer.Ddc/Ddc/BucketId.cs
@@ -57,7 +57,7 @@ namespace HordeServer.Ddc
 	sealed class BucketIdJsonConverter : JsonConverter<BucketId>
 	{
 		/// <inheritdoc/>
-		public override BucketId Read(ref Utf8JsonReader reader, Type typeToConvert, JsonSerializerOptions options) => new BucketId(reader.GetString() ?? string.Empty);
+		public override BucketId Read(ref Utf8JsonReader reader, Type typeToConvert, JsonSerializerOptions options) => new BucketId(reader.GetString() ?? String.Empty);
 
 		/// <inheritdoc/>
 		public override void Write(Utf8JsonWriter writer, BucketId value, JsonSerializerOptions options) => writer.WriteStringValue(value.ToString());
diff --git a/Engine/Source/Programs/Horde/Plugins/Ddc/HordeServer.Ddc/Ddc/BufferedPayload.cs b/Engine/Source/Programs/Horde/Plugins/Ddc/HordeServer.Ddc/Ddc/BufferedPayload.cs
--- a/Engine/Source/Programs/Horde/Plugins/Ddc/HordeServer.Ddc/Ddc/BufferedPayload.cs
+++ b/Engine/Source/Programs/Horde/Plugins/Ddc/HordeServer.Ddc/Ddc/BufferedPayload.cs
@@ -159,7 +159,7 @@ namespace HordeServer.Ddc
 		/// <summary>
 		/// If the request is smaller then MemoryBufferSize we buffer it in memory rather then as a file
 		/// </summary>
-		public long MemoryBufferSize { get; set; } = int.MaxValue;
+		public long MemoryBufferSize { get; set; } = Int32.MaxValue;
 
 		/// <summary>
 		/// The default root to create temporary buffered files under, defaults to %TEMP% or /tmp
diff --git a/Engine/Source/Programs/Horde/Plugins/Ddc/HordeServer.Ddc/Ddc/IReferenceResolver.cs b/Engine/Source/Programs/Horde/Plugins/Ddc/HordeServer.Ddc/Ddc/IReferenceResolver.cs
--- a/Engine/Source/Programs/Horde/Plugins/Ddc/HordeServer.Ddc/Ddc/IReferenceResolver.cs
+++ b/Engine/Source/Programs/Horde/Plugins/Ddc/HordeServer.Ddc/Ddc/IReferenceResolver.cs
@@ -351,7 +351,7 @@ namespace HordeServer.Ddc
 	{
 		public List<ContentId> UnresolvedReferences { get; }
 
-		public PartialReferenceResolveException(List<ContentId> unresolvedReferences) : base($"References missing: {string.Join(',', unresolvedReferences)}")
+		public PartialReferenceResolveException(List<ContentId> unresolvedReferences) : base($"References missing: {String.Join(',', unresolvedReferences)}")
 		{
 			UnresolvedReferences = unresolvedReferences;
 		}
@@ -361,7 +361,7 @@ namespace HordeServer.Ddc
 	{
 		public List<BlobId> MissingBlobs { get; }
 
-		public ReferenceIsMissingBlobsException(List<BlobId> missingBlobs) : base($"References is missing these blobs: {string.Join(',', missingBlobs)}")
+		public ReferenceIsMissingBlobsException(List<BlobId> missingBlobs) : base($"References is missing these blobs: {String.Join(',', missingBlobs)}")
 		{
 			MissingBlobs = missingBlobs;
 		}
diff --git a/Engine/Source/Programs/Horde/Plugins/Ddc/HordeServer.Ddc/Ddc/ObjectsController.cs b/Engine/Source/Programs/Horde/Plugins/Ddc/HordeServer.Ddc/Ddc/ObjectsController.cs
--- a/Engine/Source/Programs/Horde/Plugins/Ddc/HordeServer.Ddc/Ddc/ObjectsController.cs
+++ b/Engine/Source/Programs/Horde/Plugins/Ddc/HordeServer.Ddc/Ddc/ObjectsController.cs
@@ -209,11 +209,11 @@ namespace HordeServer.Ddc
 			}
 			catch (PartialReferenceResolveException e)
 			{
-				return BadRequest(new ValidationProblemDetails { Title = $"Object {id} is missing content ids", Detail = $"Following content ids are invalid: {string.Join(",", e.UnresolvedReferences)}" });
+				return BadRequest(new ValidationProblemDetails { Title = $"Object {id} is missing content ids", Detail = $"Following content ids are invalid: {String.Join(",", e.UnresolvedReferences)}" });
 			}
 			catch (ReferenceIsMissingBlobsException e)
 			{
-				return BadRequest(new ValidationProblemDetails { Title = $"Object {id} is missing blobs", Detail = $"Following blobs are missing: {string.Join(",", e.MissingBlobs)}" });
+				return BadRequest(new ValidationProblemDetails { Title = $"Object {id} is missing blobs", Detail = $"Following blobs are missing: {String.Join(",", e.MissingBlobs)}" });
 			}
 		}
 
diff --git a/Engine/Source/Programs/Horde/Plugins/Ddc/HordeServer.Ddc/Ddc/RefId.cs b/Engine/Source/Programs/Horde/Plugins/Ddc/HordeServer.Ddc/Ddc/RefId.cs
--- a/Engine/Source/Programs/Horde/Plugins/Ddc/HordeServer.Ddc/Ddc/RefId.cs
+++ b/Engine/Source/Programs/Horde/Plugins/Ddc/HordeServer.Ddc/Ddc/RefId.cs
@@ -51,7 +51,7 @@ namespace HordeServer.Ddc
 
 		public bool Equals(RefId other)
 		{
-			return string.Equals(_text, other._text, StringComparison.Ordinal);
+			return String.Equals(_text, other._text, StringComparison.Ordinal);
 		}
 
 		public override bool Equals(object? obj)
diff --git a/Engine/Source/Programs/Horde/Samples/RemoteClient/RemoteClientApp.cs b/Engine/Source/Programs/Horde/Samples/RemoteClient/RemoteClientApp.cs
--- a/Engine/Source/Programs/Horde/Samples/RemoteClient/RemoteClientApp.cs
+++ b/Engine/Source/Programs/Horde/Samples/RemoteClient/RemoteClientApp.cs
@@ -9,7 +9,6 @@ using EpicGames.Horde.Compute;
 using EpicGames.Horde.Compute.Clients;
 using EpicGames.Horde.Storage;
 using EpicGames.Horde.Storage.Bundles;
-using EpicGames.Horde.Storage.Clients;
 using EpicGames.Horde.Storage.Nodes;
 using Microsoft.Extensions.DependencyInjection;
 using Microsoft.Extensions.Logging;
diff --git a/Engine/Source/Programs/Horde/Tools/ReleaseNotes/Program.cs b/Engine/Source/Programs/Horde/Tools/ReleaseNotes/Program.cs
--- a/Engine/Source/Programs/Horde/Tools/ReleaseNotes/Program.cs
+++ b/Engine/Source/Programs/Horde/Tools/ReleaseNotes/Program.cs
@@ -6,7 +6,6 @@ using EpicGames.Horde.Server;
 using EpicGames.Perforce;
 using Microsoft.Extensions.DependencyInjection;
 using Microsoft.Extensions.Logging;
-using System.Text;
 using System.Text.RegularExpressions;
 
 namespace CreateReleaseNotes
@@ -66,7 +65,7 @@ namespace CreateReleaseNotes
 				}
 
 				logger.LogInformation("Current version: {Version}", info.ServerVersion);
-				minChange = int.Parse(match.Groups[1].Value) + 1;
+				minChange = Int32.Parse(match.Groups[1].Value) + 1;
 			}
 
 			logger.LogInformation("Finding changes after CL {MinChange}...", minChange);
@@ -130,14 +129,14 @@ namespace CreateReleaseNotes
 			}
 			lines.Add("");
 
-			string ReleaseNotesFile = $"//{perforceInfo.ClientName}/Engine/Source/Programs/Horde/Docs/ReleaseNotes.md";
-			await perforce.TryRevertAsync(-1, null, RevertOptions.None, ReleaseNotesFile);
-			await perforce.EditAsync(-1, null, EditOptions.None, ReleaseNotesFile);
+			string releaseNotesFile = $"//{perforceInfo.ClientName}/Engine/Source/Programs/Horde/Docs/ReleaseNotes.md";
+			await perforce.TryRevertAsync(-1, null, RevertOptions.None, releaseNotesFile);
+			await perforce.EditAsync(-1, null, EditOptions.None, releaseNotesFile);
 
-			WhereRecord? where = await perforce.WhereAsync(ReleaseNotesFile).FirstOrDefaultAsync();
+			WhereRecord? where = await perforce.WhereAsync(releaseNotesFile).FirstOrDefaultAsync();
 			if (where == null || String.IsNullOrEmpty(where.Path))
 			{
-				logger.LogError("Unable to get local path for file {File}", ReleaseNotesFile);
+				logger.LogError("Unable to get local path for file {File}", releaseNotesFile);
 				return 1;
 			}
 
