Remove uneeded references to System.Drawing.Common

#rnx

diff --git a/Engine/Platforms/VisionOS/Source/Programs/AutomationTool/VisionOS.Automation.csproj b/Engine/Platforms/VisionOS/Source/Programs/AutomationTool/VisionOS.Automation.csproj
--- a/Engine/Platforms/VisionOS/Source/Programs/AutomationTool/VisionOS.Automation.csproj
+++ b/Engine/Platforms/VisionOS/Source/Programs/AutomationTool/VisionOS.Automation.csproj
@@ -50,6 +50,5 @@
   </ItemGroup>
   <ItemGroup>
     <PackageReference Include="Microsoft.CSharp" Version="4.7.0" />
-    <PackageReference Include="System.Drawing.Common" Version="4.7.3" />
   </ItemGroup>
 </Project>
\ No newline at end of file
diff --git a/Engine/Source/Programs/AutomationTool/Android/AndroidPlatform.Automation.cs b/Engine/Source/Programs/AutomationTool/Android/AndroidPlatform.Automation.cs
--- a/Engine/Source/Programs/AutomationTool/Android/AndroidPlatform.Automation.cs
+++ b/Engine/Source/Programs/AutomationTool/Android/AndroidPlatform.Automation.cs
@@ -2,25 +2,18 @@
 using System;
 using System.Collections.Generic;
 using System.Linq;
-using System.Text;
 using System.IO;
 using System.Diagnostics;
-using System.Net.NetworkInformation;
 using System.Threading;
 using AutomationTool;
 using UnrealBuildTool;
 using Ionic.Zip;
 using EpicGames.Core;
 using UnrealBuildBase;
-using AutomationUtils.Automation;
 using System.Text.RegularExpressions;
 using AutomationScripts;
-using System.Drawing;
-using System.Security.Policy;
 using Microsoft.Extensions.Logging;
 
-using static AutomationTool.CommandUtils;
-
 public class AndroidPlatform : Platform
 {
 	// Maximum allowed OBB size (1 GiB, 2 GiB or 4 GiB based on project settings)
diff --git a/Engine/Source/Programs/AutomationTool/BuildGraph/BuildGraph.Automation.csproj b/Engine/Source/Programs/AutomationTool/BuildGraph/BuildGraph.Automation.csproj
--- a/Engine/Source/Programs/AutomationTool/BuildGraph/BuildGraph.Automation.csproj
+++ b/Engine/Source/Programs/AutomationTool/BuildGraph/BuildGraph.Automation.csproj
@@ -62,6 +62,5 @@
   </ItemGroup>
   <ItemGroup>
     <PackageReference Include="Microsoft.CSharp" Version="4.7.0" />
-    <PackageReference Include="System.Drawing.Common" Version="4.7.3" />
   </ItemGroup>
 </Project>
\ No newline at end of file
diff --git a/Engine/Source/Programs/AutomationTool/Gauntlet/Framework/Gauntlet.TestExecutor.cs b/Engine/Source/Programs/AutomationTool/Gauntlet/Framework/Gauntlet.TestExecutor.cs
--- a/Engine/Source/Programs/AutomationTool/Gauntlet/Framework/Gauntlet.TestExecutor.cs
+++ b/Engine/Source/Programs/AutomationTool/Gauntlet/Framework/Gauntlet.TestExecutor.cs
@@ -5,8 +5,6 @@ using System.Collections.Generic;
 using System.IO;
 using System.Linq;
 using System.Threading;
-using System.Text.RegularExpressions;
-using System.Drawing;
 using EpicGames.Core;
 using Gauntlet.Utils;
 
diff --git a/Engine/Source/Programs/AutomationTool/Gauntlet/Unreal/Base/Gauntlet.UnrealTestContext.cs b/Engine/Source/Programs/AutomationTool/Gauntlet/Unreal/Base/Gauntlet.UnrealTestContext.cs
--- a/Engine/Source/Programs/AutomationTool/Gauntlet/Unreal/Base/Gauntlet.UnrealTestContext.cs
+++ b/Engine/Source/Programs/AutomationTool/Gauntlet/Unreal/Base/Gauntlet.UnrealTestContext.cs
@@ -5,11 +5,7 @@ using System.Collections.Generic;
 using System.IO;
 using AutomationTool;
 using UnrealBuildTool;
-using System.Threading;
-using System.Text.RegularExpressions;
-using System.Drawing;
 using System.Linq;
-using System.Reflection;
 using EpicGames.Core;
 
 namespace Gauntlet
diff --git a/Engine/Source/Programs/AutomationTool/IOS/IOS.Automation.csproj b/Engine/Source/Programs/AutomationTool/IOS/IOS.Automation.csproj
--- a/Engine/Source/Programs/AutomationTool/IOS/IOS.Automation.csproj
+++ b/Engine/Source/Programs/AutomationTool/IOS/IOS.Automation.csproj
@@ -54,7 +54,6 @@
   </ItemGroup>
   <ItemGroup>
     <PackageReference Include="Microsoft.CSharp" Version="4.7.0" PrivateAssets="all" />
-    <PackageReference Include="System.Drawing.Common" Version="4.7.3" PrivateAssets="all" />
     <PackageReference Include="System.Resources.Extensions" Version="4.7.1" PrivateAssets="all" />
     <PackageReference Include="Portable.BouncyCastle" Version="1.8.6.7" PrivateAssets="all" />
   </ItemGroup>
diff --git a/Engine/Source/Programs/AutomationTool/Linux/Linux.Automation.csproj b/Engine/Source/Programs/AutomationTool/Linux/Linux.Automation.csproj
--- a/Engine/Source/Programs/AutomationTool/Linux/Linux.Automation.csproj
+++ b/Engine/Source/Programs/AutomationTool/Linux/Linux.Automation.csproj
@@ -43,6 +43,5 @@
   </ItemGroup>
   <ItemGroup>
     <PackageReference Include="Microsoft.CSharp" Version="4.7.0" />
-    <PackageReference Include="System.Drawing.Common" Version="4.7.3" />
   </ItemGroup>
 </Project>
\ No newline at end of file
diff --git a/Engine/Source/Programs/AutomationTool/Mac/Mac.Automation.csproj b/Engine/Source/Programs/AutomationTool/Mac/Mac.Automation.csproj
--- a/Engine/Source/Programs/AutomationTool/Mac/Mac.Automation.csproj
+++ b/Engine/Source/Programs/AutomationTool/Mac/Mac.Automation.csproj
@@ -42,6 +42,5 @@
   </ItemGroup>
   <ItemGroup>
     <PackageReference Include="Microsoft.CSharp" Version="4.7.0" />
-    <PackageReference Include="System.Drawing.Common" Version="4.7.3" />
   </ItemGroup>
 </Project>
\ No newline at end of file
diff --git a/Engine/Source/Programs/AutomationTool/Scripts/AutomationScripts.Automation.csproj b/Engine/Source/Programs/AutomationTool/Scripts/AutomationScripts.Automation.csproj
--- a/Engine/Source/Programs/AutomationTool/Scripts/AutomationScripts.Automation.csproj
+++ b/Engine/Source/Programs/AutomationTool/Scripts/AutomationScripts.Automation.csproj
@@ -57,7 +57,6 @@
   <ItemGroup>
     <PackageReference Include="Microsoft.CSharp" Version="4.7.0" />
     <PackageReference Include="Microsoft.VisualStudio.OLE.Interop" Version="17.11.40262" />
-    <PackageReference Include="System.Drawing.Common" Version="4.7.3" />
     <PackageReference Include="AWSSDK.S3" Version="3.3.113.2" PrivateAssets="all" />
   </ItemGroup>
 </Project>
\ No newline at end of file
diff --git a/Engine/Source/Programs/AutomationTool/Scripts/ReportItemizedExecutableCode.cs b/Engine/Source/Programs/AutomationTool/Scripts/ReportItemizedExecutableCode.cs
--- a/Engine/Source/Programs/AutomationTool/Scripts/ReportItemizedExecutableCode.cs
+++ b/Engine/Source/Programs/AutomationTool/Scripts/ReportItemizedExecutableCode.cs
@@ -6,10 +6,7 @@ using System.Collections.Generic;
 using System.IO;
 using System;
 using UnrealBuildBase;
-using IdentityModel.OidcClient;
 using System.Linq;
-using System.Drawing;
-using Microsoft.CodeAnalysis.CSharp.Syntax;
 
 namespace ReportItemizedExecutableCode.Automation
 {
diff --git a/Engine/Source/Programs/AutomationTool/SteamDeck/SteamDeck.Automation.csproj b/Engine/Source/Programs/AutomationTool/SteamDeck/SteamDeck.Automation.csproj
--- a/Engine/Source/Programs/AutomationTool/SteamDeck/SteamDeck.Automation.csproj
+++ b/Engine/Source/Programs/AutomationTool/SteamDeck/SteamDeck.Automation.csproj
@@ -42,6 +42,5 @@
   </ItemGroup>
   <ItemGroup>
     <PackageReference Include="Microsoft.CSharp" Version="4.7.0" />
-    <PackageReference Include="System.Drawing.Common" Version="4.7.3" />
   </ItemGroup>
 </Project>
\ No newline at end of file
diff --git a/Engine/Source/Programs/AutomationTool/TVOS/TVOS.Automation.csproj b/Engine/Source/Programs/AutomationTool/TVOS/TVOS.Automation.csproj
--- a/Engine/Source/Programs/AutomationTool/TVOS/TVOS.Automation.csproj
+++ b/Engine/Source/Programs/AutomationTool/TVOS/TVOS.Automation.csproj
@@ -49,6 +49,5 @@
   </ItemGroup>
   <ItemGroup>
     <PackageReference Include="Microsoft.CSharp" Version="4.7.0" />
-    <PackageReference Include="System.Drawing.Common" Version="4.7.3" />
   </ItemGroup>
 </Project>
\ No newline at end of file
diff --git a/Engine/Source/Programs/AutomationTool/Turnkey/Providers/HybridIOProvider.cs b/Engine/Source/Programs/AutomationTool/Turnkey/Providers/HybridIOProvider.cs
--- a/Engine/Source/Programs/AutomationTool/Turnkey/Providers/HybridIOProvider.cs
+++ b/Engine/Source/Programs/AutomationTool/Turnkey/Providers/HybridIOProvider.cs
@@ -2,11 +2,8 @@
 
 using System;
 using System.Collections.Generic;
-using System.Linq;
 using System.Text;
-using System.Threading.Tasks;
 using System.Text.RegularExpressions;
-using System.Drawing;
 using EpicGames.Core;
 using System.Diagnostics;
 
diff --git a/Engine/Source/Programs/AutomationTool/Turnkey/Turnkey.Automation.csproj b/Engine/Source/Programs/AutomationTool/Turnkey/Turnkey.Automation.csproj
--- a/Engine/Source/Programs/AutomationTool/Turnkey/Turnkey.Automation.csproj
+++ b/Engine/Source/Programs/AutomationTool/Turnkey/Turnkey.Automation.csproj
@@ -56,7 +56,6 @@
   </ItemGroup>
   <ItemGroup>
     <PackageReference Include="Microsoft.CSharp" Version="4.7.0" />
-    <PackageReference Include="System.Drawing.Common" Version="4.7.3" />
     <PackageReference Include="System.IO.FileSystem.AccessControl" Version="4.7.0" />
     <PackageReference Include="Google.Apis.Drive.v3" Version="1.33.1.1225" PrivateAssets="all" />
   </ItemGroup>
diff --git a/Engine/Source/Programs/AutomationTool/Win/Win.Automation.csproj b/Engine/Source/Programs/AutomationTool/Win/Win.Automation.csproj
--- a/Engine/Source/Programs/AutomationTool/Win/Win.Automation.csproj
+++ b/Engine/Source/Programs/AutomationTool/Win/Win.Automation.csproj
@@ -50,6 +50,5 @@
   </ItemGroup>
   <ItemGroup>
     <PackageReference Include="Microsoft.CSharp" Version="4.7.0" />
-    <PackageReference Include="System.Drawing.Common" Version="4.7.3" />
   </ItemGroup>
 </Project>
\ No newline at end of file
