Adding support for excluded file patterns to the <Copy> build task. E.g., <Copy From="Somewhere/*" To="SomewhereElse/*" Exclude="*.nevercopy"/>

[RN] Added support for excluded file patterns to the <Copy> build task. E.g., <Copy From="Somewhere/*" To="SomewhereElse/*" Exclude="*.nevercopy"/>

diff --git a/Engine/Source/Programs/AutomationTool/BuildGraph/Tasks/CopyTask.cs b/Engine/Source/Programs/AutomationTool/BuildGraph/Tasks/CopyTask.cs
--- a/Engine/Source/Programs/AutomationTool/BuildGraph/Tasks/CopyTask.cs
+++ b/Engine/Source/Programs/AutomationTool/BuildGraph/Tasks/CopyTask.cs
@@ -34,6 +34,12 @@ namespace AutomationTool.Tasks
 		[TaskParameter(ValidationType = TaskParameterValidationType.FileSpec)]
 		public string To { get; set; }
 
+		/// <summary>
+		/// Optional pattern(s) to exclude from the copy for example, Engine/NoCopy*.txt)
+		/// </summary>
+		[TaskParameter(Optional = true, ValidationType = TaskParameterValidationType.FileSpec)]
+		public string Exclude { get; set; }
+
 		/// <summary>
 		/// Whether or not to overwrite existing files.
 		/// </summary>
@@ -97,6 +103,30 @@ namespace AutomationTool.Tasks
 				// Build the file mapping
 				Dictionary<FileReference, FileReference> targetFileToSourceFile = FilePattern.CreateMapping(files, ref sourcePattern, ref targetPattern);
 
+				if (!String.IsNullOrEmpty(_parameters.Exclude))
+				{
+					FileFilter excludeFilter = new FileFilter();
+					foreach (string excludeRule in SplitDelimitedList(_parameters.Exclude))
+					{
+						// use "Include" because we want to match against the files we will explicitly exclude later
+						excludeFilter.Include(excludeRule);
+					}
+
+					List<FileReference> excludedFiles = new List<FileReference>();
+					foreach (KeyValuePair<FileReference, FileReference> filePair in targetFileToSourceFile)
+					{
+						if (excludeFilter.Matches(filePair.Value.ToString()))
+						{
+							excludedFiles.Add(filePair.Key);
+						}
+					}
+
+					foreach (FileReference excludedFile in excludedFiles)
+					{
+						targetFileToSourceFile.Remove(excludedFile);
+					}
+				}
+
 				// Check we got some files
 				if (targetFileToSourceFile.Count == 0)
 				{
