Add a custom structured log message for Docker errors.

#rnx

diff --git a/Engine/Source/Programs/AutomationTool/BuildGraph/Tasks/DotNetTask.cs b/Engine/Source/Programs/AutomationTool/BuildGraph/Tasks/DotNetTask.cs
--- a/Engine/Source/Programs/AutomationTool/BuildGraph/Tasks/DotNetTask.cs
+++ b/Engine/Source/Programs/AutomationTool/BuildGraph/Tasks/DotNetTask.cs
@@ -4,6 +4,7 @@ using System.Collections.Generic;
 using System.Threading.Tasks;
 using System.Xml;
 using EpicGames.Core;
+using Microsoft.Extensions.Logging;
 using UnrealBuildBase;
 
 namespace AutomationTool.Tasks
@@ -84,7 +85,8 @@ namespace AutomationTool.Tasks
 			IProcessResult result = await ExecuteAsync(dotNetFile.FullName, _parameters.Arguments, workingDir: _parameters.BaseDir, envVars: ParseEnvVars(_parameters.Environment, _parameters.EnvironmentFile));
 			if (result.ExitCode < 0 || result.ExitCode >= _parameters.ErrorLevel)
 			{
-				throw new AutomationException("Docker terminated with an exit code indicating an error ({0})", result.ExitCode);
+				Logger.LogError(KnownLogEvents.ExitCode, "Docker terminated with an exit code indicating an error ({ExitCode})", result.ExitCode);
+				throw new AutomationException("Docker terminated with an exit code indicating an error ({0})", result.ExitCode) { OutputFormat = AutomationExceptionOutputFormat.Silent };
 			}
 		}
 
