[Gauntlet] Remove libimobiledevice usage from IOS bulk content copies in favor of devicectl

diff --git a/Engine/Source/Programs/AutomationTool/Gauntlet/Platform/IOS/Gauntlet.IOSBuildSource.cs b/Engine/Source/Programs/AutomationTool/Gauntlet/Platform/IOS/Gauntlet.IOSBuildSource.cs
--- a/Engine/Source/Programs/AutomationTool/Gauntlet/Platform/IOS/Gauntlet.IOSBuildSource.cs
+++ b/Engine/Source/Programs/AutomationTool/Gauntlet/Platform/IOS/Gauntlet.IOSBuildSource.cs
@@ -32,7 +32,11 @@ namespace Gauntlet
 
 		public bool SupportsAdditionalFileCopy { get; }
 
-		public AppleBuild(UnrealTargetConfiguration InConfig, string InPackageName, string InSourcePath, BuildFlags InFlags)
+		public Dictionary<string, string> BulkContents { get; protected set; }
+
+		private const int BulkContentSearchDepth = 3;
+
+		public AppleBuild(UnrealTargetConfiguration InConfig, string InPackageName, string InSourcePath, BuildFlags InFlags, Dictionary<string, string> InBulkContents = null)
 		{
 			Configuration = InConfig;
 			PackageName = InPackageName;
@@ -40,6 +44,7 @@ namespace Gauntlet
 			Flags = InFlags;
 			SupportsAdditionalFileCopy = true;
 			IsIPAFile = Path.GetExtension(InSourcePath).Equals(".ipa", StringComparison.OrdinalIgnoreCase);
+			BulkContents = InBulkContents;
 		}
 
 		public bool CanSupportRole(UnrealTargetRole RoleType)
@@ -339,32 +344,92 @@ namespace Gauntlet
 
 			// IOS builds are always packaged, and can always replace the command line and executable (even as IPAs because we cache the unzipped app)
 			BuildFlags Flags = BuildFlags.Packaged | BuildFlags.CanReplaceCommandLine | BuildFlags.CanReplaceExecutable;
+			Dictionary<string, string> BulkContents = new Dictionary<string, string>();
 			if (BuildPath.FullName.Contains("Bulk"))
 			{
-				if(BuildPath is FileInfo)
+				if (BuildPath is FileInfo)
 				{
 					Log.Info("Bulk builds cannot be contained in an IPA file. Skipping.");
 					return null;
 				}
-				
-				Flags |= BuildFlags.Bulk;
 
-				int PayloadIndex = BuildPath.FullName.IndexOf("/Payload");
-				if(PayloadIndex < 0)
+				bool bInstallOptionalContent = Globals.Params.ParseParam("InstallOptionalContent");
+
+				// From the build, work backwards until we find the required content file
+				int SearchDepth = 0;
+				DirectoryInfo CurrentDirectory = Directory.GetParent(BuildPath.FullName);
+				FileInfo RequiredContentFile = null;
+				FileInfo OptionalContentFile = null;
+
+				while (SearchDepth < BulkContentSearchDepth)
 				{
-					Log.Info("Could not locate the payload directory for bulk build at path {BuildPath}. Skipping.", BuildPath);
+					foreach (FileInfo File in CurrentDirectory.EnumerateFiles())
+					{
+						if(File.Name.Equals("RequiredContent.txt"))
+						{
+							RequiredContentFile = File;
+						}
+
+						if(File.Name.Equals("OptionalContent.txt"))
+						{
+							OptionalContentFile = File;
+						}
+					}
+
+					if(RequiredContentFile != null && OptionalContentFile != null)
+					{
+						break;
+					}
+
+					CurrentDirectory = CurrentDirectory.Parent;
+					++SearchDepth;
+				}
+
+				if (RequiredContentFile == null)
+				{
+					Log.Info("Could not locate RequiredContent file for bulk build at path {BuildPath}. Skipping.", BuildPath);
 					return null;
 				}
-				else
+
+				if (bInstallOptionalContent && OptionalContentFile == null)
+				{
+					Log.Info("Could not locate OptionalContent file for bulk build ath path {BuildPath}. Skipping.", BuildPath);
+					return null;
+				}
+
+				string[] RequiredContents = File.ReadAllLines(RequiredContentFile.FullName);
+				string[] OptionalContents = File.ReadAllLines(OptionalContentFile.FullName);
+				
+				foreach(string RequiredContent in RequiredContents)
+				{
+					string[] Split = RequiredContent.Split(',');
+					string Source = Split[0];
+					string Target = Split[1];
+					BulkContents.Add(Source, Target);
+				}
+
+				if (bInstallOptionalContent)
 				{
-					string RootBuildDirectory = BuildPath.FullName.Substring(0, PayloadIndex);
-					string BulkContentCopyCommands = Path.Combine(RootBuildDirectory, "RequiredCommands.txt");
-					if(string.IsNullOrEmpty(BulkContentCopyCommands) || !File.Exists(BulkContentCopyCommands))
+					foreach (string OptionalContent in OptionalContents)
 					{
-						Log.Info("Could not locate bulk content command list {CommandsFile}. Skipping.", BulkContentCopyCommands);
+						string[] Split = OptionalContent.Split(',');
+						string Source = Split[0];
+						string Target = Split[1];
+						BulkContents.Add(Source, Target);
+					}
+				}
+
+				foreach (string SourceFile in BulkContents.Keys)
+				{
+					string QualifiedSourceFile = Path.Combine(BuildPath.FullName, SourceFile);
+					if(!File.Exists(QualifiedSourceFile))
+					{
+						Log.Info("Failed to find required content file {QualifiedSourceFile} for build at path {BuildPath}. Skipping", QualifiedSourceFile, BuildPath);
 						return null;
 					}
 				}
+
+				Flags |= BuildFlags.Bulk;
 			}
 			else
 			{
@@ -374,7 +439,7 @@ namespace Gauntlet
 			Log.Verbose("Found bundle id: {PackageName}", PackageName);
 			Log.Verbose("Found {Configuration} {Flags} build at {BuildPath}", Configuration, ((Flags & BuildFlags.Bulk) == BuildFlags.Bulk) ? "(bulk)" : "(not bulk)", BuildPath);
 
-			return Activator.CreateInstance(typeof(T), new object[] { Configuration, PackageName, BuildPath.FullName, Flags }) as T;
+			return Activator.CreateInstance(typeof(T), [Configuration, PackageName, BuildPath.FullName, Flags, BulkContents]) as T;
 		}
 	}
 
@@ -459,8 +524,8 @@ namespace Gauntlet
 
 	public class IOSBuild : AppleBuild
 	{
-		public IOSBuild(UnrealTargetConfiguration InConfig, string InPackageName, string InSourcePath, BuildFlags InFlags)
-			: base(InConfig, InPackageName, InSourcePath, InFlags)
+		public IOSBuild(UnrealTargetConfiguration InConfig, string InPackageName, string InSourcePath, BuildFlags InFlags, Dictionary<string, string> InBulkContents = null)
+			: base(InConfig, InPackageName, InSourcePath, InFlags, InBulkContents)
 		{ }
 
 		public override UnrealTargetPlatform Platform => UnrealTargetPlatform.IOS;
@@ -474,8 +539,8 @@ namespace Gauntlet
 
 	public class TVOSBuild : AppleBuild
 	{
-		public TVOSBuild(UnrealTargetConfiguration InConfig, string InPackageName, string InSourcePath, BuildFlags InFlags)
-			: base(InConfig, InPackageName, InSourcePath, InFlags)
+		public TVOSBuild(UnrealTargetConfiguration InConfig, string InPackageName, string InSourcePath, BuildFlags InFlags, Dictionary<string, string> InBulkContents = null)
+			: base(InConfig, InPackageName, InSourcePath, InFlags, InBulkContents)
 		{ }
 
 		public override UnrealTargetPlatform Platform => UnrealTargetPlatform.TVOS;
diff --git a/Engine/Source/Programs/AutomationTool/Gauntlet/Platform/IOS/Gauntlet.TargetDeviceIOS.cs b/Engine/Source/Programs/AutomationTool/Gauntlet/Platform/IOS/Gauntlet.TargetDeviceIOS.cs
--- a/Engine/Source/Programs/AutomationTool/Gauntlet/Platform/IOS/Gauntlet.TargetDeviceIOS.cs
+++ b/Engine/Source/Programs/AutomationTool/Gauntlet/Platform/IOS/Gauntlet.TargetDeviceIOS.cs
@@ -27,7 +27,6 @@ using static AutomationTool.CommandUtils;
  *		For now, TargetDeviceIOS uses the following tools to support automation:
  *			ios-deploy			- Third party CLI tool, still used for legacy operations and copying files to/from device - [REDACTED_PROFILE_URL]/ios-deploy
  *			devicectl			- Apple's core device framework. Has issues with copying files FROM the device, but handles just about everything else - run 'xcrun devicectl' for details
- *			libimobiledevice	- Cross-platform third party library, only used for ease of bulk content copies - https://libimobiledevice.org/
  *			
  *		In general:
  *			IF you are using xcode versions 15 or below AND iOS version 16 or below, ios-deploy will be primarily used
@@ -480,6 +479,8 @@ namespace Gauntlet
 
 			// Install the app
 			// Bad provisions can cause this to spit out errors, ignore these and let higher levels manage the exception's verbosity
+
+			DateTime StartTime = DateTime.Now;
 			using(ScopedSuspendECErrorParsing ErrorSuspension = new())
 			{
 				if(UseDeviceCtl)
@@ -496,48 +497,47 @@ namespace Gauntlet
 			{
 				throw new AutomationException("Install exited with code {0}:\n{1}", Result.ExitCode, Result.Output);
 			}
-			
-			if(Build.Flags.HasFlag(BuildFlags.Bulk))
-			{
-				int PayloadIndex = Build.SourcePath.IndexOf("/Payload");
-				DirectoryInfo RootBuildDirectory = new(Build.SourcePath.Substring(0, PayloadIndex));
-					
-				string PackageDirectory = RootBuildDirectory.Parent.FullName;
-				string BulkBuildPath = RootBuildDirectory.FullName;
-				string BulkContentsPath = Path.Combine(PackageDirectory, "BulkContents");
-				if(!Directory.Exists(BulkContentsPath))
-				{
-					throw new AutomationException("Could not locate a bulk contents directory when attempting to install a bulk build. Is your bulk build configured correctly?");
-				}
 
-				// When targeting a build in the network share, pull the build down to the host first for a faster install.
+			double RemainingTime = MaxInstallTime - (StartTime - DateTime.Now).TotalSeconds;
+
+			if (Build.Flags.HasFlag(BuildFlags.Bulk))
+			{
+				string AppPath = Build.SourcePath;
+
+				// When targeting a build in the network share, pull the build content down to the host first for a faster install.
 				// Note: it's only faster when the device is connected via USB 3.0, so WiFi is actually faster for our device farm because we use 2.0 for host connections
 				// When running on builders, we'll just copy the build directly from the share
-				string BuildVolume = GetVolumeName(PackageDirectory);
+				string BuildVolume = GetVolumeName(Build.SourcePath);
 				string CurrentVolume = GetVolumeName(Environment.CurrentDirectory);
 				if(!BuildVolume.Equals(CurrentVolume, StringComparison.OrdinalIgnoreCase) && !IsBuildMachine)
 				{
-					string BuildCache = Path.Combine(LocalCachePath, "TempBuildCache");
-					string RelativeBuildPath = BulkBuildPath.Replace(PackageDirectory, string.Empty).TrimStart('\\').Trim('/');
-					string RelativeContentPath = BulkContentsPath.Replace(PackageDirectory, string.Empty).Trim('\\').Trim('/');
+					string TempBuildCache = Path.Combine(Globals.TempDir, "TempIOSBuildCache");
+					string SourceBulkPath =  Path.GetFullPath(Path.Combine(Build.SourcePath, "..", ".."));
+					string SourceBulkContentsPath = Path.GetFullPath(Path.Combine(Build.SourcePath, "..", "..", "..", "BulkContents"));
+
+					if(!Directory.Exists(SourceBulkContentsPath))
+					{
+						throw new AutomationException("Could not locate a bulk contents directory when attempting to install a bulk build. Is your bulk build configured correctly?");
+					}
+
+					string DestinationBulkPath = Path.Combine(TempBuildCache, Path.GetFileName(SourceBulkPath));
+					string DestinationBulkContentsPath = Path.Combine(TempBuildCache, Path.GetFileName(SourceBulkContentsPath));
 
-					string DestinationBuildPath = Path.Combine(BuildCache, RelativeBuildPath);
-					string DestinationContentPath = Path.Combine(BuildCache, RelativeContentPath);
-					
 					lock(CacheMutex)
 					{
 						if(!HasCopiedNetworkBuild)
 						{
-							if (Directory.Exists(BuildCache))
+							if (Directory.Exists(TempBuildCache))
 							{
-								Directory.Delete(BuildCache, true);
+								Directory.Delete(TempBuildCache, true);
 							}
-							Directory.CreateDirectory(BuildCache);
+							Directory.CreateDirectory(TempBuildCache);
+									
+							Log.Info("Copying {Source} to {Dest}", SourceBulkPath, DestinationBulkPath);
+							SystemHelpers.CopyDirectory(SourceBulkPath, DestinationBulkPath, SystemHelpers.CopyOptions.Copy);
 
-							Log.Info("Copying {Source} to {Dest}", BulkBuildPath, DestinationBuildPath);
-							SystemHelpers.CopyDirectory(BulkBuildPath, DestinationBuildPath, SystemHelpers.CopyOptions.Mirror);
-							Log.Info("Copying {Source} to {Dest}", BulkContentsPath, DestinationContentPath);
-							SystemHelpers.CopyDirectory(BulkContentsPath, DestinationContentPath, SystemHelpers.CopyOptions.Mirror);
+							Log.Info("Copying {Source} to {Dest}", SourceBulkContentsPath, DestinationBulkContentsPath);
+							SystemHelpers.CopyDirectory(SourceBulkContentsPath, DestinationBulkContentsPath, SystemHelpers.CopyOptions.Copy);
 
 							Log.Info("Copy to host complete. Starting deploy to {DeviceName}", UUID);
 
@@ -545,24 +545,38 @@ namespace Gauntlet
 							HasCopiedNetworkBuild = true;
 						}
 
-						BulkBuildPath = DestinationBuildPath;
+						AppPath = Path.Combine(DestinationBulkPath, AppPath.Replace(SourceBulkPath, string.Empty).Trim('\\').Trim('/'));
 					}
 				}
 
-				string IDeviceFS = Path.Combine(Globals.UnrealRootDir, "Engine", "Extras", "ThirdPartyNotUE", "libimobiledevice", "mac", "idevicefs");
-				string BulkContentCopyCommand = string.Format("-u {0} -b {1} -x RequiredCommands.txt", UUID, Build.PackageName);
+				int CurrentFileIndex = 1;
+				int BulkContentFileCount = Build.BulkContents.Count;
+				Log.Info("Starting installation of Bulk Content ({FileCount} files)...", BulkContentFileCount);
 
-				if(IsBuildMachine)
+				foreach(KeyValuePair<string, string> ContentFile in Build.BulkContents)
 				{
-					// Required for forcing install via pre-paired network device instead of USB
-					BulkContentCopyCommand += " -n";
-				}
+					StartTime = DateTime.Now;
+					string FileName = Path.GetFileName(ContentFile.Key);
+					Log.Info("({Current}/{Count}) Installing {FileName}...", CurrentFileIndex, BulkContentFileCount, FileName);
 
-				Result = ExecuteDeploymentCommand(IDeviceFS, BulkContentCopyCommand, MaxInstallTime, true, BulkBuildPath, ERunOptions.NoStdOutRedirect);
+					string QualifiedSourcePath = Path.GetFullPath(Path.Combine(AppPath, ContentFile.Key));
+					IProcessResult ContentCopyResult = CopyFileToDevice(QualifiedSourcePath, ContentFile.Value, Build.PackageName, (int)RemainingTime, true);
 
-				if(Result.ExitCode != 0)
-				{
-					throw new AutomationException("Copying bulk content failed with exit code {0}. See above for info", Result.ExitCode);
+					RemainingTime -= (DateTime.Now - StartTime).TotalSeconds;
+
+					if (ContentCopyResult.ExitCode != 0)
+					{
+						if(RemainingTime <= 0)
+						{
+							throw new AutomationException("Install timed out after {MaxInstallTime} seconds", MaxInstallTime);
+						}
+						else
+						{
+							throw new AutomationException("Install of content file {0} failed with exit code {1}:\n{2}", FileName, ContentCopyResult.ExitCode, ContentCopyResult.Output);
+						}
+					}
+
+					++CurrentFileIndex;
 				}
 			}
 		}
@@ -645,6 +659,39 @@ namespace Gauntlet
 			}
 		}
 
+		public IProcessResult CopyFileToDevice(string Source, string Destination, string PackageName, int WaitTime = 60, bool bQuiet = false)
+		{
+			IProcessResult CopyResult;
+
+			if(!Destination.StartsWith("/Documents", StringComparison.OrdinalIgnoreCase) 
+			&& !Destination.StartsWith("Documents", StringComparison.OrdinalIgnoreCase))
+			{
+				Destination = Destination.StartsWith('/') 
+					? "/Documents" + Destination 
+					: "/Documents/" + Destination;  
+			}
+
+			if (UseDeviceCtl)
+			{
+				if(Destination.StartsWith('/'))
+				{
+					Destination = Destination.Substring(1);
+				}
+
+				string CopyCommand = string.Format("device copy to --device {0} --source {1} --destination {2} --domain-type appDataContainer --domain-identifier {3}",
+					UUID, Source, Destination, PackageName);
+				CopyResult = ExecuteDevicectlCommand(CopyCommand, UseDeviceID: false, WaitTime: WaitTime, bQuiet: bQuiet);
+			}
+			else
+			{
+				string CopyCommand = string.Format("--bundle_id {0} --upload={1} --to {2}",
+					PackageName, Source, Destination);
+				CopyResult = ExecuteIOSDeployCommand(CopyCommand, WaitTime, bQuiet: bQuiet);
+			}
+
+			return CopyResult;
+		}
+
 		public void PopulateDirectoryMappings(string ProjectDirectory)
 		{
 			string SavedDirectory = ProjectDirectory + "/Saved";
@@ -727,7 +774,7 @@ namespace Gauntlet
 			return new IOSAppInstance(IOSInstall, AppProcess, IOSInstall.CommandLine);
 		}
 
-		public static IProcessResult ExecuteIOSDeployCommand(string CommandLine, string Device, int WaitTime = 60, bool WarnOnTimeout = true, ERunOptions AdditionalOptions = ERunOptions.None)
+		public static IProcessResult ExecuteIOSDeployCommand(string CommandLine, string Device, int WaitTime = 60, bool WarnOnTimeout = true, ERunOptions AdditionalOptions = ERunOptions.None, bool bQuiet = false)
 		{
 			string IOSDeploy = Path.Combine(Globals.UnrealRootDir, "Engine/Extras/ThirdPartyNotUE/ios-deploy/bin/ios-deploy");
 			if (!string.IsNullOrEmpty(Device))
@@ -735,28 +782,28 @@ namespace Gauntlet
 				CommandLine = string.Format("{0} -i {1}", CommandLine, Device);
 			}
 
-			return ExecuteDeploymentCommand(IOSDeploy, CommandLine, WaitTime, WarnOnTimeout, AdditionalOptions: AdditionalOptions);
+			return ExecuteDeploymentCommand(IOSDeploy, CommandLine, WaitTime, WarnOnTimeout, null, AdditionalOptions, bQuiet);
 		}
 
-		public static ILongProcessResult ExecuteIOSDeployCommandNoWait(string CommandLine, ERunOptions AdditionalOptions = ERunOptions.None, string LocalCache = null)
+		public static ILongProcessResult ExecuteIOSDeployCommandNoWait(string CommandLine, ERunOptions AdditionalOptions = ERunOptions.None, string LocalCache = null, bool bQuiet = false)
 		{
 			string IOSDeploy = Path.Combine(Globals.UnrealRootDir, "Engine/Extras/ThirdPartyNotUE/ios-deploy/bin/ios-deploy");
-			return ExecuteDeploymentCommandNoWait(IOSDeploy, CommandLine, AdditionalOptions: AdditionalOptions, LocalCache: LocalCache);
+			return ExecuteDeploymentCommandNoWait(IOSDeploy, CommandLine, null, AdditionalOptions, LocalCache, bQuiet);
 		}
 
-		public IProcessResult ExecuteIOSDeployCommand(string CommandLine, int WaitTime = 60, bool WarnOnTimeout = true, bool UseDeviceID = true, ERunOptions AdditionalOptions = ERunOptions.None)
+		public IProcessResult ExecuteIOSDeployCommand(string CommandLine, int WaitTime = 60, bool WarnOnTimeout = true, bool UseDeviceID = true, ERunOptions AdditionalOptions = ERunOptions.None, bool bQuiet = false)
 		{
 			if (UseDeviceID && !IsDefaultDevice)
 			{
-				return ExecuteIOSDeployCommand(CommandLine, UUID, WaitTime, WarnOnTimeout, AdditionalOptions);
+				return ExecuteIOSDeployCommand(CommandLine, UUID, WaitTime, WarnOnTimeout, AdditionalOptions, bQuiet);
 			}
 			else
 			{
-				return ExecuteIOSDeployCommand(CommandLine, null, WaitTime, WarnOnTimeout, AdditionalOptions);
+				return ExecuteIOSDeployCommand(CommandLine, null, WaitTime, WarnOnTimeout, AdditionalOptions, bQuiet);
 			}
 		}
 
-		public static IProcessResult ExecuteDevicectlCommand(string CommandLine, string Device, int WaitTime = 60, bool WarnOnTimeout = true, ERunOptions AdditionalOptions = ERunOptions.None)
+		public static IProcessResult ExecuteDevicectlCommand(string CommandLine, string Device, int WaitTime = 60, bool WarnOnTimeout = true, ERunOptions AdditionalOptions = ERunOptions.None, bool bQuiet = false)
 		{
 			CommandLine = string.Format("devicectl {0}", CommandLine);
 
@@ -766,36 +813,39 @@ namespace Gauntlet
 				CommandLine = string.Format("{0} --device {1}", CommandLine, Device);
 			}
 
-			return ExecuteDeploymentCommand(XCRun, CommandLine, WaitTime, WarnOnTimeout, AdditionalOptions: AdditionalOptions);
+			return ExecuteDeploymentCommand(XCRun, CommandLine, WaitTime, WarnOnTimeout, null, AdditionalOptions, bQuiet);
 		}
 
-		public static ILongProcessResult ExecuteDevicectlCommandNoWait(string CommandLine, ERunOptions AdditionalOptions = ERunOptions.None, string LocalCache = null)
+		public static ILongProcessResult ExecuteDevicectlCommandNoWait(string CommandLine, ERunOptions AdditionalOptions = ERunOptions.None, string LocalCache = null, bool bQuiet = false)
 		{
 			CommandLine = string.Format("devicectl {0}", CommandLine);
 			string XCRun = "/usr/bin/xcrun";
-			return ExecuteDeploymentCommandNoWait(XCRun, CommandLine, AdditionalOptions: AdditionalOptions, LocalCache: LocalCache);
+			return ExecuteDeploymentCommandNoWait(XCRun, CommandLine, AdditionalOptions: AdditionalOptions, LocalCache: LocalCache, bQuiet: bQuiet);
 		}
 
-		public IProcessResult ExecuteDevicectlCommand(string CommandLine, int WaitTime = 60, bool WarnOnTimeout = true, bool UseDeviceID = true, ERunOptions AdditionalOptions = ERunOptions.None)
+		public IProcessResult ExecuteDevicectlCommand(string CommandLine, int WaitTime = 60, bool WarnOnTimeout = true, bool UseDeviceID = true, ERunOptions AdditionalOptions = ERunOptions.None, bool bQuiet = false)
 		{
 			if(UseDeviceID)
 			{
-				return ExecuteDevicectlCommand(CommandLine, UUID, WaitTime, WarnOnTimeout, AdditionalOptions);
+				return ExecuteDevicectlCommand(CommandLine, UUID, WaitTime, WarnOnTimeout, AdditionalOptions, bQuiet);
 			}
 			else
 			{
-				return ExecuteDevicectlCommand(CommandLine, null, WaitTime, WarnOnTimeout, AdditionalOptions);
+				return ExecuteDevicectlCommand(CommandLine, null, WaitTime, WarnOnTimeout, AdditionalOptions, bQuiet);
 			}
 		}
 
-		private static IProcessResult ExecuteDeploymentCommand(string Executable, string CommandLine, int WaitTime = 60, bool WarnOnTimeout = true, string WorkingDir = null, ERunOptions AdditionalOptions = ERunOptions.None)
+		private static IProcessResult ExecuteDeploymentCommand(string Executable, string CommandLine, int WaitTime = 60, bool WarnOnTimeout = true, string WorkingDir = null, ERunOptions AdditionalOptions = ERunOptions.None, bool bQuiet = false)
 		{
 			if (!AdditionalOptions.HasFlag(ERunOptions.UseShellExecute) && !File.Exists(Executable))
 			{
 				throw new AutomationException("Unable to find deployment binary at {0}", Executable);
 			}
 
-			Log.Info("{DeploymentExecutable} executing '{Command}'", Executable, CommandLine);
+			if(!bQuiet)
+			{
+				Log.Info("{DeploymentExecutable} executing '{Command}'", Executable, CommandLine);
+			}
 
 			ERunOptions RunOptions = ERunOptions.NoWaitForExit | AdditionalOptions;
 			RunOptions |= Log.IsVeryVerbose
@@ -843,21 +893,23 @@ namespace Gauntlet
 			return Result;
 		}
 
-		private static ILongProcessResult ExecuteDeploymentCommandNoWait(string Executable, string CommandLine, string WorkingDir = null, ERunOptions AdditionalOptions = ERunOptions.None, string LocalCache = null)
+		private static ILongProcessResult ExecuteDeploymentCommandNoWait(string Executable, string CommandLine, string WorkingDir = null, ERunOptions AdditionalOptions = ERunOptions.None, string LocalCache = null, bool bQuiet = false)
 		{
 			if (!AdditionalOptions.HasFlag(ERunOptions.UseShellExecute) && !File.Exists(Executable))
 			{
 				throw new AutomationException("Unable to find deployment binary at {0}", Executable);
 			}
 
-			Log.Info("{DeploymentExecutable} executing '{Command}'", Executable, CommandLine);
-
+			if(!bQuiet)
+			{
+				Log.Info("{DeploymentExecutable} executing '{Command}'", Executable, CommandLine);
+			}
+			
 			ERunOptions RunOptions = ERunOptions.NoWaitForExit | AdditionalOptions;
 			RunOptions |= Log.IsVeryVerbose
 				? ERunOptions.AllowSpew
 				: ERunOptions.NoLoggingOfRunCommand;
 
-
 			ILongProcessResult Result = null;
 			try
 			{
