Add support to gauntlet input action system for button taps

diff --git a/Engine/Source/Programs/AutomationTool/Gauntlet/Framework/Utils/Gauntlet.InputActions.cs b/Engine/Source/Programs/AutomationTool/Gauntlet/Framework/Utils/Gauntlet.InputActions.cs
--- a/Engine/Source/Programs/AutomationTool/Gauntlet/Framework/Utils/Gauntlet.InputActions.cs
+++ b/Engine/Source/Programs/AutomationTool/Gauntlet/Framework/Utils/Gauntlet.InputActions.cs
@@ -577,13 +577,15 @@ namespace Gauntlet.Utils
 		public EInputEvent KeyAction;
 		public double XDelta;
 		public double YDelta;
+		public bool Tapped;
 
-		public InputAction(FKey Inkey, EInputEvent InInputEvent, double InXDelta = 0, double InYDelta = 0)
+		public InputAction(FKey Inkey, EInputEvent InInputEvent, double InXDelta = 0, double InYDelta = 0, bool InTapped = false)
 		{
 			this.Key = Inkey;
 			this.KeyAction = InInputEvent;
 			this.XDelta = InXDelta;
 			this.YDelta = InYDelta;
+			Tapped = InTapped;
 		}
 
 	}
