Fix for Editor failing to launch iOS build on device

diff --git a/Engine/Source/Programs/AutomationTool/IOS/IOSPlatform.Automation.cs b/Engine/Source/Programs/AutomationTool/IOS/IOSPlatform.Automation.cs
--- a/Engine/Source/Programs/AutomationTool/IOS/IOSPlatform.Automation.cs
+++ b/Engine/Source/Programs/AutomationTool/IOS/IOSPlatform.Automation.cs
@@ -2059,27 +2059,42 @@ public class IOSPlatform : ApplePlatform
 				BundleIdentifier = Contents.Substring(Pos, EndPos - Pos);
 			}
 
-			string Program = GetPathToLibiMobileDeviceTool("idevicedebug"); ;
-			string Arguments = " -u '" + Params.DeviceNames[0] + "'";
-			Arguments += " --detach";
-			Arguments = GetLibimobileDeviceNetworkedArgument(Arguments, Params.DeviceNames[0]);
-			Arguments += " run '" + BundleIdentifier + "'";
-			
-			// ClientCmdLine is only relevant when running on a Mac
+			string Program;
+			string Arguments;
 			if (OperatingSystem.IsMacOS())
 			{
+				// As of iOS17, Apple changed the remote debugger protocol and libimobiledevice has stopped working for launching apps.
+				// Instead, as of Xcode14(ish), Apple added a new "xcrun devicectl" cli to do pretty much everything libimobiledevice/ios-deploy did.
+				// For now, update the app launch implementation to use devicectl. Windows may need to switch to a different 3rd party implementation.
+				Program = "xcrun";
+				Arguments = "devicectl device process launch";
+				Arguments += " --terminate-existing";	// if it's already running on device, kill it first
+				Arguments += " --console";	// attach to the console so we can get log output
+				Arguments += " --device " + Params.DeviceNames[0];
+				Arguments += " \"" + BundleIdentifier + "\"";
+				
+				// ClientCmdLine is only relevant when running on a Mac
 				Arguments += " " + ClientCmdLine;
 			}
+			else
+			{
+				Program = GetPathToLibiMobileDeviceTool("idevicedebug");
+				Arguments = " -u '" + Params.DeviceNames[0] + "'";
+				Arguments += " --detach";
+				Arguments = GetLibimobileDeviceNetworkedArgument(Arguments, Params.DeviceNames[0]);
+				Arguments += " run '" + BundleIdentifier + "'";
+			}
 
 			IProcessResult ClientProcess = Run(Program, Arguments, null, ClientRunFlags);
-			if (ClientProcess.ExitCode == -1)
+			
+			// This failure mode is now valid only for use of libimobiledevice on Windows.
+			if (OperatingSystem.IsWindows() && ClientProcess.ExitCode == -1)
 			{
 				Console.WriteLine("The application {0} has been installed on the device {1} but it cannot be launched automatically because the device does not contain the required developer software. You can launch {0} the manually by clicking its icon on the device.", BundleIdentifier, Params.DeviceNames[0]);
 				Console.WriteLine("To install the developer software tools, connect it to a Mac running Xcode, open the Devices and Simulators window and wait for the tools to be installed.");
 				IProcessResult Result = new ProcessResult("DummyApp", null, false);
 				Result.ExitCode = 0;
 				return Result;
-
 			}
 			return ClientProcess;
 
