Expose SymGenTasks.cs for iOS psym generation support
Move raddbgi_breakpad_from_pdb.exe to Extras/rad/Binaries/Win64 (source available at [REDACTED_PROFILE_URL]/raddebugger/blob/master/src/rdi_breakpad_from_pdb/rdi_breakpad_from_pdb_main.c)
Add symbol_path_fixer.exe to Extras/rad/Binaries/Win64 (source in Source/ThirdParty/Breakpad/rad)

diff --git a/Engine/Source/Programs/AutomationTool/BuildGraph/Tasks/SymGenTask.cs b/Engine/Source/Programs/AutomationTool/BuildGraph/Tasks/SymGenTask.cs
new file mode 100644
--- /dev/null
+++ b/Engine/Source/Programs/AutomationTool/BuildGraph/Tasks/SymGenTask.cs
@@ -0,0 +1,420 @@
+// Copyright Epic Games, Inc. All Rights Reserved.
+
+using AutomationTool;
+using EpicGames.BuildGraph;
+using System;
+using System.IO;
+using System.Collections.Concurrent;
+using System.Collections.Generic;
+using System.Linq;
+using System.Text;
+using System.Threading.Tasks;
+using System.Xml;
+using EpicGames.Core;
+using UnrealBuildTool;
+using UnrealBuildBase;
+using Microsoft.Extensions.Logging;
+
+using static AutomationTool.CommandUtils;
+
+namespace BuildGraph.Tasks
+{
+	/// <summary>
+	/// Parameters for a task that generates debugging symbols from a set of files
+	/// </summary>
+	public class SymGenTaskParameters
+	{
+		/// <summary>
+		/// List of file specifications separated by semicolons (eg. *.cpp;Engine/.../*.bat), or the name of a tag set
+		/// </summary>
+		[TaskParameter(ValidationType = TaskParameterValidationType.FileSpec)]
+		public string Files { get; set; }
+
+		/// <summary>
+		/// Tag to be applied to build products of this task
+		/// </summary>
+		[TaskParameter(Optional = true, ValidationType = TaskParameterValidationType.TagList)]
+		public string Tag { get; set; }
+
+		/// <summary>
+		/// If set, this will use the rad debugger pdb symbol dumper as well as the rad symbol_path_fixer.
+		/// </summary>
+		[TaskParameter(Optional = true)]
+		public bool UseRadSym { get; set; } = false;
+
+}
+
+	/// <summary>
+	/// Generates a portable symbol dump file from the specified binaries
+	/// </summary>
+	[TaskElement("SymGen", typeof(SymGenTaskParameters))]
+	public class SymGenTask : CustomTask
+	{
+		/// <summary>
+		/// Parameters for this task
+		/// </summary>
+		SymGenTaskParameters Parameters;
+
+		/// <summary>
+		/// Construct a spawn task
+		/// </summary>
+		/// <param name="InParameters">Parameters for the task</param>
+		public SymGenTask(SymGenTaskParameters InParameters)
+		{
+			Parameters = InParameters;
+		}
+
+		static UnrealArchitectures ArchitecturesInBinary(FileReference Binary)
+		{
+			if (BuildHostPlatform.Current.Platform != UnrealTargetPlatform.Mac)
+			{
+				return null;
+			}
+
+			List<UnrealArch> Arches = new();
+			string Output = Utils.RunLocalProcessAndReturnStdOut("sh", $"-c 'file \"{Binary.FullName}\"'");
+			if (Output.Contains("arm64", StringComparison.InvariantCulture))
+			{
+				Arches.Add(UnrealArch.Arm64);
+			}
+			if (Output.Contains("x86_64", StringComparison.InvariantCulture))
+			{
+				Arches.Add(UnrealArch.X64);
+			}
+			return new UnrealArchitectures(Arches);
+		}
+
+		/// <summary>
+		/// Execute the task.
+		/// </summary>
+		/// <param name="Job">Information about the current job</param>
+		/// <param name="BuildProducts">Set of build products produced by this node.</param>
+		/// <param name="TagNameToFileSet">Mapping from tag names to the set of files they include</param>
+		public override void Execute(JobContext Job, HashSet<FileReference> BuildProducts, Dictionary<string, HashSet<FileReference>> TagNameToFileSet)
+		{
+			bool bUseRadSym = Parameters.UseRadSym;
+
+			// Path to Breakpad's dump_syms executable
+			string SymbolDumperExecutable = null;
+
+			// Find the matching files
+			FileReference[] SourceFiles = ResolveFilespec(Unreal.RootDirectory, Parameters.Files, TagNameToFileSet).OrderBy(x => x.FullName).ToArray();
+
+			string RadSymDymperExecuable = Unreal.RootDirectory + @"\Engine\Extras\rad\Binaries\Win64\raddbgi_breakpad_from_pdb.exe";
+			string RadProcessSymExecuable = Unreal.RootDirectory + @"\Engine\Extras\rad\Binaries\Win64\symbol_path_fixer.exe";
+
+			// Filter out all the symbol files
+			FileReference[] SymbolSourceFiles;
+			string WorkingDirectory = null;
+			if (BuildHostPlatform.Current.Platform == UnrealTargetPlatform.Win64)
+			{
+				SymbolDumperExecutable = Unreal.RootDirectory + @"\Engine\Source\ThirdParty\Breakpad\src\tools\windows\binaries\dump_syms.exe";
+				string[] SymbolFileExtensions = { ".pdb", ".nss", ".nrs" };
+				SymbolSourceFiles = SourceFiles.Where(x => SymbolFileExtensions.Contains(x.GetExtension())).ToArray();
+				// set working dir to find our version of msdia140.dll
+				WorkingDirectory = Unreal.RootDirectory + @"\Engine\Binaries\Win64";
+			}
+			else if (BuildHostPlatform.Current.Platform == UnrealTargetPlatform.Mac)
+			{
+				SymbolDumperExecutable = Unreal.RootDirectory + "/Engine/Source/ThirdParty/Breakpad/src/tools/mac/binaries/dump_syms";
+				List<FileReference> Files = SourceFiles.Where(x => x.HasExtension(".dSYM")).ToList();
+
+				// find any zipped bundles
+				Directory.CreateDirectory(Unreal.RootDirectory + "/Engine/Intermediate/Unzipped");
+				FileReference[] ZippedFiles = SourceFiles.Where(x => x.FullName.Contains(".dSYM.zip", StringComparison.InvariantCulture)).ToArray();
+				foreach (FileReference SourceFile in ZippedFiles)
+				{
+					string[] UnzippedFiles = CommandUtils.UnzipFiles(SourceFile.FullName, Unreal.RootDirectory + "/Engine/Intermediate/Unzipped").ToArray();
+					Files.Add(new FileReference(Unreal.RootDirectory + "/Engine/Intermediate/Unzipped/" + SourceFile.GetFileNameWithoutExtension()));
+				}
+				foreach (FileReference SourceFile in Files)
+				{
+					Logger.LogInformation("Source File: {Arg0}", SourceFile.FullName);
+				}
+				SymbolSourceFiles = Files.Where(x => x.HasExtension(".dSYM")).ToArray();
+			}
+			else if (BuildHostPlatform.Current.Platform == UnrealTargetPlatform.Linux)
+			{
+				SymbolDumperExecutable = Unreal.RootDirectory + "/Engine/Binaries/Linux/dump_syms";
+				string[] SymbolFileExtensions = { ".debug" };
+				SymbolSourceFiles = SourceFiles.Where(x => SymbolFileExtensions.Contains(x.GetExtension())).ToArray();
+			}
+			else
+			{
+				throw new AutomationException("Symbol generation failed: Unknown platform {0}", BuildHostPlatform.Current.Platform);
+			}
+
+			// Remove any existing symbol files
+			foreach (string FileName in SymbolSourceFiles.Select(x => Path.ChangeExtension(x.FullName, ".psym")))
+			{
+				if (File.Exists(FileName))
+				{
+					try
+					{
+						File.Delete(FileName);
+					}
+					catch (Exception Ex)
+					{
+						throw new AutomationException("Symbol generation failed: Unable to delete existing symbol file: \"{0}\". Error: {1}", FileName, Ex.Message.TrimEnd());
+					}
+				}
+
+			}
+
+			if (SymbolSourceFiles.Length == 0)
+			{
+				Logger.LogInformation("No symbol files to convert.");
+			}
+
+			// Generate portable symbols from the symbol source files
+			ConcurrentBag<FileReference> SymbolFiles = new ConcurrentBag<FileReference>();
+
+			Parallel.ForEach(SymbolSourceFiles, (SourceFile) =>
+			{
+				string SymbolFileName = Path.ChangeExtension(SourceFile.FullName, ".psym");
+				string RadSymbolTemp = Path.ChangeExtension(SourceFile.FullName, ".radpsym");
+
+				// Check if higher priority debug file or binary already created symbols 
+				if (File.Exists(SymbolFileName))
+				{
+					return;
+				}
+
+				Logger.LogInformation("Dumping Symbols: {Arg0} to {SymbolFileName}", SourceFile.FullName, SymbolFileName);
+
+				string DumpSymsArgs;
+
+				string SymbolDumperExeForFile = SymbolDumperExecutable;
+				if (bUseRadSym &&
+					SourceFile.GetExtension() == ".pdb")
+				{
+					SymbolDumperExeForFile = RadSymDymperExecuable;
+					DumpSymsArgs = "-pdb:" + SourceFile.FullName + " -out:" + RadSymbolTemp + " -exe:" + SourceFile.FullName;
+				}
+				else
+				{
+					string ExtraOptions = "";
+					if (BuildHostPlatform.Current.Platform == UnrealTargetPlatform.Mac)
+					{
+						// dump_syms has a bug where if a universal binary is fed into it, on an Intel mac, it will fail to find the current architecture
+						// (but not on Arm macs). Specify the host architecture as a param to cause expected behavior (until we make one output per Arch)
+						if (ArchitecturesInBinary(SourceFile).bIsMultiArch)
+						{
+							// ExtraOptions = $"-a {MacExports.HostArchitecture.AppleName} ";
+							// Since IBs are universal and we typically only care about arm symbols, force
+							// the arch to always be arm.
+							ExtraOptions = $"-a arm64 ";
+						}
+					}
+
+					DumpSymsArgs = ExtraOptions + SourceFile.FullName;
+				}
+
+				IProcessResult result = CommandUtils.Run(SymbolDumperExeForFile, DumpSymsArgs, null, CommandUtils.ERunOptions.AppMustExist, null, FilterSpew, null, WorkingDirectory);
+				if (result.ExitCode == 0)
+				{
+					StringBuilder ProcessedSymbols = null;
+					if (bUseRadSym)
+					{
+						// rad dumper outputs to a file, we thunk to a custom exe to do symbol munging for speed.
+						CommandUtils.Run(RadProcessSymExecuable, RadSymbolTemp + " " + Unreal.RootDirectory.FullName, null, CommandUtils.ERunOptions.AppMustExist, null, FilterSpew, null, WorkingDirectory);
+
+						File.Move(RadSymbolTemp, SymbolFileName);
+						SymbolFiles.Add(new FileReference(SymbolFileName));
+
+					}
+					else
+					{
+						try
+						{
+							// Process the symbols
+							using (StringReader Reader = new StringReader(result.Output))
+							{
+								ProcessSymbols(SymbolFileName, Reader, out ProcessedSymbols);
+							}
+						}
+						catch (OutOfMemoryException)
+						{
+							// If we catch an OOM, it is too large to turn into a string.
+							// Write to a file and then load it into a string.
+							string TempFileName = Path.Combine(Path.GetTempPath(), Path.GetTempFileName());
+							FileReference SymbolFile = (result as ProcessResult).WriteOutputToFile(TempFileName);
+
+							try
+							{
+								using (StreamReader Reader = new StreamReader(SymbolFile.FullName))
+								{
+									ProcessSymbols(SymbolFileName, Reader, out ProcessedSymbols);
+								}
+							}
+							finally
+							{
+								FileReference.Delete(SymbolFile);
+							}
+						}
+						catch (Exception Ex)
+						{
+							// There was a problem generating symbols with the dump_syms tool
+							throw new AutomationException($"Symbol generation failed: Error Generating Symbols for {SymbolFileName}, Error: {ExceptionUtils.FormatException(Ex)}");
+						}
+					}
+
+					if (ProcessedSymbols != null && ProcessedSymbols.Length > 0)
+					{
+						using (StreamWriter Writer = new StreamWriter(SymbolFileName))
+						{
+							Writer.Write(ProcessedSymbols);
+						}
+						SymbolFiles.Add(new FileReference(SymbolFileName));
+					}
+				}
+				else
+				{
+					if (BuildHostPlatform.Current.Platform != UnrealTargetPlatform.Mac)
+					{
+						// If we fail, lets re-run with a verbose, -v to check for the error we are seeing
+						// -v not available on Mac
+						CommandUtils.Run(SymbolDumperExecutable, "-v " + SourceFile.FullName, null, CommandUtils.ERunOptions.AppMustExist, null, null, null, WorkingDirectory);
+					}
+
+					// There was a problem generating symbols with the dump_syms tool
+					throw new AutomationException("Symbol generation failed: Error Generating Symbols: {0}", SymbolFileName);
+				}
+			});
+
+			// Apply the optional tag to the build products
+			foreach (string TagName in FindTagNamesFromList(Parameters.Tag))
+			{
+				FindOrAddTagSet(TagNameToFileSet, TagName).UnionWith(SymbolFiles);
+			}
+
+			// Add them to the list of build products
+			BuildProducts.UnionWith(SymbolFiles);
+		}
+
+		/// <summary>
+		/// Processes the raw symbol dump
+		/// </summary>
+		static bool ProcessSymbols(string SymbolFileName, TextReader Reader, out StringBuilder ProcessedSymbols)
+		{
+			char[] FieldSeparator = { ' ' };
+			string RootDirectory = CommandUtils.ConvertSeparators(PathSeparator.Slash, Unreal.RootDirectory.FullName).TrimEnd('/'); ;
+
+			ProcessedSymbols = new StringBuilder();
+
+			string Line;
+			bool bSawModule = false;
+			while ((Line = Reader.ReadLine()) != null)
+			{
+				if (Line.Contains(" = ", StringComparison.InvariantCulture))
+				{
+					Logger.LogInformation("{Text}", Line);
+					continue;
+				}
+				// Ignore any output from symbol dump before MODULE, these may included erroneous warnings, etc
+				if (!bSawModule)
+				{
+					if (!Line.StartsWith("MODULE", StringComparison.InvariantCulture))
+					{
+						continue;
+					}
+
+					bSawModule = true;
+				}
+
+				string NewLine = Line;
+
+				// Process source reference FILE blocks
+				if (Line.StartsWith("FILE", StringComparison.InvariantCulture))
+				{
+					string[] Fields = Line.Split(FieldSeparator, 3);
+
+					string FileName = CommandUtils.ConvertSeparators(PathSeparator.Slash, Fields[2]);
+
+					// If the file exists locally, and is within the root, convert path
+					if (File.Exists(FileName) && FileName.StartsWith(RootDirectory, StringComparison.OrdinalIgnoreCase))
+					{
+						// Restore proper filename case on Windows (the symbol dump filenames are all lowercase)
+						if (BuildHostPlatform.Current.Platform == UnrealTargetPlatform.Win64)
+						{
+							FileName = FileUtils.FindCorrectCase(new FileInfo(FileName)).FullName;
+						}
+
+						// Shave off the root directory
+						NewLine = String.Format("FILE {0} {1}", Fields[1], FileName.Substring(RootDirectory.Length + 1).Replace('\\', '/'));
+					}
+				}
+
+				ProcessedSymbols.AppendLine(NewLine);
+			}
+
+			return true;
+
+		}
+
+		/// <summary>
+		///  Filters the output from the dump_syms executable, which depending on the platform can be pretty spammy
+		/// </summary>
+		string FilterSpew(string Message)
+		{
+			foreach (string FilterString in OutputFilterStrings)
+			{
+				if (Message.Contains(FilterString, StringComparison.InvariantCulture))
+				{
+					return null;
+				}
+			}
+
+			return Message;
+		}
+
+		/// <summary>
+		/// Array of source strings to filter from output
+		/// </summary>
+		static readonly string[] OutputFilterStrings = new string[] { "the DIE at offset", "warning: function", "warning: failed", ": in compilation unit" };
+
+		/// <summary>
+		/// Output this task out to an XML writer.
+		/// </summary>
+		public override void Write(XmlWriter Writer)
+		{
+			Write(Writer, Parameters);
+		}
+
+		/// <summary>
+		/// Find all the tags which are used as inputs to this task
+		/// </summary>
+		/// <returns>The tag names which are read by this task</returns>
+		public override IEnumerable<string> FindConsumedTagNames()
+		{
+			return FindTagNamesFromFilespec(Parameters.Files);
+		}
+
+		/// <summary>
+		/// Find all the tags which are modified by this task
+		/// </summary>
+		/// <returns>The tag names which are modified by this task</returns>
+		public override IEnumerable<string> FindProducedTagNames()
+		{
+			return FindTagNamesFromList(Parameters.Tag);
+		}
+	}
+}
+
+namespace BuildScripts.Automation
+{
+	class GeneratePsyms : BuildCommand
+	{
+		public override ExitCode Execute()
+		{
+			BuildGraph.Tasks.SymGenTaskParameters Params = new BuildGraph.Tasks.SymGenTaskParameters();
+			Params.Files = ParseRequiredStringParam("Files");
+
+			BuildGraph.Tasks.SymGenTask Task = new BuildGraph.Tasks.SymGenTask(Params);
+			Task.Execute(null, new HashSet<FileReference>(), new Dictionary<string, HashSet<FileReference>>());
+
+			return ExitCode.Success;
+		}
+	}
+}
+
