Resolve static analysis errors in SymGenTask.cs

diff --git a/Engine/Source/Programs/AutomationTool/BuildGraph/Tasks/SymGenTask.cs b/Engine/Source/Programs/AutomationTool/BuildGraph/Tasks/SymGenTask.cs
--- a/Engine/Source/Programs/AutomationTool/BuildGraph/Tasks/SymGenTask.cs
+++ b/Engine/Source/Programs/AutomationTool/BuildGraph/Tasks/SymGenTask.cs
@@ -1,7 +1,6 @@
 // Copyright Epic Games, Inc. All Rights Reserved.
 
 using AutomationTool;
-using EpicGames.BuildGraph;
 using System;
 using System.IO;
 using System.Collections.Concurrent;
@@ -15,8 +14,6 @@ using UnrealBuildTool;
 using UnrealBuildBase;
 using Microsoft.Extensions.Logging;
 
-using static AutomationTool.CommandUtils;
-
 namespace BuildGraph.Tasks
 {
 	/// <summary>
@@ -53,91 +50,91 @@ namespace BuildGraph.Tasks
 		/// <summary>
 		/// Parameters for this task
 		/// </summary>
-		SymGenTaskParameters Parameters;
+		readonly SymGenTaskParameters _parameters;
 
 		/// <summary>
 		/// Construct a spawn task
 		/// </summary>
-		/// <param name="InParameters">Parameters for the task</param>
-		public SymGenTask(SymGenTaskParameters InParameters)
+		/// <param name="parameters">Parameters for the task</param>
+		public SymGenTask(SymGenTaskParameters parameters)
 		{
-			Parameters = InParameters;
+			_parameters = parameters;
 		}
 
-		static UnrealArchitectures ArchitecturesInBinary(FileReference Binary)
+		static UnrealArchitectures ArchitecturesInBinary(FileReference binary)
 		{
 			if (BuildHostPlatform.Current.Platform != UnrealTargetPlatform.Mac)
 			{
 				return null;
 			}
 
-			List<UnrealArch> Arches = new();
-			string Output = Utils.RunLocalProcessAndReturnStdOut("sh", $"-c 'file \"{Binary.FullName}\"'");
-			if (Output.Contains("arm64", StringComparison.InvariantCulture))
+			List<UnrealArch> arches = new();
+			string output = Utils.RunLocalProcessAndReturnStdOut("sh", $"-c 'file \"{binary.FullName}\"'");
+			if (output.Contains("arm64", StringComparison.InvariantCulture))
 			{
-				Arches.Add(UnrealArch.Arm64);
+				arches.Add(UnrealArch.Arm64);
 			}
-			if (Output.Contains("x86_64", StringComparison.InvariantCulture))
+			if (output.Contains("x86_64", StringComparison.InvariantCulture))
 			{
-				Arches.Add(UnrealArch.X64);
+				arches.Add(UnrealArch.X64);
 			}
-			return new UnrealArchitectures(Arches);
+			return new UnrealArchitectures(arches);
 		}
 
 		/// <summary>
 		/// Execute the task.
 		/// </summary>
-		/// <param name="Job">Information about the current job</param>
-		/// <param name="BuildProducts">Set of build products produced by this node.</param>
-		/// <param name="TagNameToFileSet">Mapping from tag names to the set of files they include</param>
-		public override void Execute(JobContext Job, HashSet<FileReference> BuildProducts, Dictionary<string, HashSet<FileReference>> TagNameToFileSet)
+		/// <param name="job">Information about the current job</param>
+		/// <param name="buildProducts">Set of build products produced by this node.</param>
+		/// <param name="tagNameToFileSet">Mapping from tag names to the set of files they include</param>
+		public override void Execute(JobContext job, HashSet<FileReference> buildProducts, Dictionary<string, HashSet<FileReference>> tagNameToFileSet)
 		{
-			bool bUseRadSym = Parameters.UseRadSym;
+			bool bUseRadSym = _parameters.UseRadSym;
 
 			// Path to Breakpad's dump_syms executable
-			string SymbolDumperExecutable = null;
+			string symbolDumperExecutable = null;
 
 			// Find the matching files
-			FileReference[] SourceFiles = ResolveFilespec(Unreal.RootDirectory, Parameters.Files, TagNameToFileSet).OrderBy(x => x.FullName).ToArray();
+			FileReference[] sourceFiles = ResolveFilespec(Unreal.RootDirectory, _parameters.Files, tagNameToFileSet).OrderBy(x => x.FullName).ToArray();
 
-			string RadSymDymperExecuable = Unreal.RootDirectory + @"\Engine\Extras\rad\Binaries\Win64\raddbgi_breakpad_from_pdb.exe";
-			string RadProcessSymExecuable = Unreal.RootDirectory + @"\Engine\Extras\rad\Binaries\Win64\symbol_path_fixer.exe";
+			string radSymDymperExecuable = Unreal.RootDirectory + @"\Engine\Extras\rad\Binaries\Win64\raddbgi_breakpad_from_pdb.exe";
+			string radProcessSymExecuable = Unreal.RootDirectory + @"\Engine\Extras\rad\Binaries\Win64\symbol_path_fixer.exe";
 
 			// Filter out all the symbol files
-			FileReference[] SymbolSourceFiles;
-			string WorkingDirectory = null;
+			FileReference[] symbolSourceFiles;
+			string workingDirectory = null;
 			if (BuildHostPlatform.Current.Platform == UnrealTargetPlatform.Win64)
 			{
-				SymbolDumperExecutable = Unreal.RootDirectory + @"\Engine\Source\ThirdParty\Breakpad\src\tools\windows\binaries\dump_syms.exe";
-				string[] SymbolFileExtensions = { ".pdb", ".nss", ".nrs" };
-				SymbolSourceFiles = SourceFiles.Where(x => SymbolFileExtensions.Contains(x.GetExtension())).ToArray();
+				symbolDumperExecutable = Unreal.RootDirectory + @"\Engine\Source\ThirdParty\Breakpad\src\tools\windows\binaries\dump_syms.exe";
+				string[] symbolFileExtensions = { ".pdb", ".nss", ".nrs" };
+				symbolSourceFiles = sourceFiles.Where(x => symbolFileExtensions.Contains(x.GetExtension())).ToArray();
 				// set working dir to find our version of msdia140.dll
-				WorkingDirectory = Unreal.RootDirectory + @"\Engine\Binaries\Win64";
+				workingDirectory = Unreal.RootDirectory + @"\Engine\Binaries\Win64";
 			}
 			else if (BuildHostPlatform.Current.Platform == UnrealTargetPlatform.Mac)
 			{
-				SymbolDumperExecutable = Unreal.RootDirectory + "/Engine/Source/ThirdParty/Breakpad/src/tools/mac/binaries/dump_syms";
-				List<FileReference> Files = SourceFiles.Where(x => x.HasExtension(".dSYM")).ToList();
+				symbolDumperExecutable = Unreal.RootDirectory + "/Engine/Source/ThirdParty/Breakpad/src/tools/mac/binaries/dump_syms";
+				List<FileReference> files = sourceFiles.Where(x => x.HasExtension(".dSYM")).ToList();
 
 				// find any zipped bundles
 				Directory.CreateDirectory(Unreal.RootDirectory + "/Engine/Intermediate/Unzipped");
-				FileReference[] ZippedFiles = SourceFiles.Where(x => x.FullName.Contains(".dSYM.zip", StringComparison.InvariantCulture)).ToArray();
-				foreach (FileReference SourceFile in ZippedFiles)
+				FileReference[] zippedFiles = sourceFiles.Where(x => x.FullName.Contains(".dSYM.zip", StringComparison.InvariantCulture)).ToArray();
+				foreach (FileReference sourceFile in zippedFiles)
 				{
-					string[] UnzippedFiles = CommandUtils.UnzipFiles(SourceFile.FullName, Unreal.RootDirectory + "/Engine/Intermediate/Unzipped").ToArray();
-					Files.Add(new FileReference(Unreal.RootDirectory + "/Engine/Intermediate/Unzipped/" + SourceFile.GetFileNameWithoutExtension()));
+					string[] unzippedFiles = CommandUtils.UnzipFiles(sourceFile.FullName, Unreal.RootDirectory + "/Engine/Intermediate/Unzipped").ToArray();
+					files.Add(new FileReference(Unreal.RootDirectory + "/Engine/Intermediate/Unzipped/" + sourceFile.GetFileNameWithoutExtension()));
 				}
-				foreach (FileReference SourceFile in Files)
+				foreach (FileReference sourceFile in files)
 				{
-					Logger.LogInformation("Source File: {Arg0}", SourceFile.FullName);
+					Logger.LogInformation("Source File: {Arg0}", sourceFile.FullName);
 				}
-				SymbolSourceFiles = Files.Where(x => x.HasExtension(".dSYM")).ToArray();
+				symbolSourceFiles = files.Where(x => x.HasExtension(".dSYM")).ToArray();
 			}
 			else if (BuildHostPlatform.Current.Platform == UnrealTargetPlatform.Linux)
 			{
-				SymbolDumperExecutable = Unreal.RootDirectory + "/Engine/Binaries/Linux/dump_syms";
-				string[] SymbolFileExtensions = { ".debug" };
-				SymbolSourceFiles = SourceFiles.Where(x => SymbolFileExtensions.Contains(x.GetExtension())).ToArray();
+				symbolDumperExecutable = Unreal.RootDirectory + "/Engine/Binaries/Linux/dump_syms";
+				string[] symbolFileExtensions = { ".debug" };
+				symbolSourceFiles = sourceFiles.Where(x => symbolFileExtensions.Contains(x.GetExtension())).ToArray();
 			}
 			else
 			{
@@ -145,82 +142,81 @@ namespace BuildGraph.Tasks
 			}
 
 			// Remove any existing symbol files
-			foreach (string FileName in SymbolSourceFiles.Select(x => Path.ChangeExtension(x.FullName, ".psym")))
+			foreach (string fileName in symbolSourceFiles.Select(x => Path.ChangeExtension(x.FullName, ".psym")))
 			{
-				if (File.Exists(FileName))
+				if (File.Exists(fileName))
 				{
 					try
 					{
-						File.Delete(FileName);
+						File.Delete(fileName);
 					}
-					catch (Exception Ex)
+					catch (Exception ex)
 					{
-						throw new AutomationException("Symbol generation failed: Unable to delete existing symbol file: \"{0}\". Error: {1}", FileName, Ex.Message.TrimEnd());
+						throw new AutomationException("Symbol generation failed: Unable to delete existing symbol file: \"{0}\". Error: {1}", fileName, ex.Message.TrimEnd());
 					}
 				}
-
 			}
 
-			if (SymbolSourceFiles.Length == 0)
+			if (symbolSourceFiles.Length == 0)
 			{
 				Logger.LogInformation("No symbol files to convert.");
 			}
 
 			// Generate portable symbols from the symbol source files
-			ConcurrentBag<FileReference> SymbolFiles = new ConcurrentBag<FileReference>();
+			ConcurrentBag<FileReference> symbolFiles = new ConcurrentBag<FileReference>();
 
-			Parallel.ForEach(SymbolSourceFiles, (SourceFile) =>
+			Parallel.ForEach(symbolSourceFiles, (sourceFile) =>
 			{
-				string SymbolFileName = Path.ChangeExtension(SourceFile.FullName, ".psym");
-				string RadSymbolTemp = Path.ChangeExtension(SourceFile.FullName, ".radpsym");
+				string symbolFileName = Path.ChangeExtension(sourceFile.FullName, ".psym");
+				string radSymbolTemp = Path.ChangeExtension(sourceFile.FullName, ".radpsym");
 
 				// Check if higher priority debug file or binary already created symbols 
-				if (File.Exists(SymbolFileName))
+				if (File.Exists(symbolFileName))
 				{
 					return;
 				}
 
-				Logger.LogInformation("Dumping Symbols: {Arg0} to {SymbolFileName}", SourceFile.FullName, SymbolFileName);
+				Logger.LogInformation("Dumping Symbols: {Arg0} to {SymbolFileName}", sourceFile.FullName, symbolFileName);
 
-				string DumpSymsArgs;
+				string dumpSymsArgs;
 
-				string SymbolDumperExeForFile = SymbolDumperExecutable;
+				string symbolDumperExeForFile = symbolDumperExecutable;
 				if (bUseRadSym &&
-					SourceFile.GetExtension() == ".pdb")
+					sourceFile.GetExtension() == ".pdb")
 				{
-					SymbolDumperExeForFile = RadSymDymperExecuable;
-					DumpSymsArgs = "-pdb:" + SourceFile.FullName + " -out:" + RadSymbolTemp + " -exe:" + SourceFile.FullName;
+					symbolDumperExeForFile = radSymDymperExecuable;
+					dumpSymsArgs = "-pdb:" + sourceFile.FullName + " -out:" + radSymbolTemp + " -exe:" + sourceFile.FullName;
 				}
 				else
 				{
-					string ExtraOptions = "";
+					string extraOptions = "";
 					if (BuildHostPlatform.Current.Platform == UnrealTargetPlatform.Mac)
 					{
 						// dump_syms has a bug where if a universal binary is fed into it, on an Intel mac, it will fail to find the current architecture
 						// (but not on Arm macs). Specify the host architecture as a param to cause expected behavior (until we make one output per Arch)
-						if (ArchitecturesInBinary(SourceFile).bIsMultiArch)
+						if (ArchitecturesInBinary(sourceFile).bIsMultiArch)
 						{
 							// ExtraOptions = $"-a {MacExports.HostArchitecture.AppleName} ";
 							// Since IBs are universal and we typically only care about arm symbols, force
 							// the arch to always be arm.
-							ExtraOptions = $"-a arm64 ";
+							extraOptions = $"-a arm64 ";
 						}
 					}
 
-					DumpSymsArgs = ExtraOptions + SourceFile.FullName;
+					dumpSymsArgs = extraOptions + sourceFile.FullName;
 				}
 
-				IProcessResult result = CommandUtils.Run(SymbolDumperExeForFile, DumpSymsArgs, null, CommandUtils.ERunOptions.AppMustExist, null, FilterSpew, null, WorkingDirectory);
+				IProcessResult result = CommandUtils.Run(symbolDumperExeForFile, dumpSymsArgs, null, CommandUtils.ERunOptions.AppMustExist, null, FilterSpew, null, workingDirectory);
 				if (result.ExitCode == 0)
 				{
-					StringBuilder ProcessedSymbols = null;
+					StringBuilder processedSymbols = null;
 					if (bUseRadSym)
 					{
 						// rad dumper outputs to a file, we thunk to a custom exe to do symbol munging for speed.
-						CommandUtils.Run(RadProcessSymExecuable, RadSymbolTemp + " " + Unreal.RootDirectory.FullName, null, CommandUtils.ERunOptions.AppMustExist, null, FilterSpew, null, WorkingDirectory);
+						CommandUtils.Run(radProcessSymExecuable, radSymbolTemp + " " + Unreal.RootDirectory.FullName, null, CommandUtils.ERunOptions.AppMustExist, null, FilterSpew, null, workingDirectory);
 
-						File.Move(RadSymbolTemp, SymbolFileName);
-						SymbolFiles.Add(new FileReference(SymbolFileName));
+						File.Move(radSymbolTemp, symbolFileName);
+						symbolFiles.Add(new FileReference(symbolFileName));
 
 					}
 					else
@@ -228,44 +224,44 @@ namespace BuildGraph.Tasks
 						try
 						{
 							// Process the symbols
-							using (StringReader Reader = new StringReader(result.Output))
+							using (StringReader reader = new StringReader(result.Output))
 							{
-								ProcessSymbols(SymbolFileName, Reader, out ProcessedSymbols);
+								ProcessSymbols(reader, out processedSymbols);
 							}
 						}
 						catch (OutOfMemoryException)
 						{
 							// If we catch an OOM, it is too large to turn into a string.
 							// Write to a file and then load it into a string.
-							string TempFileName = Path.Combine(Path.GetTempPath(), Path.GetTempFileName());
-							FileReference SymbolFile = (result as ProcessResult).WriteOutputToFile(TempFileName);
+							string tempFileName = Path.Combine(Path.GetTempPath(), Path.GetTempFileName());
+							FileReference symbolFile = (result as ProcessResult).WriteOutputToFile(tempFileName);
 
 							try
 							{
-								using (StreamReader Reader = new StreamReader(SymbolFile.FullName))
+								using (StreamReader reader = new StreamReader(symbolFile.FullName))
 								{
-									ProcessSymbols(SymbolFileName, Reader, out ProcessedSymbols);
+									ProcessSymbols(reader, out processedSymbols);
 								}
 							}
 							finally
 							{
-								FileReference.Delete(SymbolFile);
+								FileReference.Delete(symbolFile);
 							}
 						}
-						catch (Exception Ex)
+						catch (Exception ex)
 						{
 							// There was a problem generating symbols with the dump_syms tool
-							throw new AutomationException($"Symbol generation failed: Error Generating Symbols for {SymbolFileName}, Error: {ExceptionUtils.FormatException(Ex)}");
+							throw new AutomationException($"Symbol generation failed: Error Generating Symbols for {symbolFileName}, Error: {ExceptionUtils.FormatException(ex)}");
 						}
 					}
 
-					if (ProcessedSymbols != null && ProcessedSymbols.Length > 0)
+					if (processedSymbols != null && processedSymbols.Length > 0)
 					{
-						using (StreamWriter Writer = new StreamWriter(SymbolFileName))
+						using (StreamWriter writer = new StreamWriter(symbolFileName))
 						{
-							Writer.Write(ProcessedSymbols);
+							writer.Write(processedSymbols);
 						}
-						SymbolFiles.Add(new FileReference(SymbolFileName));
+						symbolFiles.Add(new FileReference(symbolFileName));
 					}
 				}
 				else
@@ -274,47 +270,48 @@ namespace BuildGraph.Tasks
 					{
 						// If we fail, lets re-run with a verbose, -v to check for the error we are seeing
 						// -v not available on Mac
-						CommandUtils.Run(SymbolDumperExecutable, "-v " + SourceFile.FullName, null, CommandUtils.ERunOptions.AppMustExist, null, null, null, WorkingDirectory);
+						CommandUtils.Run(symbolDumperExecutable, "-v " + sourceFile.FullName, null, CommandUtils.ERunOptions.AppMustExist, null, null, null, workingDirectory);
 					}
 
 					// There was a problem generating symbols with the dump_syms tool
-					throw new AutomationException("Symbol generation failed: Error Generating Symbols: {0}", SymbolFileName);
+					throw new AutomationException("Symbol generation failed: Error Generating Symbols: {0}", symbolFileName);
 				}
 			});
 
 			// Apply the optional tag to the build products
-			foreach (string TagName in FindTagNamesFromList(Parameters.Tag))
+			foreach (string tagName in FindTagNamesFromList(_parameters.Tag))
 			{
-				FindOrAddTagSet(TagNameToFileSet, TagName).UnionWith(SymbolFiles);
+				FindOrAddTagSet(tagNameToFileSet, tagName).UnionWith(symbolFiles);
 			}
 
 			// Add them to the list of build products
-			BuildProducts.UnionWith(SymbolFiles);
+			buildProducts.UnionWith(symbolFiles);
 		}
 
 		/// <summary>
 		/// Processes the raw symbol dump
 		/// </summary>
-		static bool ProcessSymbols(string SymbolFileName, TextReader Reader, out StringBuilder ProcessedSymbols)
+
+		static bool ProcessSymbols(TextReader reader, out StringBuilder processedSymbols)
 		{
-			char[] FieldSeparator = { ' ' };
-			string RootDirectory = CommandUtils.ConvertSeparators(PathSeparator.Slash, Unreal.RootDirectory.FullName).TrimEnd('/'); ;
+			char[] fieldSeparator = { ' ' };
+			string rootDirectory = CommandUtils.ConvertSeparators(PathSeparator.Slash, Unreal.RootDirectory.FullName).TrimEnd('/');
 
-			ProcessedSymbols = new StringBuilder();
+			processedSymbols = new StringBuilder();
 
-			string Line;
+			string line;
 			bool bSawModule = false;
-			while ((Line = Reader.ReadLine()) != null)
+			while ((line = reader.ReadLine()) != null)
 			{
-				if (Line.Contains(" = ", StringComparison.InvariantCulture))
+				if (line.Contains(" = ", StringComparison.InvariantCulture))
 				{
-					Logger.LogInformation("{Text}", Line);
+					Logger.LogInformation("{Text}", line);
 					continue;
 				}
 				// Ignore any output from symbol dump before MODULE, these may included erroneous warnings, etc
 				if (!bSawModule)
 				{
-					if (!Line.StartsWith("MODULE", StringComparison.InvariantCulture))
+					if (!line.StartsWith("MODULE", StringComparison.InvariantCulture))
 					{
 						continue;
 					}
@@ -322,30 +319,30 @@ namespace BuildGraph.Tasks
 					bSawModule = true;
 				}
 
-				string NewLine = Line;
+				string newLine = line;
 
 				// Process source reference FILE blocks
-				if (Line.StartsWith("FILE", StringComparison.InvariantCulture))
+				if (line.StartsWith("FILE", StringComparison.InvariantCulture))
 				{
-					string[] Fields = Line.Split(FieldSeparator, 3);
+					string[] fields = line.Split(fieldSeparator, 3);
 
-					string FileName = CommandUtils.ConvertSeparators(PathSeparator.Slash, Fields[2]);
+					string fileName = CommandUtils.ConvertSeparators(PathSeparator.Slash, fields[2]);
 
 					// If the file exists locally, and is within the root, convert path
-					if (File.Exists(FileName) && FileName.StartsWith(RootDirectory, StringComparison.OrdinalIgnoreCase))
+					if (File.Exists(fileName) && fileName.StartsWith(rootDirectory, StringComparison.OrdinalIgnoreCase))
 					{
 						// Restore proper filename case on Windows (the symbol dump filenames are all lowercase)
 						if (BuildHostPlatform.Current.Platform == UnrealTargetPlatform.Win64)
 						{
-							FileName = FileUtils.FindCorrectCase(new FileInfo(FileName)).FullName;
+							fileName = FileUtils.FindCorrectCase(new FileInfo(fileName)).FullName;
 						}
 
 						// Shave off the root directory
-						NewLine = String.Format("FILE {0} {1}", Fields[1], FileName.Substring(RootDirectory.Length + 1).Replace('\\', '/'));
+						newLine = String.Format("FILE {0} {1}", fields[1], fileName.Substring(rootDirectory.Length + 1).Replace('\\', '/'));
 					}
 				}
 
-				ProcessedSymbols.AppendLine(NewLine);
+				processedSymbols.AppendLine(newLine);
 			}
 
 			return true;
@@ -355,30 +352,30 @@ namespace BuildGraph.Tasks
 		/// <summary>
 		///  Filters the output from the dump_syms executable, which depending on the platform can be pretty spammy
 		/// </summary>
-		string FilterSpew(string Message)
+		string FilterSpew(string message)
 		{
-			foreach (string FilterString in OutputFilterStrings)
+			foreach (string filterString in s_outputFilterStrings)
 			{
-				if (Message.Contains(FilterString, StringComparison.InvariantCulture))
+				if (message.Contains(filterString, StringComparison.InvariantCulture))
 				{
 					return null;
 				}
 			}
 
-			return Message;
+			return message;
 		}
 
 		/// <summary>
 		/// Array of source strings to filter from output
 		/// </summary>
-		static readonly string[] OutputFilterStrings = new string[] { "the DIE at offset", "warning: function", "warning: failed", ": in compilation unit" };
+		static readonly string[] s_outputFilterStrings = new string[] { "the DIE at offset", "warning: function", "warning: failed", ": in compilation unit" };
 
 		/// <summary>
 		/// Output this task out to an XML writer.
 		/// </summary>
-		public override void Write(XmlWriter Writer)
+		public override void Write(XmlWriter writer)
 		{
-			Write(Writer, Parameters);
+			Write(writer, _parameters);
 		}
 
 		/// <summary>
@@ -387,7 +384,7 @@ namespace BuildGraph.Tasks
 		/// <returns>The tag names which are read by this task</returns>
 		public override IEnumerable<string> FindConsumedTagNames()
 		{
-			return FindTagNamesFromFilespec(Parameters.Files);
+			return FindTagNamesFromFilespec(_parameters.Files);
 		}
 
 		/// <summary>
@@ -396,7 +393,7 @@ namespace BuildGraph.Tasks
 		/// <returns>The tag names which are modified by this task</returns>
 		public override IEnumerable<string> FindProducedTagNames()
 		{
-			return FindTagNamesFromList(Parameters.Tag);
+			return FindTagNamesFromList(_parameters.Tag);
 		}
 	}
 }
@@ -407,11 +404,11 @@ namespace BuildScripts.Automation
 	{
 		public override ExitCode Execute()
 		{
-			BuildGraph.Tasks.SymGenTaskParameters Params = new BuildGraph.Tasks.SymGenTaskParameters();
-			Params.Files = ParseRequiredStringParam("Files");
+			BuildGraph.Tasks.SymGenTaskParameters symGenParams = new BuildGraph.Tasks.SymGenTaskParameters();
+			symGenParams.Files = ParseRequiredStringParam("Files");
 
-			BuildGraph.Tasks.SymGenTask Task = new BuildGraph.Tasks.SymGenTask(Params);
-			Task.Execute(null, new HashSet<FileReference>(), new Dictionary<string, HashSet<FileReference>>());
+			BuildGraph.Tasks.SymGenTask symGenTask = new BuildGraph.Tasks.SymGenTask(symGenParams);
+			symGenTask.Execute(null, new HashSet<FileReference>(), new Dictionary<string, HashSet<FileReference>>());
 
 			return ExitCode.Success;
 		}
