[BuildGraph] Adds optional IgnoreDuplicateIncludes attribute to the Include element
If a script has already been read by a previous Include, it will be ignored
#tests local BuildGraph -listonly

diff --git a/Engine/Source/Programs/AutomationTool/BuildGraph/BgSchema.cs b/Engine/Source/Programs/AutomationTool/BuildGraph/BgSchema.cs
--- a/Engine/Source/Programs/AutomationTool/BuildGraph/BgSchema.cs
+++ b/Engine/Source/Programs/AutomationTool/BuildGraph/BgSchema.cs
@@ -753,6 +753,7 @@ namespace AutomationTool
 			propertyType.Name = GetTypeName(ScriptSchemaStandardType.Include);
 			propertyType.Attributes.Add(CreateSchemaAttribute("Script", ScriptSchemaStandardType.BalancedString, XmlSchemaUse.Required));
 			propertyType.Attributes.Add(CreateSchemaAttribute("If", ScriptSchemaStandardType.BalancedString, XmlSchemaUse.Optional));
+			propertyType.Attributes.Add(CreateSchemaAttribute("IgnoreDuplicateIncludes", ScriptSchemaStandardType.BalancedString, XmlSchemaUse.Optional));
 			return propertyType;
 		}
 
diff --git a/Engine/Source/Programs/AutomationTool/BuildGraph/BgScriptReader.cs b/Engine/Source/Programs/AutomationTool/BuildGraph/BgScriptReader.cs
--- a/Engine/Source/Programs/AutomationTool/BuildGraph/BgScriptReader.cs
+++ b/Engine/Source/Programs/AutomationTool/BuildGraph/BgScriptReader.cs
@@ -326,6 +326,12 @@ namespace AutomationTool
 		/// </summary>
 		readonly Dictionary<string, BgScriptMacro> _macroNameToDefinition = new Dictionary<string, BgScriptMacro>();
 
+		/// <summary>
+		/// Set of scripts that have been previously read via an include element
+		/// This is used to prevent duplicate reads of a script if it has been guarded against
+		/// </summary>
+		readonly HashSet<FileReference> _includedScripts = new HashSet<FileReference>();
+
 		/// <summary>
 		/// The current graph
 		/// </summary>
@@ -458,6 +464,8 @@ namespace AutomationTool
 				return false;
 			}
 
+			_includedScripts.Add(file);
+
 			// Read the root BuildGraph element
 			await ReadGraphBodyAsync(document.DocumentElement!);
 			return true;
@@ -693,8 +701,31 @@ namespace AutomationTool
 
 				foreach (FileReference file in files.OrderBy(x => x.FullName, StringComparer.OrdinalIgnoreCase))
 				{
-					Logger.LogDebug("Including file {File}", file);
-					await TryReadAsync(file);
+					// Check if duplicate includes should be guarded against
+					bool ignoreDuplicateIncludes = false;
+					const string AttributeName = "IgnoreDuplicateIncludes";
+					if (element.HasAttribute(AttributeName))
+					{
+						try
+						{
+							string text = ExpandProperties(element, element.GetAttribute(AttributeName));
+							if (await BgCondition.EvaluateAsync(text, new BgConditionContext(_rootDir)))
+							{
+								ignoreDuplicateIncludes = true;
+							}
+						}
+						catch (BgConditionException ex)
+						{
+							LogError(element, "Error in condition: {Message}", ex.Message);
+							return;
+						}
+					}
+
+					if(!ignoreDuplicateIncludes || !_includedScripts.Contains(file))
+					{
+						Logger.LogDebug("Including file {File}", file);
+						await TryReadAsync(file);
+					}	
 				}
 			}
 		}
