[RpcFramework] Use strongly typed delegates for handling messages and http routes. More performant (avoid DynamicInvoke), safer (ensures the signature is correct) and convenient (lambda handlers can skip specifying the param types).

diff --git a/Engine/Source/Programs/AutomationTool/Gauntlet/Framework/RpcFramework/Gauntlet.RpcExecutor.cs b/Engine/Source/Programs/AutomationTool/Gauntlet/Framework/RpcFramework/Gauntlet.RpcExecutor.cs
--- a/Engine/Source/Programs/AutomationTool/Gauntlet/Framework/RpcFramework/Gauntlet.RpcExecutor.cs
+++ b/Engine/Source/Programs/AutomationTool/Gauntlet/Framework/RpcFramework/Gauntlet.RpcExecutor.cs
@@ -89,7 +89,7 @@ namespace Gauntlet
 		}
 		public override string ToString()
 		{
-			return string.Format("{0}|{1}", Path.ToLower(), Verb.ToLower()); ;
+			return string.Format("{0}|{1}", Path.ToLower(), Verb.ToLower());
 		}
 
 	}
@@ -153,8 +153,8 @@ namespace Gauntlet
 		}
 
 		public Queue<GauntletIncomingRpcMessage> IncomingMessageQueue = new Queue<GauntletIncomingRpcMessage>();
-		protected Dictionary<GauntletRpcDelegateRegistryFilter, Delegate> MessageDelegates = new Dictionary<GauntletRpcDelegateRegistryFilter, Delegate>();
-		protected Dictionary<RouteMapping, Delegate> HttpRoutes = new Dictionary<RouteMapping, Delegate>();
+		protected Dictionary<GauntletRpcDelegateRegistryFilter, GauntletIncomingRpcMessageDelegate> MessageDelegates = new Dictionary<GauntletRpcDelegateRegistryFilter, GauntletIncomingRpcMessageDelegate>();
+		protected Dictionary<RouteMapping, GauntletHttpRouteDelegate> HttpRoutes = new Dictionary<RouteMapping, GauntletHttpRouteDelegate>();
 		protected object MessageQueueLock = new object();
 
 		protected object ReceivedTerminationCallLock = new object();
@@ -243,7 +243,7 @@ namespace Gauntlet
 			ResponseData CallResponse = new ResponseData();
 			CallResponse.ResponseBody = "This endpoint is not handled";
 
-			foreach (KeyValuePair<RouteMapping, Delegate> RouteEntry in HttpRoutes)
+			foreach (KeyValuePair<RouteMapping, GauntletHttpRouteDelegate> RouteEntry in HttpRoutes)
 			{
 				if (CallPath.Equals(RouteEntry.Key))
 				{
@@ -265,7 +265,7 @@ namespace Gauntlet
 								}
 								StreamReader RequestReader = new StreamReader(ListenerRequest.InputStream, ListenerRequest.ContentEncoding);
 								string RequestBody = RequestReader.ReadToEnd();
-								RouteEntry.Value.DynamicInvoke(RequestBody, CallResponse);
+								RouteEntry.Value.Invoke(RequestBody, CallResponse);
 								ListenerContext.Response.StatusCode = CallResponse.ResponseCode;
 								byte[] ResponseBuffer = Encoding.UTF8.GetBytes(CallResponse.ResponseBody);
 								if (ResponseBuffer == null)
@@ -598,6 +598,11 @@ namespace Gauntlet
 			}
 		} 
 
+		/// <summary>
+		/// Delegate type for incoming messages from Unreal processes.
+		/// </summary>
+		public delegate void GauntletIncomingRpcMessageDelegate(GauntletIncomingRpcMessage Message);
+
 		/// <summary>
 		///  Delegate registry function. Set any fields that don't matter to the filter to "" or String.Empty, and your delegate
 		///  will fire off any time a message comes in that matches all criteria of your filter. It does not fuzzy match.
@@ -606,7 +611,7 @@ namespace Gauntlet
 		/// <param name="SenderFilter">If you care about what process the message is sent from, put it here.</param>
 		/// <param name="CategoryFilter">Message category you care about filtering on.</param>
 		/// <param name="PayloadFilter">Exact message payload you'd likte to trigger on.</param>
-		public virtual void RegisterMessageDelegate(Delegate FunctionToRun, string SenderFilter = "", string CategoryFilter = "", string PayloadFilter = "", bool bDeregisterAfterTrigger = false)
+		public virtual void RegisterMessageDelegate(GauntletIncomingRpcMessageDelegate FunctionToRun, string SenderFilter = "", string CategoryFilter = "", string PayloadFilter = "", bool bDeregisterAfterTrigger = false)
 		{
 			if (string.IsNullOrEmpty(SenderFilter) && string.IsNullOrEmpty(CategoryFilter) && string.IsNullOrEmpty(PayloadFilter))
 			{
@@ -625,7 +630,7 @@ namespace Gauntlet
 		/// <param name="CategoryFilter">Category filter for the delegate trigger.</param>
 		/// <param name="PayloadFilter">Payload filter for the delegate trigger.</param>
 		/// <returns></returns>
-		public virtual bool DeregisterMessageDelegate(Delegate FunctionToRun, string SenderFilter = "", string CategoryFilter = "", string PayloadFilter = "")
+		public virtual bool DeregisterMessageDelegate(GauntletIncomingRpcMessageDelegate FunctionToRun, string SenderFilter = "", string CategoryFilter = "", string PayloadFilter = "")
 		{
 			bool bDeregisteredDelegate = false;
 			foreach (GauntletRpcDelegateRegistryFilter MessageFilter in MessageDelegates.Keys)
@@ -649,7 +654,7 @@ namespace Gauntlet
 		{
 			foreach (GauntletRpcDelegateRegistryFilter MessageFilter in MessageDelegates.Keys.Where(key => NewMessage.MatchesMessageFilter(key)).Select(MessageFilter => MessageFilter))
 			{
-				MessageDelegates[MessageFilter].DynamicInvoke(NewMessage);
+				MessageDelegates[MessageFilter].Invoke(NewMessage);
 				if (MessageFilter.bDeregisterAfterFirstTrigger)
 				{
 					MessageDelegates.Remove(MessageFilter);
@@ -675,7 +680,12 @@ namespace Gauntlet
 			}
 		}
 
-		public void RegisterHttpRoute(RouteMapping Path, Delegate RouteDelegate, bool bOverrideIfExists = false)
+		/// <summary>
+		/// Delegate type for handling messages on a http route.
+		/// </summary>
+		public delegate void GauntletHttpRouteDelegate(string MessageBody, ResponseData CallResponse);
+
+		public void RegisterHttpRoute(RouteMapping Path, GauntletHttpRouteDelegate RouteDelegate, bool bOverrideIfExists = false)
 		{
 			if (HttpRoutes.ContainsKey(Path))
 			{
