UAT: Adjust how binaries are referenced for ImageMagick

#rnx

diff --git a/Engine/Source/Programs/AutomationTool/AutomationTool.csproj b/Engine/Source/Programs/AutomationTool/AutomationTool.csproj
--- a/Engine/Source/Programs/AutomationTool/AutomationTool.csproj
+++ b/Engine/Source/Programs/AutomationTool/AutomationTool.csproj
@@ -97,6 +97,7 @@
 
   <ItemGroup>
     <PackageReference Include="AWSSDK.SecurityToken" Version="3.7.300.105" />
+    <PackageReference Include="Magick.NET-Q16-HDRI-AnyCPU" Version="14.0.0" />
     <PackageReference Include="Microsoft.Build" Version="17.11.4" ExcludeAssets="runtime" PrivateAssets="all">
       <IncludeAssets>compile; build; native; contentfiles; analyzers; buildtransitive</IncludeAssets>
     </PackageReference>
@@ -105,6 +106,14 @@
          Adding a loader context resolves many of the issues, but it still fails when MSBuild tries to write the json -->
     <PackageReference Include="Newtonsoft.Json" Version="13.0.3" />
   </ItemGroup>
+
+  <!-- Ensure Magick.NET native binaries are copied -->
+  <PropertyGroup>
+    <MagickCopyNativeWindows>true</MagickCopyNativeWindows>
+    <MagickCopyNativeLinux>true</MagickCopyNativeLinux>
+    <MagickCopyNativeLinuxMusl>true</MagickCopyNativeLinuxMusl>
+    <MagickCopyNativeMacOS>true</MagickCopyNativeMacOS>
+  </PropertyGroup>
   
   <!-- This set of dependencies should be kept as small as possible
        This list must match the list of locations searched for changed files in BuildUAT.bat and BuildUAT.sh -->
diff --git a/Engine/Source/Programs/AutomationTool/AutomationUtils/AutomationUtils.Automation.csproj b/Engine/Source/Programs/AutomationTool/AutomationUtils/AutomationUtils.Automation.csproj
--- a/Engine/Source/Programs/AutomationTool/AutomationUtils/AutomationUtils.Automation.csproj
+++ b/Engine/Source/Programs/AutomationTool/AutomationUtils/AutomationUtils.Automation.csproj
@@ -62,7 +62,7 @@
     <ProjectReference Include="..\..\UnrealBuildTool\UnrealBuildTool.csproj" PrivateAssets="All"><Private>true</Private></ProjectReference>
   </ItemGroup>
   <ItemGroup>
-    <PackageReference Include="Magick.NET-Q16-HDRI-AnyCPU" Version="14.0.0" />
+    <PackageReference Include="Magick.NET-Q16-HDRI-AnyCPU" Version="14.0.0" PrivateAssets="all" />
     <PackageReference Include="System.Drawing.Common" Version="8.0.10" />
     <PackageReference Include="System.Net.Http" Version="4.3.4" />
 
diff --git a/Engine/Source/Programs/AutomationTool/Gauntlet/Gauntlet.Automation.csproj b/Engine/Source/Programs/AutomationTool/Gauntlet/Gauntlet.Automation.csproj
--- a/Engine/Source/Programs/AutomationTool/Gauntlet/Gauntlet.Automation.csproj
+++ b/Engine/Source/Programs/AutomationTool/Gauntlet/Gauntlet.Automation.csproj
@@ -77,7 +77,7 @@
     <ProjectReference Include="..\Android\Android.Automation.csproj" PrivateAssets="All"><Private>false</Private></ProjectReference>
   </ItemGroup>
   <ItemGroup>
-    <PackageReference Include="Magick.NET-Q16-HDRI-AnyCPU" Version="14.0.0" />
+    <PackageReference Include="Magick.NET-Q16-HDRI-AnyCPU" Version="14.0.0" PrivateAssets="all" />
     <PackageReference Include="Microsoft.CSharp" Version="4.7.0" />
     <PackageReference Include="MySql.Data" Version="6.10.9" PrivateAssets="all" />
   </ItemGroup>
