Change the number of days on formal builds to be a double to allow for finer tuning (hours)
Move the delete action of delete formal builds to be a parallel with an optional number of threads

diff --git a/Engine/Source/Programs/AutomationTool/AutomationUtils/CommandUtils.cs b/Engine/Source/Programs/AutomationTool/AutomationUtils/CommandUtils.cs
--- a/Engine/Source/Programs/AutomationTool/AutomationUtils/CommandUtils.cs
+++ b/Engine/Source/Programs/AutomationTool/AutomationUtils/CommandUtils.cs
@@ -2298,35 +2298,47 @@ namespace AutomationTool
 			return false;
 		}
 
-		public static void CleanFormalBuilds(string ParentDir, string SearchPattern, int MaximumDaysToKeepTempStorage = 4)
+		public static void CleanFormalBuilds(string ParentDir, string SearchPattern, double MaximumDaysToKeepTempStorage = 4, int maxParallelDeletes = 4)
 		{
 			if (!IsBuildMachine || !ParentDir.StartsWith(RootBuildStorageDirectory()))
 			{
 				return;
 			}
-			try
+			DirectoryInfo DirInfo = new DirectoryInfo(ParentDir);
+			Logger.LogInformation("Looking for directories to delete in {ParentDir}", ParentDir);
+			List<DirectoryInfo> BuildsToClean = new List<DirectoryInfo>();
+			foreach (DirectoryInfo ThisDirInfo in DirInfo.EnumerateDirectories(SearchPattern))
 			{
-				DirectoryInfo DirInfo = new DirectoryInfo(ParentDir);
-				Logger.LogInformation("Looking for directories to delete in {ParentDir}", ParentDir);
-				foreach (DirectoryInfo ThisDirInfo in DirInfo.EnumerateDirectories(SearchPattern))
+				double AgeDays = (DateTime.UtcNow - ThisDirInfo.CreationTimeUtc).TotalDays;
+				if (AgeDays > MaximumDaysToKeepTempStorage)
 				{
-					double AgeDays = (DateTime.UtcNow - ThisDirInfo.CreationTimeUtc).TotalDays;
-					if (AgeDays > MaximumDaysToKeepTempStorage)
-					{
-						Logger.LogInformation("Deleting formal build directory {Directory}, because it is {Age} days old (maximum {MaximumDaysToKeepTempStorage}).", ThisDirInfo.FullName, (int)AgeDays, MaximumDaysToKeepTempStorage);
-						DeleteDirectory_NoExceptions(true, ThisDirInfo.FullName);
-					}
-					else
-					{
-						Logger.LogDebug("Not deleting formal build directory {Directory}, because it is {Age} days old (maximum {MaximumDaysToKeepTempStorage}).", ThisDirInfo.FullName, (int)AgeDays, MaximumDaysToKeepTempStorage);
-					}
+					Logger.LogInformation("Marking formal build for delete directory {Directory}, because it is {Age} days old (maximum {MaximumDaysToKeepTempStorage}).", ThisDirInfo.FullName, AgeDays, MaximumDaysToKeepTempStorage);
+					BuildsToClean.Add(new DirectoryInfo(ThisDirInfo.FullName));
+				}
+				else
+				{
+					Logger.LogDebug("Not deleting formal build directory {Directory}, because it is {Age} days old (maximum {MaximumDaysToKeepTempStorage}).", ThisDirInfo.FullName, AgeDays, MaximumDaysToKeepTempStorage);
 				}
 			}
-			catch (Exception Ex)
+
+			Parallel.ForEach(BuildsToClean, new ParallelOptions() { MaxDegreeOfParallelism = maxParallelDeletes }, (buildName, CancelToken) =>
 			{
-				Logger.LogWarning("Unable to clean formal builds from directory: {ParentDir}", ParentDir);
-				Logger.LogWarning(Ex, " Exception was {Exception}", LogUtils.FormatException(Ex));
-			}
+				try
+				{
+					Logger.LogInformation("Deleting {Path}", buildName.FullName);
+					Stopwatch Timer = Stopwatch.StartNew();
+
+					DeleteDirectory_NoExceptions(false, buildName.FullName);
+
+					Timer.Stop();
+					Logger.LogInformation("{Time} time spent deleting {Path}", Timer.Elapsed, buildName.FullName);
+				}
+				catch (System.Exception Ex)
+				{
+					Logger.LogWarning("Failed to delete {Path} due to error: {Exception}", buildName.FullName, LogUtils.FormatException(Ex));
+					Logger.LogWarning(Ex, " Exception was {Exception}", LogUtils.FormatException(Ex));
+				}
+			});
 		}
 
 
diff --git a/Engine/Source/Programs/AutomationTool/Scripts/CleanFormalBuilds.Automation.cs b/Engine/Source/Programs/AutomationTool/Scripts/CleanFormalBuilds.Automation.cs
--- a/Engine/Source/Programs/AutomationTool/Scripts/CleanFormalBuilds.Automation.cs
+++ b/Engine/Source/Programs/AutomationTool/Scripts/CleanFormalBuilds.Automation.cs
@@ -38,8 +38,8 @@ namespace AutomationTool
 			}
 			else
 			{
-				int DaysValue;
-				if (!int.TryParse(Days, out DaysValue))
+				double DaysValue;
+				if (!double.TryParse(Days, out DaysValue))
 				{
 					throw new AutomationException("'{0}' is not a valid value for the -Days parameter", Days);
 				}
