Fix missing directory error.

diff --git a/Engine/Source/Programs/AutomationTool/CookedEditor/CookedEditor.Automation.cs b/Engine/Source/Programs/AutomationTool/CookedEditor/CookedEditor.Automation.cs
--- a/Engine/Source/Programs/AutomationTool/CookedEditor/CookedEditor.Automation.cs
+++ b/Engine/Source/Programs/AutomationTool/CookedEditor/CookedEditor.Automation.cs
@@ -681,7 +681,7 @@ public class MakeCookedEditor : BuildCommand
 		{
 			Context.UFSFilesToStage.AddRange(DirectoryReference.EnumerateFiles(PluginSubdir, "*", SearchOption.AllDirectories));
 		}
-		else if (DirNameLower == "documentation")
+		else if (DirNameLower == "documentation" && Directory.Exists(PluginSubdir.FullName + "/Source/Shared"))
 		{
 			Context.UFSFilesToStage.AddRange(DirectoryReference.EnumerateFiles(PluginSubdir, "Source/Shared/*", SearchOption.AllDirectories));
 		}
