Fix wrong empty commandline arguments sent by gauntlet

diff --git a/Engine/Source/Programs/AutomationTool/Gauntlet/Unreal/Base/Gauntlet.UnrealTestConfiguration.cs b/Engine/Source/Programs/AutomationTool/Gauntlet/Unreal/Base/Gauntlet.UnrealTestConfiguration.cs
--- a/Engine/Source/Programs/AutomationTool/Gauntlet/Unreal/Base/Gauntlet.UnrealTestConfiguration.cs
+++ b/Engine/Source/Programs/AutomationTool/Gauntlet/Unreal/Base/Gauntlet.UnrealTestConfiguration.cs
@@ -288,7 +288,16 @@ namespace Gauntlet
 		/// <returns></returns>
 		public string GenerateFullCommandLine()
 		{
-			string FinalCommandline = string.Format("\"{0}\" \"{1}\" ", Project, GameMap);
+			string FinalCommandline = "";
+			if (!string.IsNullOrEmpty(Project))
+			{
+				FinalCommandline += string.Format("\"{0}\" ", Project);
+			}
+			if (!string.IsNullOrEmpty(GameMap))
+			{
+				FinalCommandline += string.Format("\"{0}\" ", GameMap);
+			}
+
 			foreach (string Key in Params.Keys)
 			{
 				string CurrentArgument;
