[Gauntlet] Clarify error when querying external storage setup on Android device

#rnx

diff --git a/Engine/Source/Programs/AutomationTool/Gauntlet/Platform/Android/Gauntlet.TargetDeviceAndroid.cs b/Engine/Source/Programs/AutomationTool/Gauntlet/Platform/Android/Gauntlet.TargetDeviceAndroid.cs
--- a/Engine/Source/Programs/AutomationTool/Gauntlet/Platform/Android/Gauntlet.TargetDeviceAndroid.cs
+++ b/Engine/Source/Programs/AutomationTool/Gauntlet/Platform/Android/Gauntlet.TargetDeviceAndroid.cs
@@ -151,15 +151,14 @@ namespace Gauntlet
 			var ConnectedDevices = GetAllConnectedDevices();
 
 			IsExistingDevice = ConnectedDevices.Keys.Contains(DeviceName);
-
+			if (!IsExistingDevice && !DeviceName.Contains(":"))
+			{
+				// adb uses port 5555 by default
+				DeviceName += ":5555";
+				IsExistingDevice = ConnectedDevices.Keys.Contains(DeviceName);
+			}
 			if (!IsExistingDevice)
 			{
-				// adb uses 5555 by default
-				if (DeviceName.Contains(":") == false)
-				{
-					DeviceName = DeviceName + ":5555";
-				}
-
 				lock (Globals.MainLock)
 				{
 					using (var PauseEC = new ScopedSuspendECErrorParsing())
@@ -172,41 +171,42 @@ namespace Gauntlet
 						}
 					}
 
-					Log.Info("Connected to {0}", DeviceName);
+					Log.Info("Connecting to {0}", DeviceName);
 
 					// Need to sleep for adb service process to register, otherwise get an unauthorized (especially on parallel device use)
 					Thread.Sleep(5000);
 				}
-			}
 
-			LocalDirectoryMappings = new Dictionary<EIntendedBaseCopyDirectory, string>();
+				ConnectedDevices = GetAllConnectedDevices();
 
-			// for IP devices need to sanitize this
-			Name = DeviceName.Replace(":", "_");
+				// sanity check that it was now found
+				if (!ConnectedDevices.Keys.Contains(DeviceName))
+				{
+					throw new AutomationException("Failed to find new device {0} in connection list", DeviceName);
+				}
+			}
 
-			// Temp path used when pulling artifact
-			LocalCachePath = InCachePath ?? Path.Combine(Globals.TempDir, "AndroidDevice_" + Name);
+			if (ConnectedDevices[DeviceName] != "device")
+			{
+				Dispose();
+				throw new AutomationException("Device {0} is '{1}'.", DeviceName, ConnectedDevices[DeviceName]);
+			}
 
+			// Get the external storage path from the device (once the device is validated as connected)
 			IProcessResult StorageQueryResult = RunAdbDeviceCommand(AndroidPlatform.GetStorageQueryCommand());
-			if(StorageQueryResult.ExitCode != 0)
+			if (StorageQueryResult.ExitCode != 0)
 			{
-				throw new AutomationException("Failed to find root storage path");
+				throw new AutomationException("Failed to query external storage setup on {0}", DeviceName);
 			}
 			StoragePath = StorageQueryResult.Output.Trim();
 
-			ConnectedDevices = GetAllConnectedDevices();
+			LocalDirectoryMappings = new Dictionary<EIntendedBaseCopyDirectory, string>();
 
-			// sanity check that it was now dound
-			if (ConnectedDevices.Keys.Contains(DeviceName) == false)
-			{
-				throw new AutomationException("Failed to find new device {0} in connection list", DeviceName);
-			}
+			// for IP devices need to sanitize this
+			Name = DeviceName.Replace(":", "_");
 
-			if (ConnectedDevices[DeviceName] != "device")
-			{
-				Dispose();
-				throw new AutomationException("Device {0} is '{1}'.", DeviceName, ConnectedDevices[DeviceName]);
-			}
+			// Temp path used when pulling artifact
+			LocalCachePath = InCachePath ?? Path.Combine(Globals.TempDir, "AndroidDevice_" + Name);
 		}
 
 		~TargetDeviceAndroid()
