[Gauntlet] Increase verbosity of failure to create device log to a warning

diff --git a/Engine/Source/Programs/AutomationTool/Gauntlet/Framework/Devices/Gauntlet.DevicePool.cs b/Engine/Source/Programs/AutomationTool/Gauntlet/Framework/Devices/Gauntlet.DevicePool.cs
--- a/Engine/Source/Programs/AutomationTool/Gauntlet/Framework/Devices/Gauntlet.DevicePool.cs
+++ b/Engine/Source/Programs/AutomationTool/Gauntlet/Framework/Devices/Gauntlet.DevicePool.cs
@@ -1137,7 +1137,7 @@ namespace Gauntlet
 			}
 			catch (Exception Ex)
 			{
-				Log.Info(KnownLogEvents.Gauntlet_DeviceEvent, "Failed to create device {Name}. {Message}\n{StackTrace}", Def.ToString(), Ex.Message, Ex.StackTrace);
+				Log.Warning(KnownLogEvents.Gauntlet_DeviceEvent, "Failed to create device {Name}. {Message}\n{StackTrace}", Def.ToString(), Ex.Message, Ex.StackTrace);
 			}
 
 			return NewDevice;
