Localization:
Introduce the ability for plugins set for Auto localization to use template localization config files. Now plugins set with an "Auto" localization config generation policy can have their localization config files generated from template files instead of hardcoded default values.
- Introduced the "-TemplateLocalizationConfigFiles" parameter to the Localize UAT command. This parameter accepts a list of relative file paths from UEProjectDirectory for localization config files separated with a ';' as a delimiter. Absolute file paths are also supported if config files from other project directories are to be used.
- Introduced the ability to generate localization config files based on template localization config files where placeholders such as %LOCTARGETROOT% and %LOCTARGETNAME% can be substituted with the various names of plugisn and plugin directories that are participating in Auto plugin localization
- Introduced the ability to load a localization config file from disk into a LocalizationConfigFile object to be manipulated in memory.

diff --git a/Engine/Source/Programs/AutomationTool/Localization/LocalizationConfigFile.cs b/Engine/Source/Programs/AutomationTool/Localization/LocalizationConfigFile.cs
--- a/Engine/Source/Programs/AutomationTool/Localization/LocalizationConfigFile.cs
+++ b/Engine/Source/Programs/AutomationTool/Localization/LocalizationConfigFile.cs
@@ -1,8 +1,12 @@
 // Copyright Epic Games, Inc. All Rights Reserved.
 using EpicGames.Core;
 using System;
+using System.Collections.Generic;
+using System.Linq;
+using Microsoft.Extensions.Logging;
 using UnrealBuildTool;
 
+using static AutomationTool.CommandUtils;
 namespace EpicGames.Localization
 {
 
@@ -30,17 +34,98 @@ namespace EpicGames.Localization
 			Name = fileName;
 		}
 
+		private LocalizationConfigFile(ConfigFile configFile, string fileName)
+		{
+			_configFile = (configFile == null) ? new ConfigFile() : configFile;
+			Name = fileName;
+		}
+
 		public LocalizationConfigFileSection FindOrAddSection(string sectionName)
 		{
 			ConfigFileSection configSection = _configFile.FindOrAddSection(sectionName);
 			return new LocalizationConfigFileSection(configSection);
 		}
 
-		public void Load(string filePath)
+		public static LocalizationConfigFile Load(string filePath)
+		{
+			LocalizationConfigFile localizationConfigFile = new();
+			if (String.IsNullOrEmpty(filePath))
+			{
+				return null;
+			}
+			FileReference filePathReference = new FileReference(filePath);
+			if (!filePathReference.ToFileInfo().Exists)
+			{
+				return null;
+			}
+			// Parse the passed in config file 
+			ConfigFile configFile = new ConfigFile(filePathReference);
+			localizationConfigFile = new LocalizationConfigFile(configFile, filePathReference.GetFileName());
+			
+			return localizationConfigFile;
+		}
+
+		private LocalizationConfigFile DeepCopy()
+		{
+			// All constructors of this class should initialize _configFile in some way. It should never be null and so we don't handle that case
+			ConfigFile copyConfigFile = new ConfigFile();
+			foreach (string sectionName in _configFile.SectionNames)
+			{
+				ConfigFileSection originalSection;
+				if (_configFile.TryGetSection(sectionName, out originalSection))
+				{
+					ConfigFileSection copySection = copyConfigFile.FindOrAddSection(originalSection.Name);
+					foreach (ConfigLine originalLine in originalSection.Lines)
+					{
+						// The key and value for entries in localization config files will ever only be strings or bools
+						// Strings are immutable and bools are simple value so a direct copy is fine 
+						ConfigLine copyLine = new ConfigLine(originalLine.Action, originalLine.Key, originalLine.Value);
+						copySection.Lines.Add(copyLine);
+					}
+				}
+			}
+			// now create a new localization config file 
+			LocalizationConfigFile localizationConfigFile = new LocalizationConfigFile(copyConfigFile, Name);
+			return localizationConfigFile;
+		}
+
+		public LocalizationConfigFile ReplacePlaceHolders(Dictionary<string, string> replacements)
 		{
-			// @TODOLocalization: Load the config file and parse all the data into the various LocalizationConfigFileParam classes 
+			LocalizationConfigFile copy = DeepCopy();
+			if (replacements.Count == 0)
+			{
+				return copy;
+			}
+			// Loop through all of the sections and lines and replace any placeholder values 
+			foreach (string sectionName in copy._configFile.SectionNames)
+			{
+				ConfigFileSection configSection;
+				if (copy._configFile.TryGetSection(sectionName, out configSection))
+				{
+					foreach (ConfigLine configLine in configSection.Lines)
+					{
+						if (!String.IsNullOrEmpty(configLine.Value))
+						{
+							foreach (var pair in replacements)
+							{
+								// @TODOLocalization: For now, we only replace placeholders found in the values of the config files, replacing keys is not supported 
+								configLine.Value = configLine.Value.Replace(pair.Key, pair.Value);
+							}
+						}
+					}
+				}
+			}
+
+			return copy;
 		}
 
+		public bool IsEmpty()
+		{
+			return _configFile.SectionNames.Count() == 0;
+		}
+
+		// @TODOLocalization: There should be an IsValid function of some kind to ensure the validity of the localization config file. Should also implement validators 
+
 		public void Write(FileReference destinationFilePath)
 		{
 			// Go through all the config values and replace \ directory separators with / for uniformity across platforms
diff --git a/Engine/Source/Programs/AutomationTool/Localization/LocalizationConfigFileGenerator.cs b/Engine/Source/Programs/AutomationTool/Localization/LocalizationConfigFileGenerator.cs
--- a/Engine/Source/Programs/AutomationTool/Localization/LocalizationConfigFileGenerator.cs
+++ b/Engine/Source/Programs/AutomationTool/Localization/LocalizationConfigFileGenerator.cs
@@ -1,6 +1,10 @@
 // Copyright Epic Games, Inc. All Rights Reserved.
 using System;
 using System.Collections.Generic;
+using System.IO;
+using Microsoft.Extensions.Logging;
+
+using static AutomationTool.CommandUtils;
 
 namespace EpicGames.Localization
 {
@@ -24,7 +28,7 @@ namespace EpicGames.Localization
 		}
 
 		public abstract List<LocalizationConfigFile> GenerateDefaultSettingsConfigFiles(LocalizationConfigFileGeneratorParams generatorParams);
-		
+		public abstract List<LocalizationConfigFile> GenerateConfigFilesFromTemplate(LocalizationConfigFileGeneratorParams generatorParams, List<LocalizationConfigFile> templateFiles);
 	}
 
 	public class ModularLocalizationConfigFileGenerator : LocalizationConfigFileGenerator
@@ -40,6 +44,38 @@ namespace EpicGames.Localization
 			return configFiles;
 		}
 
+		public override List<LocalizationConfigFile> GenerateConfigFilesFromTemplate(LocalizationConfigFileGeneratorParams generatorParams, List<LocalizationConfigFile> templateFiles)
+		{
+			List<LocalizationConfigFile> generatedFiles = new List<LocalizationConfigFile>();
+			foreach (LocalizationConfigFile templateFile in templateFiles)
+			{
+				// @TODOLocalization: Should have additional validity checks 
+				if (templateFile.IsEmpty())
+				{
+					Logger.LogError($"Provided template localization config file '{templateFile.Name}' is invalid. The file will not be used to generate a localization config file for '{generatorParams.LocalizationTargetName}'.");
+					continue;
+				}
+
+				// @TODOLocalization: Should also include %LOCENGINEROOT%, %LOCPROJECTROOT% placeholders as well 
+				Dictionary<string, string> replacements = new Dictionary<string, string>
+				{
+					{"%LOCTARGETNAME%", generatorParams.LocalizationTargetName},
+					{"%LOCTARGETROOT%", generatorParams.LocalizationTargetRootDirectory}
+				};
+				LocalizationConfigFile copyFile = templateFile.ReplacePlaceHolders(replacements);
+				copyFile.Name = $"{generatorParams.LocalizationTargetName}{LocalizationUtilities.GetModularLocalizationConfigSuffix(templateFile.Name)}{LocalizationFileExtensions.ModularLocalizationConfigFileExtension}";
+
+				generatedFiles.Add(copyFile);
+			}
+			// Somehow we skipped using a template file to use for generation. Most likely because the template file is invalid 
+			// Warn/log here 
+			if (generatedFiles.Count != templateFiles.Count)
+			{
+				Logger.LogError($"{templateFiles.Count} template files provided but {generatedFiles.Count} files generated for localization target '{generatorParams.LocalizationTargetName}'. Verify that all provided template localization config files are valid.");
+			}
+			return generatedFiles;
+		}
+
 		private LocalizationConfigFile GenerateLocalizationGatherConfigFile(LocalizationConfigFileGeneratorParams generatorParam)
 		{
 			CommonSettingsParams commonSettings = CommonSettingsParams.CreateDefault(generatorParam.LocalizationTargetName, generatorParam.LocalizationTargetRootDirectory);
diff --git a/Engine/Source/Programs/AutomationTool/Localization/LocalizationUtilities.cs b/Engine/Source/Programs/AutomationTool/Localization/LocalizationUtilities.cs
--- a/Engine/Source/Programs/AutomationTool/Localization/LocalizationUtilities.cs
+++ b/Engine/Source/Programs/AutomationTool/Localization/LocalizationUtilities.cs
@@ -3,6 +3,7 @@ using System.Collections.Generic;
 using System.IO;
 using UnrealBuildTool;
 using EpicGames.Core;
+using System;
 
 namespace EpicGames.Localization
 {
@@ -49,5 +50,42 @@ namespace EpicGames.Localization
 			return FileSuffixes;
 		}
 
+		public static bool IsModularLocalizationConfigFile(String file)
+		{
+			if (String.IsNullOrEmpty(file))
+			{
+				return false;
+			}
+
+			String fileName = Path.GetFileNameWithoutExtension(file);
+			String[] suffixes = LocalizationUtilities.GetModularConfigSuffixes();
+			foreach (String suffix in suffixes)
+			{
+				if (fileName.EndsWith(suffix))
+				{
+					return true;
+				}
+			}
+			return false;
+		}
+
+		public static string GetModularLocalizationConfigSuffix(String file)
+		{
+			if (String.IsNullOrEmpty(file))
+			{
+				return "";
+			}
+
+			String fileName = Path.GetFileNameWithoutExtension(file);
+			String[] suffixes = LocalizationUtilities.GetModularConfigSuffixes();
+			foreach (String suffix in suffixes)
+			{
+				if (fileName.EndsWith(suffix))
+				{
+					return suffix;
+				}
+			}
+			return "";
+		}
 	}
 }
diff --git a/Engine/Source/Programs/AutomationTool/Scripts/Localisation.Automation.cs b/Engine/Source/Programs/AutomationTool/Scripts/Localisation.Automation.cs
--- a/Engine/Source/Programs/AutomationTool/Scripts/Localisation.Automation.cs
+++ b/Engine/Source/Programs/AutomationTool/Scripts/Localisation.Automation.cs
@@ -11,6 +11,7 @@ using EpicGames.Localization;
 using EpicGames.Core;
 using System.Threading.Tasks;
 using Microsoft.Extensions.Logging;
+using System.Diagnostics;
 
 
 
@@ -32,6 +33,7 @@ using Microsoft.Extensions.Logging;
 [Help("AdditionalCommandletArguments", "Optional arguments to pass to the gather process.")]
 [Help("ParallelGather", "Run the gather processes for a single batch in parallel rather than sequence.")]
 [Help("Preview", "Run the localization command in preview mode. This passes the -Preview flag along to all commandlets as an additional argument and deletes all temporary files generated by the commandlets in preview mode. Primarily used for build farm automation where localization warnings from localization gathers can be previewed without checking out any files under SCC.")]
+[Help("TemplateLocalizationConfigFiles", "An optional list of paths to localization config files to be used as templates for Auto localization config file generation. If the paths provided are relative, they should be relative to UEProjectDirectory. Else the path should be absolute. All paths should be separated by a ;.")]
 class Localize : BuildCommand
 {
 	private class LocalizationBatch
@@ -217,6 +219,7 @@ class Localize : BuildCommand
 	private List<string> ExcludePlugins = new();
 	private bool bShouldGatherPlatforms = false;
 	private string AdditionalCommandletArguments = "";
+	private List<LocalizationConfigFile> TemplateLocalizationConfigFiles = new List<LocalizationConfigFile>();
 	private bool bEnableParallelGather = false;
 	private bool bIsRunningInPreview = false;
 	private bool bPreserveAutoGeneratedResources = false;
@@ -228,7 +231,11 @@ class Localize : BuildCommand
 
 	public override void ExecuteBuild()
 	{
-		ParseCommandLine();
+		if (!ParseCommandLine())
+		{
+			Logger.LogError($"Errors detected in command line. Please fix the parameters and arguments for the command line and re-run the command..");
+			return;
+		}
 
 		var StartTime = DateTime.UtcNow;
 
@@ -364,7 +371,7 @@ class Localize : BuildCommand
 		Logger.LogInformation("Localize command finished in {Arg0} seconds", RunDuration / 1000);
 	}
 
-	private void ParseCommandLine()
+	private bool ParseCommandLine()
 	{
 		UEProjectRoot = ParseParamValue("UEProjectRoot", Default: CmdEnv.LocalRoot);
 		UEProjectDirectory = ParseParamValue("UEProjectDirectory");
@@ -374,7 +381,7 @@ class Localize : BuildCommand
 		}
 
 		UEProjectName = ParseParamValue("UEProjectName", Default: "");
-
+		
 		{
 			var LocalizationProjectNamesStr = ParseParamValue("LocalizationProjectNames");
 			if (LocalizationProjectNamesStr != null)
@@ -427,7 +434,6 @@ class Localize : BuildCommand
 				string AbsolutePathToExcludePluginsDirectory = Path.Combine(PluginsRootPath, ExcludePluginDirectory.Trim());
 				ExcludePlugins.AddRange(LocalizationUtilities.GetPluginNamesUnderDirectory(AbsolutePathToExcludePluginsDirectory, PluginsRootPath, UEProjectName.Length == 0 ? PluginType.Engine : PluginType.Project));
 			}
-
 		}
 
 		if (bShouldGatherPlugins)
@@ -490,6 +496,83 @@ class Localize : BuildCommand
 		{
 			Logger.LogInformation("Consolidating config files for various localization targets. A single instance of the Editor will be run to perform all localization gather steps.");
 		}
+
+		return ParseTemplateLocalizationConfigFiles();
+	}
+
+	private bool ParseTemplateLocalizationConfigFiles()
+	{
+		string TemplateLocalizationConfigFilesString = ParseParamValue("TemplateLocalizationConfigFiles");
+		if (String.IsNullOrEmpty(TemplateLocalizationConfigFilesString))
+		{
+			Logger.LogInformation("No template localization config files provided. Default localization config values will be used for Auto localization config generation.");
+			// No template parameter exists, it's still considered a success 
+			return true;
+		}
+
+		List<string> LocalizationConfigFileStrings = new List<string>();
+		bool bShouldEarlyOut = false;
+		foreach (string TemplateLocalizationConfigFile in TemplateLocalizationConfigFilesString.Split(';'))
+		{
+			string absoluteTemplateFilePath = Path.IsPathRooted(TemplateLocalizationConfigFile) 
+					? TemplateLocalizationConfigFile
+					: Path.Combine(UEProjectRoot, UEProjectDirectory, TemplateLocalizationConfigFile);
+			if (!File.Exists(absoluteTemplateFilePath))
+			{ 
+				Logger.LogError($"Provided relative template file path '{TemplateLocalizationConfigFile}' expands to absolute file path '{absoluteTemplateFilePath}' which cannot be found. Please check the template file paths provided are correct.");
+				bShouldEarlyOut = true; 
+				continue;
+			}
+			LocalizationConfigFileStrings.Add(absoluteTemplateFilePath);
+		}
+
+		if (bShouldEarlyOut)
+		{
+			Logger.LogError("1 or more provided template localization config files could not be found on disk. Please check the paths are correct and relative to the provided UEProjectDirectory.");
+			return false;
+		}
+
+		// We assume if more than 1 template file is provided that this is intended for modular localization config files 
+		if (LocalizationConfigFileStrings.Count > 1)
+		{
+			foreach (string TemplateLocalizationConfigFile in LocalizationConfigFileStrings)
+			{
+				if (!LocalizationUtilities.IsModularLocalizationConfigFile(TemplateLocalizationConfigFile))
+				{
+					Logger.LogError($"Provided template localization config file '{TemplateLocalizationConfigFile}' is not a Monolithic localization config file.");
+					bShouldEarlyOut = true;
+					continue;
+				}
+			}
+
+			if (bShouldEarlyOut)
+			{
+				Logger.LogError("Provided template localization config files all expected to be Monolithic localization config files. 1 or more provided template files are not Monolithic.");
+				return false;
+			}
+		}
+		// @TODOLocalization: Handle the monolithic template loc config file that's currently not supported 
+		// At this point we know all of the files exist and we have verified whether or not the files are all valid/if there are copies etc
+		// We load the localization config files up front here because the template files should really only be loaded once instead of every single time for every plugin 
+		foreach (string TemplateLocalizationConfigFile in LocalizationConfigFileStrings)
+		{
+			LocalizationConfigFile template = LocalizationConfigFile.Load(TemplateLocalizationConfigFile);
+			if (template == null)
+			{
+				Logger.LogError($"Failed to load template localization config file '{TemplateLocalizationConfigFile}'. Verify that the provided template localization config file is correct and properly formatted.");
+				bShouldEarlyOut = true;
+				continue;
+			}
+			Logger.LogInformation($"Template localization config file '{TemplateLocalizationConfigFile}' successfully loaded and will be be used for Auto localization config generation.");
+			TemplateLocalizationConfigFiles.Add(template);
+		}
+
+		if (bShouldEarlyOut)
+		{
+			Logger.LogError("1 or more provided template localization config files failed to be loaded. Verify that all provided template localization config files are valid and correctly formatted.");
+			return false;
+		}
+		return true;
 	}
 
 	private void AddStaticLocalizationBatches(List<LocalizationBatch> LocalizationBatches)
@@ -857,6 +940,7 @@ class Localize : BuildCommand
 
 	private void AutoGeneratePluginLocalizationFiles(HashSet<string> PluginNames)
 	{
+		Stopwatch stopWatch = Stopwatch.StartNew();
 		string PluginsRootDirectory = Path.Combine(UEProjectRoot, UEProjectDirectory);
 		DirectoryReference PluginsRootDirectoryReference = new DirectoryReference(PluginsRootDirectory);
 
@@ -883,7 +967,23 @@ class Localize : BuildCommand
 					LocalizationConfigFileGeneratorParams GeneratorParams = new();
 					GeneratorParams.LocalizationTargetName = Target.Name;
 					GeneratorParams.LocalizationTargetRootDirectory = Info.Directory.MakeRelativeTo(PluginsRootDirectoryReference);
-					List<LocalizationConfigFile> PluginLocalizationConfigFiles = Generator.GenerateDefaultSettingsConfigFiles(GeneratorParams);
+					List<LocalizationConfigFile> PluginLocalizationConfigFiles;
+					// If we parsed template localization config files to use, we generate with templates 
+					if (TemplateLocalizationConfigFiles.Count > 0)
+					{
+						PluginLocalizationConfigFiles = Generator.GenerateConfigFilesFromTemplate(GeneratorParams, TemplateLocalizationConfigFiles);
+					}
+					else
+					{
+						PluginLocalizationConfigFiles = Generator.GenerateDefaultSettingsConfigFiles(GeneratorParams);
+					}
+					
+					if (PluginLocalizationConfigFiles.Count != TemplateLocalizationConfigFiles.Count)
+					{
+						Logger.LogError($"1 or more localization config files failed to be auto generated for localization target '{GeneratorParams.LocalizationTargetName}'");
+						continue;
+// We don't early out because we want all generated localization config files to still be cleaned up and not let loose on disk 
+					}
 					foreach (LocalizationConfigFile PluginLocalizationConfigFile in PluginLocalizationConfigFiles)
 					{
 						string SaveFilePath = Path.Combine(PluginConfigLocalizationPath, PluginLocalizationConfigFile.Name);
@@ -902,6 +1002,8 @@ class Localize : BuildCommand
 				}
 			}
 		}
+		stopWatch.Stop();
+		Logger.LogInformation($"Generating localization config files for {PluginNames.Count} plugins took {stopWatch.ElapsedMilliseconds}ms.");
 	}
 
 	private void CleanUpAutoGeneratedFiles()
