Improved flush in VKPlay session logic:
*Added get matchmaking initialization rpc
*Added get pre travel state rpc
*Added get matchmaking error rpc
*Reworked flush in VKPlay logic

#rnx

diff --git a/Engine/Source/Programs/AutomationTool/Gauntlet/Framework/RpcFramework/Gauntlet.RpcLibrary.cs b/Engine/Source/Programs/AutomationTool/Gauntlet/Framework/RpcFramework/Gauntlet.RpcLibrary.cs
--- a/Engine/Source/Programs/AutomationTool/Gauntlet/Framework/RpcFramework/Gauntlet.RpcLibrary.cs
+++ b/Engine/Source/Programs/AutomationTool/Gauntlet/Framework/RpcFramework/Gauntlet.RpcLibrary.cs
@@ -5,6 +5,7 @@ using System;
 using System.Collections.Generic;
 using System.Diagnostics;
 using System.IO;
+using System.Linq;
 using System.Net.Http;
 using System.Threading;
 
@@ -455,7 +456,8 @@ namespace Gauntlet
 		public static void WaitForIncomingTargetMessage(List<RpcTarget> InTargets, string MessageCategory,
 		string MessageContents, int TimeoutInSeconds = 300, int TimeBetweenRetriesInSeconds = 60)
 		{
-			int TotalTargets = InTargets.Count;
+			int DeferedRolesAmount = InTargets.Count(Target => Target.IsDeferStart);
+			int TotalTargets = InTargets.Count - DeferedRolesAmount;
 			int TotalMatches = 0;
 
 			DateTime TimeoutEnd = DateTime.UtcNow.AddSeconds(TimeoutInSeconds);
@@ -563,7 +565,8 @@ namespace Gauntlet
 		/// <exception cref="TestException"></exception>
 		public static void WaitForAllTargetsToGetEndpoints(List<RpcTarget> InTargets, int TimeoutInSeconds = 300, int TimeBetweenRetriesInSeconds = 60)
 		{
-			int TotalTargets = InTargets.Count;
+			int DeferedRolesAmount = InTargets.Count(Target => Target.IsDeferStart);
+			int TotalTargets = InTargets.Count - DeferedRolesAmount;
 			int TargetsThatHaveRpcs = 0;
 			DateTime TimeoutEnd = DateTime.UtcNow.AddSeconds(TimeoutInSeconds);
 
@@ -574,6 +577,12 @@ namespace Gauntlet
 
 				foreach (RpcTarget Target in InTargets)
 				{
+					if (Target.IsDeferStart)
+					{
+						Target.LogInfo("Target {0} is marked as a defer start. Move on to the next target.", Target.TargetName);
+						continue;
+					}
+
 					Target.LogInfo("Attempting to get RPC's from Remote Process");
 
 					// target already retrieved RPC's, don't re-query
diff --git a/Engine/Source/Programs/AutomationTool/Gauntlet/Framework/RpcFramework/Gauntlet.RpcTargetRegistry.cs b/Engine/Source/Programs/AutomationTool/Gauntlet/Framework/RpcFramework/Gauntlet.RpcTargetRegistry.cs
--- a/Engine/Source/Programs/AutomationTool/Gauntlet/Framework/RpcFramework/Gauntlet.RpcTargetRegistry.cs
+++ b/Engine/Source/Programs/AutomationTool/Gauntlet/Framework/RpcFramework/Gauntlet.RpcTargetRegistry.cs
@@ -28,6 +28,7 @@ namespace Gauntlet
 		/// If this is true, we know the target is capable of sending messages to us and we no longer need to manually poll those processes for messages.
 		/// </summary>
 		public bool HasSentMessage = false;
+		public bool IsDeferStart = false;
 
 		public RpcTarget(string InIpAddress, int InPort, string InName, UnrealTargetRole InType = UnrealTargetRole.Client)
 		{
@@ -445,6 +446,11 @@ namespace Gauntlet
 					NewRpcTarget.IpAddress = UriOverride;
 				}
 
+				if (TestRole.DeferredLaunch)
+				{
+					NewRpcTarget.IsDeferStart = true;
+				}
+
 				// defines the port the role's process should listen on and adds to command line
 				TestRole.CommandLineParams.AddRawCommandline(string.Format("-rpcport={0}", NewRpcTarget.Port));
 				// Override HTTPServer.Listeners .ini settings
