Raising crash mini dump callstack as error and as Gauntlet_TestEvent id

#rnx

diff --git a/Engine/Source/Programs/AutomationTool/Gauntlet/Platform/Windows/Gauntlet.TargetDeviceWindows.cs b/Engine/Source/Programs/AutomationTool/Gauntlet/Platform/Windows/Gauntlet.TargetDeviceWindows.cs
--- a/Engine/Source/Programs/AutomationTool/Gauntlet/Platform/Windows/Gauntlet.TargetDeviceWindows.cs
+++ b/Engine/Source/Programs/AutomationTool/Gauntlet/Platform/Windows/Gauntlet.TargetDeviceWindows.cs
@@ -388,7 +388,7 @@ namespace Gauntlet
 			}
 			else
 			{
-				Log.Warning(KnownLogEvents.Gauntlet, "Wrote minidump to {FileName}", DumpName);
+				Log.Error(KnownLogEvents.Gauntlet_TestEvent, "Wrote minidump to {FileName}", DumpName);
 
 				if (CommandUtils.IsBuildMachine)
 				{
@@ -414,7 +414,7 @@ namespace Gauntlet
 						}
 						if (Stack.Any())
 						{
-							Log.Warning(KnownLogEvents.Gauntlet, "Stack analysis produced with 'cdb {Args}':\n{Callstack}", Args, string.Join('\n', Stack));
+							Log.Error(KnownLogEvents.Gauntlet_TestEvent, "Stack analysis produced with 'cdb {Args}':\n{Callstack}", Args, string.Join('\n', Stack));
 						}
 						else
 						{
