Added POST_NOTIFICATIONS permission to AndroidTargetDevice to avoid application hang up

diff --git a/Engine/Source/Programs/AutomationTool/Gauntlet/Platform/Android/Gauntlet.TargetDeviceAndroid.cs b/Engine/Source/Programs/AutomationTool/Gauntlet/Platform/Android/Gauntlet.TargetDeviceAndroid.cs
--- a/Engine/Source/Programs/AutomationTool/Gauntlet/Platform/Android/Gauntlet.TargetDeviceAndroid.cs
+++ b/Engine/Source/Programs/AutomationTool/Gauntlet/Platform/Android/Gauntlet.TargetDeviceAndroid.cs
@@ -961,7 +961,7 @@ namespace Gauntlet
 		/// </summary>
 		public void EnablePermissions(string AndroidPackageName)
 		{
-			List<string> Permissions = new List<string> { "WRITE_EXTERNAL_STORAGE", "READ_EXTERNAL_STORAGE", "GET_ACCOUNTS", "RECORD_AUDIO" };
+			List<string> Permissions = new List<string> { "POST_NOTIFICATIONS", "WRITE_EXTERNAL_STORAGE", "READ_EXTERNAL_STORAGE", "GET_ACCOUNTS", "RECORD_AUDIO" };
 			Permissions.ForEach(Permission =>
 			{
 				string CommandLine = string.Format("shell pm grant {0} android.permission.{1}", AndroidPackageName, Permission);
