Automation: Fixed the problem of .NET incompatibility preventing some external tools spawned by UAT from launching.

The problem is that when we launch UAT and UBT we run GetDotnetPath.bat to setup .NET environment variables to lock the .NET version to the one we bundle with the engine. However, these settings are inherited by all tools we launch e.g. using CommandUtils.Run. If the launched tool requires a different .NET version than the one the engine uses, they will fail to run. For instance, UE 5.4 is bundled with .NET 6, but tools from some platform SDKs require .NET 8.

Resolved by modifying CommandUtils.Run to clear DOTNET_ROOT when launching programs from outside of the root engine directory.

diff --git a/Engine/Source/Programs/AutomationTool/AutomationUtils/ProcessUtils.cs b/Engine/Source/Programs/AutomationTool/AutomationUtils/ProcessUtils.cs
--- a/Engine/Source/Programs/AutomationTool/AutomationUtils/ProcessUtils.cs
+++ b/Engine/Source/Programs/AutomationTool/AutomationUtils/ProcessUtils.cs
@@ -876,6 +876,30 @@ namespace AutomationTool
 			return ResolvedPath;
 		}
 
+		/// <summary>
+		/// Returns true if the given application executable is located withing the engine directory (as defined by CmdEnv.LocalRoot).
+		/// This check may not be perfectly accurate because it only relies on the string paths.
+		/// </summary>
+		/// <param name="App"></param>
+		/// <returns></returns>
+		private static bool IsEngineExecutable(string App)
+		{
+			if (!Path.IsPathFullyQualified(App))
+			{
+				App = WhichApp(App);
+
+				if (App == null)
+				{
+					return false;
+				}
+			}
+
+			string NormalizedRoot = ConvertSeparators(PathSeparator.Default, Path.GetFullPath(CmdEnv.LocalRoot));
+			string NormalizedApp  = ConvertSeparators(PathSeparator.Default, Path.GetFullPath(App));
+
+			return NormalizedApp.StartsWith(NormalizedRoot);
+		}
+
 		/// <summary>
 		/// Runs external program.
 		/// </summary>
@@ -926,6 +950,23 @@ namespace AutomationTool
 				LogWithVerbosity(SpewVerbosity,"Running: " + App + " " + (String.IsNullOrEmpty(CommandLine) ? "" : CommandLine));
 			}
 
+			if (!IsEngineExecutable(App))
+			{
+				Dictionary<string, string> FinalEnv = (Env != null) ? new Dictionary<string, string>(Env) : new Dictionary<string, string>();
+			
+				// We're clearing DOTNET_* for any external programs to avoid enforcing our version of .NET
+				// on applications that may not be compatible with it e.g. requiring a newer version that may be
+				// installed on the system as part of their installer.
+				FinalEnv["DOTNET_ROOT"] = "";
+				FinalEnv["DOTNET_MULTILEVEL_LOOKUP"] = "";
+				FinalEnv["DOTNET_ROLL_FORWARD"] = "";
+
+				// Clearing our variable as well to make sure that if the spawned tool ends up calling our .bat files, we will reinitialize the envvars.
+				FinalEnv["UE_DOTNET_VERSION"] = "";
+
+				Env = FinalEnv;					
+			}
+
 			bool bUseShellExecute = Options.HasFlag(ERunOptions.UseShellExecute);
 			bool bRedirectStdOut = !bUseShellExecute && !Options.HasFlag(ERunOptions.NoStdOutRedirect);
 			bool bAllowSpew = bRedirectStdOut && Options.HasFlag(ERunOptions.AllowSpew);
