Ensuring that stage-to-pak from zen snapshot for ios targets works correctly and handles *.metallib and *.metalmap files appropriately.

diff --git a/Engine/Source/Programs/AutomationTool/Scripts/CopyBuildToStagingDirectory.Automation.cs b/Engine/Source/Programs/AutomationTool/Scripts/CopyBuildToStagingDirectory.Automation.cs
--- a/Engine/Source/Programs/AutomationTool/Scripts/CopyBuildToStagingDirectory.Automation.cs
+++ b/Engine/Source/Programs/AutomationTool/Scripts/CopyBuildToStagingDirectory.Automation.cs
@@ -951,6 +951,7 @@ namespace AutomationScripts
 					{
 						string RelativeFilename = PackageDataField["clientpath"].AsString();
 						if (RelativeFilename.EndsWith(".bin") || RelativeFilename.EndsWith(".ushaderbytecode") ||
+							RelativeFilename.EndsWith(".metallib") || RelativeFilename.EndsWith(".metalmap") ||
 							(RelativeFilename.EndsWith(".txt") && RelativeFilename.Contains("/Metadata/ChunkManifest/"))
 							)
 						{
@@ -1186,6 +1187,99 @@ namespace AutomationScripts
 			return true;
 		}
 
+		private static async Task<byte[]> ReadChunkId(HttpClient Client, ZenServerStoreData StoreData, IoHash ChunkId)
+		{
+			using var Request = new HttpRequestMessage(HttpMethod.Get, string.Format("http://{0}:{1}/prj/{2}/oplog/{3}/{4}", StoreData.HostName, StoreData.HostPort, StoreData.ProjectId, StoreData.OplogId, ChunkId));
+			Request.Headers.Add("Accept", "application/octet-stream");
+			return await Client.Send(Request).Content.ReadAsByteArrayAsync();
+		}
+
+		private static FileReference ZenCookedFileToFileReference(ProjectParams Params, DeploymentContext SC, ZenCookedFile CookedFile, HttpClient Client, List<Task> TaskList, HashSet<string> TemporaryNonUFSFiles)
+		{
+			FileReference CookedFileReference = new FileReference(CookedFile.Filename);
+
+			if (CookedFile.ChunkId != IoHash.Zero)
+			{
+				StagedFileType? CookedFileStageType = GetStagingTypeForCookedFile(CookedFileReference);
+				if (CookedFileStageType.HasValue && (CookedFileStageType.Value == StagedFileType.NonUFS) && Params.UsePak(SC.StageTargetPlatform))
+				{
+					if (!FileReference.Exists(CookedFileReference))
+					{
+						Task NewTask = Task.Run( async () =>
+							{
+								InternalUtils.SafeCreateDirectory(CookedFileReference.Directory.FullName, true);
+								InternalUtils.SafeWriteAllBytes(CookedFile.Filename, await ReadChunkId(Client, SC.PackageStoreData.ZenServerStore, CookedFile.ChunkId));
+							});
+						TaskList.Add(NewTask);
+					}
+					TemporaryNonUFSFiles.Add(CookedFile.Filename);
+				}
+			}
+
+			return CookedFileReference;
+		}
+
+		private static void SynchronizeTemporaryNonUFSFiles(DeploymentContext SC, HashSet<string> TemporaryNonUFSFiles)
+		{
+			FileReference TemporaryNonUFSFilesManifest = FileReference.Combine(SC.MetadataDir, "TemporaryNonUFSFiles.txt");
+			List<string> OnDiskTemporaryNonUFSFiles = new List<string>();
+			if (InternalUtils.SafeFileExists(TemporaryNonUFSFilesManifest.FullName))
+			{
+				OnDiskTemporaryNonUFSFiles = new List<string>(InternalUtils.SafeReadAllLines(TemporaryNonUFSFilesManifest.FullName));
+				InternalUtils.SafeDeleteFile(TemporaryNonUFSFilesManifest.FullName);
+			}
+			foreach (string OnDiskFilename in OnDiskTemporaryNonUFSFiles)
+			{
+				if (!TemporaryNonUFSFiles.Contains(OnDiskFilename))
+				{
+					FileReference OnDiskReference = new FileReference(OnDiskFilename);
+					StagedFileReference ExistingStagedFile = new StagedFileReference(OnDiskReference.MakeRelativeTo(SC.PlatformCookDir));
+					SC.FilesToStage.NonUFSFiles.Remove(ExistingStagedFile);
+					InternalUtils.SafeDeleteFile(OnDiskFilename);
+				}
+			}
+			if (TemporaryNonUFSFiles.Any())
+			{
+				InternalUtils.SafeCreateDirectory(TemporaryNonUFSFilesManifest.Directory.FullName);
+				InternalUtils.SafeWriteAllLines(TemporaryNonUFSFilesManifest.FullName, TemporaryNonUFSFiles.ToArray());
+			}
+		}
+
+		private static List<FileReference> DiscoverCookedFiles(ProjectParams Params, DeploymentContext SC, bool bAlwaysFullyScanCookDir)
+		{
+			List<FileReference> CookedFiles = new List<FileReference>();
+
+			// When cooking to Zen get the list of cooked package files
+			HashSet<string> TemporaryNonUFSFiles = new HashSet<string>();
+			LoadPackageStoreData(Params, SC);
+			bool bReadZenCookedFiles = (SC.PackageStoreData != null && SC.PackageStoreData.ZenCookedFiles != null);
+			if (bReadZenCookedFiles)
+			{
+				if (!bAlwaysFullyScanCookDir && DirectoryReference.Exists(SC.PlatformCookDir))
+				{
+					CookedFiles.AddRange(DirectoryReference.EnumerateFiles(SC.PlatformCookDir, "*.projectstore", SearchOption.TopDirectoryOnly).ToList());
+				}
+				HttpClient Client = new HttpClient();
+				List<Task> TaskList = new List<Task>();
+				foreach (ZenCookedFile CookedFile in SC.PackageStoreData.ZenCookedFiles)
+				{
+					CookedFiles.Add(ZenCookedFileToFileReference(Params, SC, CookedFile, Client, TaskList, TemporaryNonUFSFiles));
+				}
+				Task.WaitAll(TaskList.ToArray());
+			}
+			SynchronizeTemporaryNonUFSFiles(SC, TemporaryNonUFSFiles);
+
+			if (bAlwaysFullyScanCookDir || !bReadZenCookedFiles)
+			{
+				if (DirectoryReference.Exists(SC.PlatformCookDir))
+				{
+					CookedFiles.AddRange(DirectoryReference.EnumerateFiles(SC.PlatformCookDir, "*", SearchOption.AllDirectories).ToList());
+				}
+			}
+
+			return CookedFiles;
+		}
+
 		public static void CreateStagingManifest(ProjectParams Params, DeploymentContext SC)
 		{
 			if (!Params.Stage)
@@ -1267,29 +1361,7 @@ namespace AutomationScripts
 				}
 
 				// Stage all the cooked data, this is the same rule as normal stage except we may skip Engine
-				List<FileReference> CookedFiles = new List<FileReference>();
-
-				// When cooking to Zen get the list of cooked package files
-				LoadPackageStoreData(Params, SC);
-				if (SC.PackageStoreData != null && SC.PackageStoreData.ZenCookedFiles != null)
-				{
-					if (DirectoryReference.Exists(SC.PlatformCookDir))
-					{
-						CookedFiles.AddRange(DirectoryReference.EnumerateFiles(SC.PlatformCookDir, "*.projectstore", SearchOption.TopDirectoryOnly).ToList());
-					}
-					foreach (ZenCookedFile CookedFile in SC.PackageStoreData.ZenCookedFiles)
-					{
-						CookedFiles.Add(new FileReference(CookedFile.Filename));
-					}
-				}
-				else
-				{
-
-					if (DirectoryReference.Exists(SC.PlatformCookDir))
-					{
-						CookedFiles.AddRange(DirectoryReference.EnumerateFiles(SC.PlatformCookDir, "*", SearchOption.AllDirectories).ToList());
-					}
-				}
+				List<FileReference> CookedFiles = DiscoverCookedFiles(Params, SC, /* bAlwaysFullyScanCookDir */ false);
 
 				foreach (FileReference CookedFile in CookedFiles)
 				{
@@ -1631,17 +1703,7 @@ namespace AutomationScripts
 					// Stage all the cooked data. Currently not filtering this by restricted folders, since we shouldn't mask invalid references by filtering them out.
 					if (DirectoryReference.Exists(SC.PlatformCookDir))
 					{
-						List<FileReference> CookedFiles = DirectoryReference.EnumerateFiles(SC.PlatformCookDir, "*", SearchOption.AllDirectories).ToList();
-
-						// When cooking to Zen get the list of cooked package files
-						LoadPackageStoreData(Params, SC);
-						if (SC.PackageStoreData != null && SC.PackageStoreData.ZenCookedFiles != null)
-						{
-							foreach (ZenCookedFile CookedFile in SC.PackageStoreData.ZenCookedFiles)
-							{
-								CookedFiles.Add(new FileReference(CookedFile.Filename));
-							}
-						}
+						List<FileReference> CookedFiles = DiscoverCookedFiles(Params, SC, /* bAlwaysFullyScanCookDir */ true);
 
 						foreach (FileReference CookedFile in CookedFiles)
 						{
@@ -4949,9 +5011,11 @@ namespace AutomationScripts
 						Logger.LogInformation("Pak rule '{Arg0}' with {Arg1} matching files, overrides chunk manifest and might cause broken references", KeyValue.Key, Filenames.Count);
 					}
 
-					String FilePath = Path.Combine(CmdEnv.EngineSavedFolder, "Logs", "ChunkManifestOverride_DebugFiles.txt");
+					String ChunkManifestOverrideDirectoryPath = Path.Combine(CmdEnv.EngineSavedFolder, "Logs");
+					String FilePath = Path.Combine(ChunkManifestOverrideDirectoryPath, "ChunkManifestOverride_DebugFiles.txt");
 					Logger.LogInformation("Saving chunk manifest override debug information to '{Arg0}'", FilePath);
 
+					InternalUtils.SafeCreateDirectory(ChunkManifestOverrideDirectoryPath);
 					using (var Writer = new StreamWriter(FilePath, false, new System.Text.UTF8Encoding(true)))
 					{
 						foreach (KeyValuePair<String, List<String>> KeyValue in RuleNameToFilenames)
