Adding explicit account TearDown call to allow tests to clean up accounts.

diff --git a/Engine/Source/Programs/AutomationTool/Gauntlet/Unreal/RunUnreal.cs b/Engine/Source/Programs/AutomationTool/Gauntlet/Unreal/RunUnreal.cs
--- a/Engine/Source/Programs/AutomationTool/Gauntlet/Unreal/RunUnreal.cs
+++ b/Engine/Source/Programs/AutomationTool/Gauntlet/Unreal/RunUnreal.cs
@@ -309,6 +309,9 @@ namespace Gauntlet
 
 			bool AllTestsPassed = ExecuteTests(ContextOptions, AllTestNodes);
 
+			// Tear down accounts
+			TeardownAccounts();
+
 			// dispose now, not during shutdown gc, because this runs commands...
 			DevicePool.Instance.Dispose();
 
@@ -476,6 +479,15 @@ namespace Gauntlet
 			}
 		}
 
+		/// <summary>
+		/// TearDown the account pool that was SetUp for tests.
+		/// </summary>
+		/// <returns></returns>
+		protected virtual void TeardownAccounts()
+		{
+			AccountPool.Shutdown();
+		}
+
 		protected void SetupDevices(UnrealTargetPlatform DefaultPlatform, UnrealTestOptions Options)
 		{
 			Reservation.ReservationDetails = Options.JobDetails;
