Adding AgeStore to StandardTasks class

Allows calling from BGv2

#rnx

diff --git a/Engine/Source/Programs/AutomationTool/BuildGraph/Tasks/AgeStoreTask.cs b/Engine/Source/Programs/AutomationTool/BuildGraph/Tasks/AgeStoreTask.cs
--- a/Engine/Source/Programs/AutomationTool/BuildGraph/Tasks/AgeStoreTask.cs
+++ b/Engine/Source/Programs/AutomationTool/BuildGraph/Tasks/AgeStoreTask.cs
@@ -7,6 +7,7 @@ using System.Linq;
 using System.Threading.Tasks;
 using System.Xml;
 using EpicGames.Core;
+using Microsoft.CodeAnalysis.CSharp.Syntax;
 using Microsoft.Extensions.Logging;
 using UnrealBuildTool;
 
@@ -252,4 +253,30 @@ namespace AutomationTool.Tasks
 			yield break;
 		}
 	}
+
+	public static partial class StandardTasks
+	{
+		/// <summary>
+		/// Task that strips symbols from a set of files
+		/// </summary>
+		/// <param name="platform">Platform to clean</param>
+		/// <param name="storeDir">Path to symbol store</param>
+		/// <param name="days">Number of days to keep</param>
+		/// <param name="buildDir">The root of the build directory to check for existing buildversion named directories</param>
+		/// <param name="filter">A substring to match in directory file names before deleting symbols</param>
+		/// <returns></returns>
+		public static async Task<FileSet> AgeStoreAsync(UnrealTargetPlatform platform, string storeDir, int days, string buildDir, string filter = null)
+		{
+			AgeStoreTaskParameters parameters = new AgeStoreTaskParameters()
+			{
+				Platform = platform,
+				StoreDir = storeDir,
+				Days = days,
+				BuildDir = buildDir,
+				Filter = filter
+			};
+
+			return await ExecuteAsync(new AgeStoreTask(parameters));
+		}
+	}
 }
