Take client and editor screenshots on test timeouts in RPC tests

#rnx

diff --git a/Engine/Source/Programs/AutomationTool/Gauntlet/Framework/RpcFramework/Gauntlet.RpcLibrary.cs b/Engine/Source/Programs/AutomationTool/Gauntlet/Framework/RpcFramework/Gauntlet.RpcLibrary.cs
--- a/Engine/Source/Programs/AutomationTool/Gauntlet/Framework/RpcFramework/Gauntlet.RpcLibrary.cs
+++ b/Engine/Source/Programs/AutomationTool/Gauntlet/Framework/RpcFramework/Gauntlet.RpcLibrary.cs
@@ -197,8 +197,10 @@ namespace Gauntlet
 
 		}
 
+		public delegate void OnBeforeTimeoutDelegate();
+
 		public static bool DoWhile(string FuncName, Func<bool> RpcCall, 
-			int DelayInMilliseconds = 5 * 1000, int TimeoutInMilliseconds = 60 * 1000, bool bSuppressExceptions = false)
+			int DelayInMilliseconds = 5 * 1000, int TimeoutInMilliseconds = 60 * 1000, bool bSuppressExceptions = false, OnBeforeTimeoutDelegate OnBeforeTimeout = null)
 		{
 			Stopwatch StopTimer = Stopwatch.StartNew();
 			while (true)
@@ -208,6 +210,7 @@ namespace Gauntlet
 					// Do while RpcCall() returns true
 					if (!RpcCall())
 					{
+						OnBeforeTimeout?.Invoke();
 						return true; // Return true to caller to signal we are done
 					}
 				}
@@ -221,6 +224,7 @@ namespace Gauntlet
 
 				if (StopTimer.ElapsedMilliseconds > TimeoutInMilliseconds)
 				{
+					OnBeforeTimeout?.Invoke();
 					throw new TestException($"{FuncName} timed out after {TimeoutInMilliseconds}ms");
 				}
 
