- Only allow enabled platforms to add "code specified" sdk versions in the set of Turnkey known SDK sources
(licensee verified)

diff --git a/Engine/Source/Programs/AutomationTool/Turnkey/TurnkeyManifest.cs b/Engine/Source/Programs/AutomationTool/Turnkey/TurnkeyManifest.cs
--- a/Engine/Source/Programs/AutomationTool/Turnkey/TurnkeyManifest.cs
+++ b/Engine/Source/Programs/AutomationTool/Turnkey/TurnkeyManifest.cs
@@ -170,18 +170,21 @@ namespace Turnkey
 				// also manually create local "FileSource" objects specified via code
 				foreach (UnrealTargetPlatform Platform in UnrealTargetPlatform.GetValidPlatforms())
 				{
-					// these are usually going to be empty
-					foreach (string SpecifiedVersion in AutomationTool.Platform.GetPlatform(Platform).GetCodeSpecifiedSdkVersions())
+					if (DataDrivenPlatformInfo.GetDataDrivenInfoForPlatform(Platform.ToString()).bIsEnabled)
 					{
-						FileSource CodeSource = FileSource.CreateCodeSpecifiedSource($"{Platform} SDK {SpecifiedVersion}", SpecifiedVersion, Platform);
-						DiscoveredFileSources.Add(CodeSource);
-					}
+						// these are usually going to be empty
+						foreach (string SpecifiedVersion in AutomationTool.Platform.GetPlatform(Platform).GetCodeSpecifiedSdkVersions())
+						{
+							FileSource CodeSource = FileSource.CreateCodeSpecifiedSource($"{Platform} SDK {SpecifiedVersion}", SpecifiedVersion, Platform);
+							DiscoveredFileSources.Add(CodeSource);
+						}
 
-					foreach(string SpecifiedVersion in AutomationTool.Platform.GetPlatform(Platform).GetCodeSpecifiedDeviceSoftwareUpdateVersions())
-					{
-						FileSource CodeSource = FileSource.CreateCodeSpecifiedSource($"{Platform} Device Software/Flash {SpecifiedVersion}", SpecifiedVersion, Platform);
-						CodeSource.Type = FileSource.SourceType.Flash;
-						DiscoveredFileSources.Add(CodeSource);
+						foreach (string SpecifiedVersion in AutomationTool.Platform.GetPlatform(Platform).GetCodeSpecifiedDeviceSoftwareUpdateVersions())
+						{
+							FileSource CodeSource = FileSource.CreateCodeSpecifiedSource($"{Platform} Device Software/Flash {SpecifiedVersion}", SpecifiedVersion, Platform);
+							CodeSource.Type = FileSource.SourceType.Flash;
+							DiscoveredFileSources.Add(CodeSource);
+						}
 					}
 				}
 			}
