AutomationTool: Don't build directly to the Engine\Binaries\DotNET directory

#rnx

diff --git a/Engine/Build/BatchFiles/BuildUAT.bat b/Engine/Build/BatchFiles/BuildUAT.bat
--- a/Engine/Build/BatchFiles/BuildUAT.bat
+++ b/Engine/Build/BatchFiles/BuildUAT.bat
@@ -8,9 +8,6 @@ rem ## if you copy it to a different location and run it.
 
 setlocal
 
-rem ## Make sure we use the Windows "find" utility and not a Unix-like tool found in PATH
-set FIND=%SYSTEMROOT%\System32\findstr.exe /i /r
-
 rem ## First, make sure the batch file exists in the folder we expect it to.  This is necessary in order to
 rem ## verify that our relative path to the /Engine/Source directory is correct
 if not exist "%~dp0..\..\Source" goto Error_BatchFileInWrongLocation
@@ -19,86 +16,57 @@ rem ## Change the CWD to /Engine/Source.
 pushd "%~dp0..\..\Source"
 if not exist ..\Build\BatchFiles\BuildUAT.bat goto Error_BatchFileInWrongLocation
 
+rem ## Verify that dotnet is present
+call "%~dp0GetDotnetPath.bat"
+if errorlevel 1 goto Error_NoDotnetSDK
+
 rem Ensure UnrealBuildTool is up to date as it is a prereq
 call "%~dp0\BuildUBT.bat"
 if errorlevel 1 goto Error_UBTCompileFailed
 
-rem Check to see if the files in the AutomationTool, EpicGames.Build, EpicGames.Core, or UnrealBuildTool directory have changed.
-rem find ".cs" files to only lines that match those names - excludes lines that will change for uninteresting reasons, like free space
-md ..\Intermediate\Build >nul 2>nul
-
-dir /s ^
- Programs\Shared\EpicGames.Build ^
- Programs\Shared\EpicGames.Core ^
- Programs\Shared\EpicGames.Horde ^
- Programs\Shared\EpicGames.IoHash ^
- Programs\Shared\EpicGames.MsBuild ^
- Programs\Shared\EpicGames.OIDC ^
- Programs\Shared\EpicGames.Serialization ^
- Programs\Shared\EpicGames.UBA ^
- Programs\Shared\EpicGames.UHT ^
- Programs\UnrealBuildTool ^
- Programs\AutomationTool ^
- 2>nul ^
- | %FIND% /c:"\.cs$" /c:"\.csproj$" > ..\Intermediate\Build\AutomationToolFiles.txt
-
-dir /s ^
- ..\Binaries\Win64\UnrealBuildAccelerator ^
- 2>nul ^
- | %FIND% /c:"\.dll$" /c:"\.exe$" >> ..\Intermediate\Build\AutomationToolFiles.txt
-
-if not exist ..\Platforms goto NoPlatforms
-for /d %%D in (..\Platforms\*) do (
-	dir /s ^
-	 %%D\Source\Programs\AutomationTool ^
-	 %%D\Source\Programs\Shared ^
-	 %%D\Source\Programs\UnrealBuildTool ^
-	 2>nul ^
-	 | %FIND% /c:"\.cs$" /c:"\.csproj$" >> ..\Intermediate\Build\AutomationToolFiles.txt
-)
-:NoPlatforms
-
-if not exist ..\Restricted goto NoRestricted
-for /d %%D in (..\Restricted\*) do (
-	dir /s ^
-	 %%D\Source\Programs\AutomationTool ^
-	 %%D\Source\Programs\Shared ^
-	 %%D\Source\Programs\UnrealBuildTool ^
-	 2>nul ^
-	 | %FIND% /c:"\.cs$" /c:"\.csproj$" >> ..\Intermediate\Build\AutomationToolFiles.txt
-)
-:NoRestricted
-
-rem note: no /s
-dir ^
- Programs\Shared\MetaData.cs ^
- 2>nul ^
- | %FIND% /c:"\.cs$" >>..\Intermediate\Build\AutomationToolFiles.txt
+rem ## Properties
+set DEPENDS_FILE=..\Intermediate\Build\AutomationTool.dep.csv 
+:tempRetry
+set TEMP_DEPENDS_PATH=%TEMP%\BuildUAT-%random%%random%
+if exist %TEMP_DEPENDS_PATH% goto :tempRetry
+set TEMP_DEPENDS_FILE=%TEMP_DEPENDS_PATH%.dep.csv
 
+rem ## Command line arguments
 set MSBUILD_LOGLEVEL=%1
 if not defined %MSBUILD_LOGLEVEL set MSBUILD_LOGLEVEL=quiet
 
-set ARGUMENT=%2
-if not defined %ARGUMENT goto Check_UpToDate
-if /I "%ARGUMENT%" == "FORCE" goto Build_AutomationTool
+rem Force build by deleting dependency file
+if /I "%2" == "FORCE" del %DEPENDS_FILE% >nul 2>&1
+
+rem Check to see if the files in the AutomationTool solution have changed
+mkdir "%TEMP_DEPENDS_PATH%"
+dotnet msbuild Programs\AutomationTool\AutomationTool.sln -t:ProjectDependencies -p:Configuration=Development -p:Platform="Any CPU" -p:OutputPath="%TEMP_DEPENDS_PATH%" -noLogo -v:%MSBUILD_LOGLEVEL%
+if errorlevel 1 goto Error_UATDependsFailed
+copy /b "%TEMP_DEPENDS_PATH%\*.dep.csv" "%TEMP_DEPENDS_FILE%" >nul 2>&1
+"%SYSTEMROOT%\System32\sort.exe" /unique /o "%TEMP_DEPENDS_FILE%" "%TEMP_DEPENDS_FILE%"
+if errorlevel 1 goto Error_UATDependsFailed
 
 :Check_UpToDate
 if not exist ..\Binaries\DotNET\AutomationTool\AutomationTool.dll goto Build_AutomationTool
-set RUNUAT_EXITCODE=0
 rem per https://ss64.com/nt/fc.html using redirection syntax rather than errorlevel, based on observed inconsistent results from this function
-fc ..\Intermediate\Build\AutomationToolFiles.txt ..\Intermediate\Build\AutomationToolPrevFiles.txt >nul 2>&1 && goto Exit
+fc /C "%TEMP_DEPENDS_FILE%" %DEPENDS_FILE% >nul 2>&1 && (
+	set RUNUAT_EXITCODE=0
+	goto Exit
+)
 
 :Build_AutomationTool
-rem ## Verify that dotnet is present
-call "%~dp0GetDotnetPath.bat"
-if errorlevel 1 goto Error_NoDotnetSDK
-
 echo Building AutomationTool...
 dotnet build Programs\AutomationTool\AutomationTool.csproj -c Development -v %MSBUILD_LOGLEVEL%
 if errorlevel 1 goto Error_UATCompileFailed
 
+echo Publishing AutomationTool...
+if not exist ..\Binaries\DotNET\AutomationTool mkdir ..\Binaries\DotNET\AutomationTool >nul 2>&1
+dotnet publish Programs\AutomationTool\AutomationTool.csproj -c Development --output ..\Binaries\DotNET\AutomationTool --no-build -v %MSBUILD_LOGLEVEL%
+if errorlevel 1 goto Error_UATCompileFailed
+
 rem record input files - regardless of how we got here, these are now our point of reference
-copy /y ..\Intermediate\Build\AutomationToolFiles.txt ..\Intermediate\Build\AutomationToolPrevFiles.txt >nul
+if not exist ..\Intermediate\Build mkdir ..\Intermediate\Build >nul 2>&1
+move /y "%TEMP_DEPENDS_FILE%" %DEPENDS_FILE% >nul
 
 goto Exit
 
@@ -117,6 +85,13 @@ echo.
 set RUNUAT_EXITCODE=1
 goto Exit
 
+:Error_UATDependsFailed
+echo.
+echo RunUAT ERROR: AutomationTool failed to check dependencies.
+echo.
+set RUNUAT_EXITCODE=1
+goto Exit
+
 :Error_UATCompileFailed
 echo.
 echo RunUAT ERROR: AutomationTool failed to compile.
@@ -129,6 +104,9 @@ set RUNUAT_EXITCODE=1
 goto Exit
 
 :Exit
+rem ## Clean up temp files
+del "%TEMP_DEPENDS_FILE%" >nul 2>&1
+rmdir /s /q "%TEMP_DEPENDS_PATH%" >nul 2>&1
 rem ## Restore original CWD in case we change it
 popd
 exit /B %RUNUAT_EXITCODE%
diff --git a/Engine/Build/BatchFiles/BuildUAT.sh b/Engine/Build/BatchFiles/BuildUAT.sh
--- a/Engine/Build/BatchFiles/BuildUAT.sh
+++ b/Engine/Build/BatchFiles/BuildUAT.sh
@@ -28,72 +28,29 @@ fi
 
 MSBuild_Verbosity="${1:-quiet}"
 
-# Check to see if the files in the AutomationTool, EpicGames.Build, EpicGames.Core, or UnrealBuildTool
-# directory have changed.
+# Check to see if the files in the AutomationTool solution have changed
 
 mkdir -p ../Intermediate/Build
 
+DEPENDS_FILE=../Intermediate/Build/AutomationTool.dep.csv
+TEMP_DEPENDS_FILE=$(mktemp)
+
+dotnet msbuild Programs/AutomationTool/AutomationTool.sln -t:ProjectDependencies -p:Configuration=Development -p:DependsPath="$TEMP_DEPENDS_FILE" -p:DependsOverwrite=false -m:1 -noLogo -v:"$MSBuild_Verbosity"
+sort -u -o "$TEMP_DEPENDS_FILE" "$TEMP_DEPENDS_FILE"
+
 PERFORM_REBUILD=0
-if [ "$1" == "FORCE" ]; then
+if [ "$2" == "FORCE" ]; then
   PERFORM_REBUILD=1
-  echo "Rebuilding: build requested"
+  echo "Rebuilding: Force build requested"
 
 elif [ ! -f ../Binaries/DotNET/AutomationTool/AutomationTool.dll ]; then
   PERFORM_REBUILD=1
   echo "Rebuilding: AutomationTool assembly not found"
 
-elif [ -f ../Intermediate/Build/AutomationToolLastBuildTime ]; then
-  UPDATED_DEP_FILES="$(find \
-    Programs/Shared/EpicGames.Build \
-    Programs/Shared/EpicGames.Core \
-    Programs/Shared/EpicGames.Horde \
-    Programs/Shared/EpicGames.IoHash \
-    Programs/Shared/EpicGames.MsBuild \
-    Programs/Shared/EpicGames.OIDC \
-    Programs/Shared/EpicGames.Serialization \
-    Programs/Shared/EpicGames.UBA \
-    Programs/Shared/EpicGames.UHT \
-    Programs/UnrealBuildTool \
-    Programs/AutomationTool \
-    ../Restricted/**/Source/Programs/Shared \
-    ../Platforms/*/Source/Programs/Shared \
-    ../Restricted/**/Source/Programs/UnrealBuildTool \
-    ../Platforms/*/Source/Programs/UnrealBuildTool \
-    ../Restricted/**/Source/Programs/AutomationTool \
-    ../Platforms/*/Source/Programs/AutomationTool \
-    -type f \
-    \( -iname \*.cs -or -iname \*.csproj \) \
-    -newer ../Intermediate/Build/AutomationToolLastBuildTime)"
-  
-  if [ -n "$UPDATED_DEP_FILES" ]; then
-    PERFORM_REBUILD=1
-    echo "Rebuilding: found updated files:"
-    echo "$UPDATED_DEP_FILES"
-  fi
-
-  UPDATED_AUTOMATIONTOOL_FILES="$(find \
-    Programs/Shared \
-    -maxdepth 1 \
-    -type f \
-    \( -iname \*.cs -or -iname \*.csproj \) \
-    -newer ../Intermediate/Build/AutomationToolLastBuildTime)"
-  if [ -n "$UPDATED_AUTOMATIONTOOL_FILES" ]; then
-    PERFORM_REBUILD=1
-    echo "Rebuilding: Found updated files:"
-    echo "$UPDATED_AUTOMATIONTOOL_FILES"
-  fi
-
-  UPDATED_DEP_FILES="$(find \
-    ../Binaries/Linux* \
-    ../Binaries/Mac \
-    -maxdepth 2 \
-    -type f \
-    \( -iname \*Uba* \) \
-    -newer ../Intermediate/Build/AutomationToolLastBuildTime)"
-  if [ -n "$UPDATED_DEP_FILES" ]; then
+elif [ -f ../Intermediate/Build/AutomationTool.dep.csv ]; then
+  if ! cmp --silent $DEPENDS_FILE "$TEMP_DEPENDS_FILE"; then
     PERFORM_REBUILD=1
-    echo "Rebuilding: Found updated files:"
-    echo "$UPDATED_DEP_FILES"
+    echo "Rebuilding: Found updated files"
   fi
 
 else
@@ -102,26 +59,25 @@ else
 fi
 
 if [ $PERFORM_REBUILD -eq 1 ]; then
-  if [ "$(uname)" = "Darwin" ]; then
-    # Setup Environment
-    source "$SCRIPT_DIR/Mac/SetupEnvironment.sh" -dotnet "$SCRIPT_DIR/Mac"
-  fi
-
-  if [ "$(uname)" = "Linux" ]; then
-    # Setup Environment
-    source "$SCRIPT_DIR/Linux/SetupEnvironment.sh" -dotnet "$SCRIPT_DIR/Linux"
+  echo "Building AutomationTool..."
+  dotnet build Programs/AutomationTool/AutomationTool.csproj -c Development -v "$MSBuild_Verbosity"
+  if [ $? -ne 0 ]; then
+    echo "AutomationTool compilation failed"
+    popd >/dev/null || exit 1
+    exit 1
   fi
 
-  dotnet build Programs/AutomationTool/AutomationTool.csproj -c Development -v $MSBuild_Verbosity
+  echo "Publishing AutomationTool..."
+  dotnet publish Programs/AutomationTool/AutomationTool.csproj -c Development --output ../Binaries/DotNET/AutomationTool --no-build -v "$MSBuild_Verbosity"
   if [ $? -ne 0 ]; then
-    echo "Compilation failed"
-    popd >/dev/null
+    echo "AutomationTool publish failed"
+    popd >/dev/null || exit 1
     exit 1
   fi
 
-  touch ../Intermediate/Build/AutomationToolLastBuildTime
+  mv -f "$TEMP_DEPENDS_FILE" $DEPENDS_FILE >/dev/null
 else
-  echo "AutomationTool.dll is up to date"
+  echo "AutomationTool is up to date"
 fi
 
 popd >/dev/null
diff --git a/Engine/Build/BatchFiles/BuildUBT.sh b/Engine/Build/BatchFiles/BuildUBT.sh
--- a/Engine/Build/BatchFiles/BuildUBT.sh
+++ b/Engine/Build/BatchFiles/BuildUBT.sh
@@ -29,8 +29,7 @@ fi
 
 MSBuild_Verbosity="${1:-quiet}"
 
-# Check to see if the files in the UnrealBuildTool, EpicGames.Build, EpicGames.Core, or UnrealBuildTool
-# directory have changed.
+# Check to see if the files in the UnrealBuildTool solution have changed
 
 mkdir -p ../Intermediate/Build
 
diff --git a/Engine/Source/Programs/AutomationTool.Tests/AutomationTool.Tests.csproj b/Engine/Source/Programs/AutomationTool.Tests/AutomationTool.Tests.csproj
--- a/Engine/Source/Programs/AutomationTool.Tests/AutomationTool.Tests.csproj
+++ b/Engine/Source/Programs/AutomationTool.Tests/AutomationTool.Tests.csproj
@@ -1,4 +1,5 @@
 <Project Sdk="Microsoft.NET.Sdk">
+  <Import Project="..\Shared\UnrealEngine.csproj.props" />
 
   <PropertyGroup>
     <TargetFramework>net8.0</TargetFramework>
diff --git a/Engine/Source/Programs/AutomationTool/AutomationTool.csproj b/Engine/Source/Programs/AutomationTool/AutomationTool.csproj
--- a/Engine/Source/Programs/AutomationTool/AutomationTool.csproj
+++ b/Engine/Source/Programs/AutomationTool/AutomationTool.csproj
@@ -1,14 +1,8 @@
 <Project Sdk="Microsoft.NET.Sdk">
   <Import Project="..\Shared\UnrealEngine.csproj.props" />
-  
-  <PropertyGroup>
-    <IsWindows Condition="'$([System.Runtime.InteropServices.RuntimeInformation]::IsOSPlatform($([System.Runtime.InteropServices.OSPlatform]::Windows)))' == 'true'">true</IsWindows> 
-    <IsOSX Condition="'$([System.Runtime.InteropServices.RuntimeInformation]::IsOSPlatform($([System.Runtime.InteropServices.OSPlatform]::OSX)))' == 'true'">true</IsOSX> 
-    <IsLinux Condition="'$([System.Runtime.InteropServices.RuntimeInformation]::IsOSPlatform($([System.Runtime.InteropServices.OSPlatform]::Linux)))' == 'true'">true</IsLinux> 
-  </PropertyGroup>
-  
+
   <Choose>
-    <When Condition="'$(IsWindows)'=='true'">
+    <When Condition="'$(IsWindows)' == 'true'">
       <PropertyGroup>
         <TargetFramework>net8.0-windows10.0.19041</TargetFramework>
         <UseWindowsForms>true</UseWindowsForms>
@@ -33,12 +27,11 @@
     <AssemblyName>AutomationTool</AssemblyName>
     <WarningsNotAsErrors>612,618</WarningsNotAsErrors>
     <NoWarn>$(NoWarn);NETSDK1206</NoWarn>
-    <OutputPath>..\..\..\Binaries\DotNET\AutomationTool</OutputPath>
     <EnableDefaultItems>false</EnableDefaultItems>
     <AppendRuntimeIdentifierToOutputPath>false</AppendRuntimeIdentifierToOutputPath>
     <SatelliteResourceLanguages>en</SatelliteResourceLanguages> <!-- remove non english resource languages -->
     <DebugType>pdbonly</DebugType>
-    <DebugType Condition="'$([System.Runtime.InteropServices.RuntimeInformation]::IsOSPlatform($([System.Runtime.InteropServices.OSPlatform]::Windows)))' == 'true' And '$([System.Runtime.InteropServices.RuntimeInformation]::OSArchitecture)' == 'Arm64'">portable</DebugType>
+    <DebugType Condition="'$(IsWindows)' == 'true' And '$(IsArm64)' == 'true'">portable</DebugType>
     <AllowUnsafeBlocks>true</AllowUnsafeBlocks> <!-- required for FileMatcher support code in MSBuild subdir -->
   </PropertyGroup>
   
@@ -81,8 +74,7 @@
     </When>
   </Choose>
 
-  <!-- This set of dependencies should be kept as small as possible
-       This list must match the list of locations searched for changed files in BuildUAT.bat and BuildUAT.sh -->
+  <!-- This set of dependencies should be kept as small as possible -->
   <ItemGroup>
     <Compile Include="..\Shared\MetaData.cs">
       <Link>Properties\MetaData.cs</Link>
@@ -101,45 +93,24 @@
     <PackageReference Include="Microsoft.Build" Version="17.11.4" ExcludeAssets="runtime" PrivateAssets="all">
       <IncludeAssets>compile; build; native; contentfiles; analyzers; buildtransitive</IncludeAssets>
     </PackageReference>
-    <PackageReference Include="Microsoft.Build.Locator" Version="1.7.8" PrivateAssets="all" />
+    <PackageReference Include="Microsoft.Build.Locator" Version="1.7.8" />
     <!-- Added as a temporary workaround to running from visual studio where the system version of dotnet is run. 
          Adding a loader context resolves many of the issues, but it still fails when MSBuild tries to write the json -->
     <PackageReference Include="Newtonsoft.Json" Version="13.0.3" />
   </ItemGroup>
-
-  <!-- Ensure Magick.NET native binaries are copied -->
-  <PropertyGroup>
-    <MagickCopyNativeWindows>true</MagickCopyNativeWindows>
-    <MagickCopyNativeLinux>true</MagickCopyNativeLinux>
-    <MagickCopyNativeLinuxMusl>true</MagickCopyNativeLinuxMusl>
-    <MagickCopyNativeMacOS>true</MagickCopyNativeMacOS>
-  </PropertyGroup>
   
-  <!-- This set of dependencies should be kept as small as possible
-       This list must match the list of locations searched for changed files in BuildUAT.bat and BuildUAT.sh -->
+  <!-- This set of dependencies should be kept as small as possible -->
   <ItemGroup>
     <ProjectReference Include="..\Shared\EpicGames.Core\EpicGames.Core.csproj" PrivateAssets="all" />
     <ProjectReference Include="..\Shared\EpicGames.Build\EpicGames.Build.csproj" PrivateAssets="all" />
     <ProjectReference Include="..\UnrealBuildTool\UnrealBuildTool.csproj" PrivateAssets="all" />
   </ItemGroup>
 
-  <Choose>
-    <When Condition="'$([System.Runtime.InteropServices.RuntimeInformation]::IsOSPlatform($([System.Runtime.InteropServices.OSPlatform]::Windows)))' == 'true' And '$([System.Runtime.InteropServices.RuntimeInformation]::OSArchitecture)' != 'Arm64'">
-      <PropertyGroup Condition="Exists('..\..\..\Restricted\NotForLicensees\Source\Programs\UnrealBuildTool\Executors\Experimental\BoxExecutor.cs')">
-        <DefineConstants>$(DefineConstants);__BOXEXECUTOR_AVAILABLE__</DefineConstants>
-      </PropertyGroup>
-      <PropertyGroup Condition="Exists('..\..\..\Restricted\NotForLicensees\Source\Programs\Shared\EpicGames.Box\BoxServer.cs')">
-        <DefineConstants>$(DefineConstants);__BOX_AVAILABLE__</DefineConstants>
-      </PropertyGroup>
-    </When>
-  </Choose>
-
+  <!-- Various batch files and shell scripts that build or run AutomationTool -->
   <ItemGroup>
-    <None Include="..\..\..\Build\BatchFiles\BuildUAT.bat"><Link>BatchFiles\BuildUAT.bat</Link></None>
-    <None Include="..\..\..\Build\BatchFiles\BuildUAT.sh"><Link>BatchFiles\BuildUAT.sh</Link></None>
-    <None Include="..\..\..\Build\BatchFiles\RunUAT.bat"><Link>BatchFiles\RunUAT.bat</Link></None>
-    <None Include="..\..\..\Build\BatchFiles\RunUAT.sh"><Link>BatchFiles\RunUAT.sh</Link></None>
-    <None Include="..\..\..\Build\BatchFiles\RunUAT.command"><Link>BatchFiles\RunUAT.command</Link></None>
+    <None Include="../../../Build/BatchFiles/**/*.*">
+      <Link>$([System.Text.RegularExpressions.Regex]::Replace(%(FullPath), '^.+?[\\/]BatchFiles[\\/]', 'BatchFiles/'))</Link>
+    </None>
   </ItemGroup>
   
   <!-- Reference file is generated by GenerateProjectFiles scripts (calling UnrealBuildTool). It contains a list of all project files known to UBT based on how it was run. -->
diff --git a/Engine/Source/Programs/AutomationTool/AutomationUtils/AutomationUtils.Automation.csproj b/Engine/Source/Programs/AutomationTool/AutomationUtils/AutomationUtils.Automation.csproj
--- a/Engine/Source/Programs/AutomationTool/AutomationUtils/AutomationUtils.Automation.csproj
+++ b/Engine/Source/Programs/AutomationTool/AutomationUtils/AutomationUtils.Automation.csproj
@@ -9,12 +9,14 @@
     <AssemblyName>AutomationUtils.Automation</AssemblyName>
     <GenerateAssemblyInfo>false</GenerateAssemblyInfo>
     <GenerateTargetFrameworkAttribute>false</GenerateTargetFrameworkAttribute>
-    <OutputPath>..\..\..\..\Binaries\DotNET\AutomationTool\AutomationUtils</OutputPath>
     <Configurations>Debug;Release;Development</Configurations>
     <TreatWarningsAsErrors>true</TreatWarningsAsErrors>
     <AppendTargetFrameworkToOutputPath>false</AppendTargetFrameworkToOutputPath>
     <DebugType>pdbonly</DebugType>
-    <DebugType Condition="'$([System.Runtime.InteropServices.RuntimeInformation]::IsOSPlatform($([System.Runtime.InteropServices.OSPlatform]::Windows)))' == 'true' And '$([System.Runtime.InteropServices.RuntimeInformation]::OSArchitecture)' == 'Arm64'">portable</DebugType>
+    <DebugType Condition="'$(IsWindows)' == 'true' And '$(IsArm64)' == 'true'">portable</DebugType>
+  </PropertyGroup>
+  <PropertyGroup Condition="'$(EpicGamesMsBuild)' == 'true'">
+    <OutputPath>..\..\..\..\Binaries\DotNET\AutomationTool\AutomationUtils</OutputPath>
   </PropertyGroup>
   <ItemGroup>
     <Reference Include="fastJSON">
diff --git a/Engine/Source/Programs/AutomationTool/Program.cs b/Engine/Source/Programs/AutomationTool/Program.cs
--- a/Engine/Source/Programs/AutomationTool/Program.cs
+++ b/Engine/Source/Programs/AutomationTool/Program.cs
@@ -425,12 +425,12 @@ namespace AutomationToolDriver
 			try
 			{
 				// Set the working directory to the Unreal root directory
-				Environment.CurrentDirectory = Path.GetFullPath(Path.Combine(Path.GetDirectoryName(Assembly.GetExecutingAssembly().GetOriginalLocation()), "..", "..", "..", ".."));
+				DirectoryReference.SetCurrentDirectory(Unreal.RootDirectory);
 
 				// Ensure we can resolve any external assemblies as necessary.
-				string PathToBinariesDotNET = Path.GetDirectoryName(Assembly.GetEntryAssembly().GetOriginalLocation());
-				AssemblyUtils.InstallAssemblyResolver(PathToBinariesDotNET);
-				AssemblyUtils.InstallRecursiveAssemblyResolver(PathToBinariesDotNET);
+				DirectoryReference AutomationToolLocation = new DirectoryReference(Path.GetDirectoryName(Assembly.GetEntryAssembly().GetOriginalLocation()));
+				AssemblyUtils.InstallAssemblyResolver(AutomationToolLocation.FullName);
+				AssemblyUtils.InstallRecursiveAssemblyResolver(AutomationToolLocation.FullName);
 
 				// Log the operating environment. Since we usually compile to AnyCPU, we may be executed using different system paths under WOW64.
 				Logger.LogDebug("Running on {Platform} as a {Bitness}-bit process.", RuntimePlatform.Current.ToString(), Environment.Is64BitProcess ? 64 : 32);
@@ -503,6 +503,7 @@ namespace AutomationToolDriver
 
 			// Load AutomationUtils.Automation.dll
 			FileReference AssemblyPath = ScriptModuleAssemblyPaths.FirstOrDefault(x => x.GetFileNameWithoutExtension().Contains("AutomationUtils.Automation"));
+			AssemblyUtils.InstallRecursiveAssemblyResolver(AssemblyPath.Directory.FullName);
 			Assembly AutomationUtilsAssembly = AssemblyPath != null ? Assembly.LoadFrom(AssemblyPath.FullName) : null;
 
 			if (AutomationUtilsAssembly == null)
diff --git a/Engine/Source/Programs/Shared/EpicGames.Core/AssemblyUtils.cs b/Engine/Source/Programs/Shared/EpicGames.Core/AssemblyUtils.cs
--- a/Engine/Source/Programs/Shared/EpicGames.Core/AssemblyUtils.cs
+++ b/Engine/Source/Programs/Shared/EpicGames.Core/AssemblyUtils.cs
@@ -78,10 +78,13 @@ namespace EpicGames.Core
 
 		private static void RefreshAssemblyCache(string rootDirectory, string pattern = "*.dll")
 		{
-			// Initialize our cache of assemblies by enumerating all files in the given folder.
-			foreach (string discoveredAssembly in Directory.EnumerateFiles(rootDirectory, pattern, SearchOption.AllDirectories))
+			if (Directory.Exists(rootDirectory))
 			{
-				AddFileToAssemblyCache(discoveredAssembly);
+				// Initialize our cache of assemblies by enumerating all files in the given folder.
+				foreach (string discoveredAssembly in Directory.EnumerateFiles(rootDirectory, pattern, SearchOption.AllDirectories))
+				{
+					AddFileToAssemblyCache(discoveredAssembly);
+				}
 			}
 		}
 
diff --git a/Engine/Source/Programs/Shared/EpicGames.MsBuild/CsProjBuilder.cs b/Engine/Source/Programs/Shared/EpicGames.MsBuild/CsProjBuilder.cs
--- a/Engine/Source/Programs/Shared/EpicGames.MsBuild/CsProjBuilder.cs
+++ b/Engine/Source/Programs/Shared/EpicGames.MsBuild/CsProjBuilder.cs
@@ -184,6 +184,8 @@ namespace EpicGames.MsBuild
 			Dictionary<string, string> globalProperties = new Dictionary<string, string>
 			{
 				{ "EngineDir", hook.EngineDirectory.FullName },
+				{ "EnginePath", hook.EngineDirectory.FullName },
+				{ "EpicGamesMsBuild", "true" },
 #if DEBUG
 				{ "Configuration", "Debug" },
 #else
diff --git a/Engine/Source/Programs/Shared/EpicGames.ProjectStore/EpicGames.ProjectStore.csproj b/Engine/Source/Programs/Shared/EpicGames.ProjectStore/EpicGames.ProjectStore.csproj
--- a/Engine/Source/Programs/Shared/EpicGames.ProjectStore/EpicGames.ProjectStore.csproj
+++ b/Engine/Source/Programs/Shared/EpicGames.ProjectStore/EpicGames.ProjectStore.csproj
@@ -1,4 +1,6 @@
 <Project Sdk="Microsoft.NET.Sdk">
+  <Import Project="..\UnrealEngine.csproj.props" />
+  
   <PropertyGroup>
     <TargetFramework>net8.0</TargetFramework>
     <Configuration Condition=" '$(Configuration)' == '' ">Development</Configuration>
