Rpc Framework - added option for invalidated rpc target and reducing logging level for errors for invalid targets

diff --git a/Engine/Source/Programs/AutomationTool/Gauntlet/Framework/RpcFramework/Gauntlet.RpcExecutor.cs b/Engine/Source/Programs/AutomationTool/Gauntlet/Framework/RpcFramework/Gauntlet.RpcExecutor.cs
--- a/Engine/Source/Programs/AutomationTool/Gauntlet/Framework/RpcFramework/Gauntlet.RpcExecutor.cs
+++ b/Engine/Source/Programs/AutomationTool/Gauntlet/Framework/RpcFramework/Gauntlet.RpcExecutor.cs
@@ -805,10 +805,19 @@ namespace Gauntlet
 				HttpResponseMessage RpcResponse = ExecutorClient.Send(RpcRequest);
 				return RpcResponse;
 			}
-			catch (Exception e)
+			catch (Exception Error)
 			{
+				string Message = $"CallRpc failed with following exception: {Error}";
+
 				Log.SuspendECErrorParsing();
-				Log.Info($"CallRpc failed with the following exception: {e}");
+				if (InTarget.IsInvalidated())
+				{
+					Log.Verbose(Message);
+				}
+				else
+				{
+					Log.Info(Message);
+				}
 				Log.ResumeECErrorParsing();
 				return null;
 			}
@@ -860,11 +869,20 @@ namespace Gauntlet
 				HttpResponseMessage RpcResponse = ExecutorClient.Send(RpcRequest);
         		return RpcResponse;
         	}
-        	catch (Exception e)
+        	catch (Exception Error)
         	{
-        		Log.SuspendECErrorParsing();
-        		Log.Info(String.Format("CallRpc failed with following exception: {0}", e.Message));
-        		Log.ResumeECErrorParsing();
+				string Message = $"CallRpc failed with following exception: {Error}";
+
+				Log.SuspendECErrorParsing();
+				if (InTarget.IsInvalidated())
+				{
+					Log.Verbose(Message);
+				}
+				else
+				{
+					Log.Info(Message);
+				}
+				Log.ResumeECErrorParsing();
         		return null;
         	}
         }
diff --git a/Engine/Source/Programs/AutomationTool/Gauntlet/Framework/RpcFramework/Gauntlet.RpcTargetRegistry.cs b/Engine/Source/Programs/AutomationTool/Gauntlet/Framework/RpcFramework/Gauntlet.RpcTargetRegistry.cs
--- a/Engine/Source/Programs/AutomationTool/Gauntlet/Framework/RpcFramework/Gauntlet.RpcTargetRegistry.cs
+++ b/Engine/Source/Programs/AutomationTool/Gauntlet/Framework/RpcFramework/Gauntlet.RpcTargetRegistry.cs
@@ -29,6 +29,8 @@ namespace Gauntlet
 		/// </summary>
 		public bool HasSentMessage = false;
 		public bool IsDeferStart = false;
+		protected bool bIsInvalidated = false;
+
 
 		public RpcTarget(string InIpAddress, int InPort, string InName, UnrealTargetRole InType = UnrealTargetRole.Client)
 		{
@@ -134,6 +136,15 @@ namespace Gauntlet
 			return string.Format(Format, Args);
 		}
 
+		public void SetIsInvalidated(bool InIsInvalidated = true)
+		{
+			bIsInvalidated = InIsInvalidated;
+		}
+
+		public bool IsInvalidated()
+		{
+			return bIsInvalidated;
+		}
 	}
 
 	public class RpcTargetRegistry
