Fix static analysis warning.

#rnx

diff --git a/Engine/Source/Programs/AutomationTool/BuildGraph/Tasks/CreateArtifactTask.cs b/Engine/Source/Programs/AutomationTool/BuildGraph/Tasks/CreateArtifactTask.cs
--- a/Engine/Source/Programs/AutomationTool/BuildGraph/Tasks/CreateArtifactTask.cs
+++ b/Engine/Source/Programs/AutomationTool/BuildGraph/Tasks/CreateArtifactTask.cs
@@ -3,7 +3,6 @@
 using System;
 using System.Collections.Generic;
 using System.Diagnostics;
-using System.IO;
 using System.Linq;
 using System.Threading;
 using System.Threading.Tasks;
@@ -198,7 +197,7 @@ namespace AutomationTool.Tasks
 			// Create the new artifact
 			IHordeClient hordeClient = CommandUtils.ServiceProvider.GetRequiredService<IHordeClient>();
 			IArtifactBuilder artifact = await hordeClient.Artifacts.CreateAsync(name, type, _parameters.Description, streamId, commitId, keys, metadata);
-			Logger.LogInformation("Creating artifact {ArtifactId} '{ArtifactName}' ({ArtifactType}) from {NumFiles} files ({TotalSize:n1}mb). Namespace: {NamespaceId}, ref {RefName} ({Link})", artifact.Id, name, type,files.Count, totalSize / (1024.0 * 1024.0), artifact.NamespaceId, artifact.RefName, new Uri(hordeClient.ServerUrl, $"/api/v1/storage/{artifact.NamespaceId}/refs/{artifact.RefName}").ToString());
+			Logger.LogInformation("Creating artifact {ArtifactId} '{ArtifactName}' ({ArtifactType}) from {NumFiles} files ({TotalSize:n1}mb). Namespace: {NamespaceId}, ref {RefName} ({Link})", artifact.Id, name, type, files.Count, totalSize / (1024.0 * 1024.0), artifact.NamespaceId, artifact.RefName, new Uri(hordeClient.ServerUrl, $"/api/v1/storage/{artifact.NamespaceId}/refs/{artifact.RefName}").ToString());
 
 			// Upload the files
 			Stopwatch timer = Stopwatch.StartNew();
