Horde: Improve symbol indexing time; perform hash calculations in parallel, and upload aliases in smaller batches.

diff --git a/Engine/Source/Programs/AutomationTool/BuildGraph/Tasks/CreateArtifactTask.cs b/Engine/Source/Programs/AutomationTool/BuildGraph/Tasks/CreateArtifactTask.cs
--- a/Engine/Source/Programs/AutomationTool/BuildGraph/Tasks/CreateArtifactTask.cs
+++ b/Engine/Source/Programs/AutomationTool/BuildGraph/Tasks/CreateArtifactTask.cs
@@ -1,6 +1,7 @@
 // Copyright Epic Games, Inc. All Rights Reserved.
 
 using System;
+using System.Collections.Concurrent;
 using System.Collections.Generic;
 using System.Diagnostics;
 using System.Linq;
@@ -245,36 +246,47 @@ namespace AutomationTool.Tasks
 					rootRef = await blobWriter.WriteFilesAsync(baseDir, files);
 				}
 			}
+			double uploadTimeSecs = timer.Elapsed.TotalSeconds;
 
 			// Tag any uploaded symbols
 			if (_parameters.Symbols)
 			{
-				foreach (FileReference file in files)
-				{
-					string? hash = await SymStore.GetHashAsync(file, Logger, default);
-					if (hash != null)
-					{
-						string path = file.MakeRelativeTo(baseDir);
+				// Find all the hashes
+				ConcurrentBag<(IBlobRef Blob, string Alias)> aliases = new ConcurrentBag<(IBlobRef Blob, string Alias)>();
+				await Parallel.ForEachAsync(files, (file, cancellationToken) => ComputeHashAsync(baseDir, file, rootRef, aliases, cancellationToken));
 
-						FileEntry? fileEntry = await FindFileAsync(rootRef, path);
-						if (fileEntry == null)
-						{
-							Logger.LogWarning("Unable to find file {Path} in uploaded data.", path);
-						}
-						else
-						{
-							string fileName = file.GetFileName().ToUpperInvariant();
-							string alias = $"sym:{fileName}/{hash}/{fileName}";
-							Logger.LogInformation("Adding symbol alias: {Alias}", alias);
-							await artifact.AddAliasAsync(alias, fileEntry.Target.Handle);
-						}
-					}
+				// Add them to the artifact
+				foreach ((IBlobRef blob, string alias) in aliases)
+				{
+					await artifact.AddAliasAsync(alias, blob);
 				}
 			}
 
 			// Finalize the artifact
 			await artifact.CompleteAsync(rootRef);
-			Logger.LogInformation("Uploaded artifact {ArtifactId} in {Time:n1}s ({Rate:n1}mb/s, {RateMbps:n1}mbps)", artifact.Id, timer.Elapsed.TotalSeconds, totalSize / (timer.Elapsed.TotalSeconds * 1024.0 * 1024.0), (totalSize * 8.0) / (timer.Elapsed.TotalSeconds * 1024.0 * 1024.0));
+			Logger.LogInformation("Uploaded artifact {ArtifactId} in {Time:n1}s ({Rate:n1}mb/s, {RateMbps:n1}mbps)", artifact.Id, timer.Elapsed.TotalSeconds, totalSize / (uploadTimeSecs * 1024.0 * 1024.0), (totalSize * 8.0) / (uploadTimeSecs * 1024.0 * 1024.0));
+		}
+
+		static async ValueTask ComputeHashAsync(DirectoryReference baseDir, FileReference file, IBlobRef<DirectoryNode> rootRef, ConcurrentBag<(IBlobRef, string)> aliases, CancellationToken cancellationToken)
+		{
+			string? hash = await SymStore.GetHashAsync(file, Logger, default);
+			if (hash != null)
+			{
+				string path = file.MakeRelativeTo(baseDir);
+
+				FileEntry? fileEntry = await FindFileAsync(rootRef, path);
+				if (fileEntry == null)
+				{
+					Logger.LogWarning("Unable to find file {Path} in uploaded data.", path);
+				}
+				else
+				{
+					string fileName = file.GetFileName().ToUpperInvariant();
+					string alias = $"sym:{fileName}/{hash}/{fileName}";
+					Logger.LogInformation("Adding symbol alias: {Alias}", alias);
+					aliases.Add((fileEntry.Target.Handle, alias));
+				}
+			}
 		}
 
 		static async Task<FileEntry?> FindFileAsync(IBlobRef<DirectoryNode> rootDir, string path, CancellationToken cancellationToken = default)
diff --git a/Engine/Source/Programs/Shared/EpicGames.Horde/Storage/IStorageWriter.cs b/Engine/Source/Programs/Shared/EpicGames.Horde/Storage/IStorageWriter.cs
--- a/Engine/Source/Programs/Shared/EpicGames.Horde/Storage/IStorageWriter.cs
+++ b/Engine/Source/Programs/Shared/EpicGames.Horde/Storage/IStorageWriter.cs
@@ -148,7 +148,7 @@ namespace EpicGames.Horde.Storage
 		async Task CheckForFlushAsync(CancellationToken cancellationToken = default)
 		{
 			int numItems = _request.AddAliases.Count + _request.RemoveAliases.Count + _request.AddRefs.Count + _request.RemoveRefs.Count;
-			if (numItems > 1000)
+			if (numItems >= 100)
 			{
 				await FlushAsync(cancellationToken);
 			}
@@ -157,8 +157,11 @@ namespace EpicGames.Horde.Storage
 		/// <inheritdoc/>
 		public async Task FlushAsync(CancellationToken cancellationToken = default)
 		{
-			await _storageBackend.UpdateMetadataAsync(_request, cancellationToken);
-			_request = new UpdateMetadataRequest();
+			if (_request.AddAliases.Count > 0 || _request.RemoveAliases.Count > 0 || _request.AddRefs.Count > 0 || _request.RemoveRefs.Count > 0)
+			{
+				await _storageBackend.UpdateMetadataAsync(_request, cancellationToken);
+				_request = new UpdateMetadataRequest();
+			}
 		}
 	}
 }
