BuildGraph: Fix static analysis warning.

#rnx

diff --git a/Engine/Source/Programs/AutomationTool/BuildGraph/Tasks/CreateArtifactTask.cs b/Engine/Source/Programs/AutomationTool/BuildGraph/Tasks/CreateArtifactTask.cs
--- a/Engine/Source/Programs/AutomationTool/BuildGraph/Tasks/CreateArtifactTask.cs
+++ b/Engine/Source/Programs/AutomationTool/BuildGraph/Tasks/CreateArtifactTask.cs
@@ -254,7 +254,7 @@ namespace AutomationTool.Tasks
 				// Find all the hashes
 				ConcurrentBag<(IBlobRef Blob, string Alias)> aliases = new ConcurrentBag<(IBlobRef Blob, string Alias)>();
 				await Parallel.ForEachAsync(files, (file, cancellationToken) => ComputeHashAsync(baseDir, file, rootRef, aliases, cancellationToken));
-
+				
 				// Add them to the artifact
 				foreach ((IBlobRef blob, string alias) in aliases)
 				{
@@ -269,7 +269,7 @@ namespace AutomationTool.Tasks
 
 		static async ValueTask ComputeHashAsync(DirectoryReference baseDir, FileReference file, IBlobRef<DirectoryNode> rootRef, ConcurrentBag<(IBlobRef, string)> aliases, CancellationToken cancellationToken)
 		{
-			string? hash = await SymStore.GetHashAsync(file, Logger, default);
+			string? hash = await SymStore.GetHashAsync(file, Logger, cancellationToken);
 			if (hash != null)
 			{
 				string path = file.MakeRelativeTo(baseDir);
