BuildGraph: Static analysis fix.

#rnx

diff --git a/Engine/Source/Programs/AutomationTool/BuildGraph/Tasks/CreateArtifactTask.cs b/Engine/Source/Programs/AutomationTool/BuildGraph/Tasks/CreateArtifactTask.cs
--- a/Engine/Source/Programs/AutomationTool/BuildGraph/Tasks/CreateArtifactTask.cs
+++ b/Engine/Source/Programs/AutomationTool/BuildGraph/Tasks/CreateArtifactTask.cs
@@ -274,7 +274,7 @@ namespace AutomationTool.Tasks
 			{
 				string path = file.MakeRelativeTo(baseDir);
 
-				FileEntry? fileEntry = await FindFileAsync(rootRef, path);
+				FileEntry? fileEntry = await FindFileAsync(rootRef, path, cancellationToken);
 				if (fileEntry == null)
 				{
 					Logger.LogWarning("Unable to find file {Path} in uploaded data.", path);
