- Fixed IOS plist generation to take a receipt instead of calcuating the path, which can be wrong for Restricted Programs

diff --git a/Engine/Platforms/VisionOS/Source/Programs/AutomationTool/VisionOSPlatform.Automation.cs b/Engine/Platforms/VisionOS/Source/Programs/AutomationTool/VisionOSPlatform.Automation.cs
--- a/Engine/Platforms/VisionOS/Source/Programs/AutomationTool/VisionOSPlatform.Automation.cs
+++ b/Engine/Platforms/VisionOS/Source/Programs/AutomationTool/VisionOSPlatform.Automation.cs
@@ -32,19 +32,9 @@ public class VisionOSPlatform : IOSPlatform
 		return VisionOSExports.PrepForUATPackageOrDeploy(Config, ProjectFile, InProjectName, InProjectDirectory, Executable, InEngineDir, bForDistribution, CookFlavor, bIsDataDeploy, bCreateStubIPA, TargetReceiptFileName, Log.Logger);
 	}
 
-	public override bool DeployGeneratePList(FileReference ProjectFile, UnrealTargetConfiguration Config, DirectoryReference ProjectDirectory, bool bIsUEGame, string GameName, bool bIsClient, string ProjectName, DirectoryReference InEngineDir, DirectoryReference AppDirectory, string InExecutablePath)
+	public override bool DeployGeneratePList(FileReference ProjectFile, UnrealTargetConfiguration Config, DirectoryReference ProjectDirectory, bool bIsUEGame, string GameName, bool bIsClient, string ProjectName, DirectoryReference InEngineDir, DirectoryReference AppDirectory, TargetReceipt Receipt)
 	{
-		string TargetName = Path.GetFileNameWithoutExtension(InExecutablePath).Split("-".ToCharArray())[0];
-		FileReference TargetReceiptFileName;
-		if (bIsUEGame)
-		{
-			TargetReceiptFileName = TargetReceipt.GetDefaultPath(InEngineDir, "UnrealGame", UnrealTargetPlatform.VisionOS, Config, null);
-		}
-		else
-		{
-			TargetReceiptFileName = TargetReceipt.GetDefaultPath(ProjectDirectory, TargetName, UnrealTargetPlatform.VisionOS, Config, null);
-		}
-		return VisionOSExports.GeneratePList(ProjectFile, Config, ProjectDirectory, bIsUEGame, GameName, bIsClient, ProjectName, InEngineDir, AppDirectory, TargetReceiptFileName, Log.Logger);
+		return VisionOSExports.GeneratePList(ProjectFile, Config, ProjectDirectory, bIsUEGame, GameName, bIsClient, ProjectName, InEngineDir, AppDirectory, Receipt, Log.Logger);
 	}
 
     public override string GetCookPlatform(bool bDedicatedServer, bool bIsClientOnly)
@@ -120,7 +110,7 @@ public class VisionOSPlatform : IOSPlatform
 						"VisionOS",
 						"Payload",
 						(SC.IsCodeBasedProject ? SC.ShortProjectName : "UnrealGame") + ".app"),
-						SC.StageExecutables[0]);
+						SC.StageTargets[0].Receipt);
                 }
 
 
diff --git a/Engine/Platforms/VisionOS/Source/Programs/UnrealBuildTool/VisionOSExports.cs b/Engine/Platforms/VisionOS/Source/Programs/UnrealBuildTool/VisionOSExports.cs
--- a/Engine/Platforms/VisionOS/Source/Programs/UnrealBuildTool/VisionOSExports.cs
+++ b/Engine/Platforms/VisionOS/Source/Programs/UnrealBuildTool/VisionOSExports.cs
@@ -44,12 +44,11 @@ namespace UnrealBuildTool
 		/// <param name="ProjectName"></param>
 		/// <param name="InEngineDir"></param>
 		/// <param name="AppDirectory"></param>
-		/// <param name="BuildReceiptFileName"></param>
+		/// <param name="Receipt"></param>
 		/// <param name="Logger"></param>
 		/// <returns></returns>
-		public static bool GeneratePList(FileReference ProjectFile, UnrealTargetConfiguration Config, DirectoryReference ProjectDirectory, bool bIsUnrealGame, string GameName, bool bIsClient, string ProjectName, DirectoryReference InEngineDir, DirectoryReference AppDirectory, FileReference BuildReceiptFileName, ILogger Logger)
+		public static bool GeneratePList(FileReference ProjectFile, UnrealTargetConfiguration Config, DirectoryReference ProjectDirectory, bool bIsUnrealGame, string GameName, bool bIsClient, string ProjectName, DirectoryReference InEngineDir, DirectoryReference AppDirectory, TargetReceipt Receipt, ILogger Logger)
 		{
-			TargetReceipt Receipt = TargetReceipt.Read(BuildReceiptFileName);
 			return new UEDeployVisionOS(Logger).GeneratePList(ProjectFile, Config, ProjectDirectory.FullName, bIsUnrealGame, GameName, bIsClient, ProjectName, InEngineDir.FullName, AppDirectory.FullName, Receipt);
 		}
 	}
diff --git a/Engine/Source/Programs/AutomationTool/IOS/IOSPlatform.Automation.cs b/Engine/Source/Programs/AutomationTool/IOS/IOSPlatform.Automation.cs
--- a/Engine/Source/Programs/AutomationTool/IOS/IOSPlatform.Automation.cs
+++ b/Engine/Source/Programs/AutomationTool/IOS/IOSPlatform.Automation.cs
@@ -588,6 +588,7 @@ public class IOSPlatform : ApplePlatform
 
 	private FileReference GetTargetReceiptFileName(UnrealTargetConfiguration Config, string InExecutablePath, DirectoryReference InEngineDir, DirectoryReference InProjectDirectory, FileReference ProjectFile, bool bIsUEGame)
 	{
+Console.WriteLine("target receipt executable path: {0}", InExecutablePath);
 		string TargetName = Path.GetFileNameWithoutExtension(InExecutablePath).Split("-".ToCharArray())[0];
 		FileReference TargetReceiptFileName;
 		UnrealArchitectures Architectures = UnrealArchitectureConfig.ForPlatform(UnrealTargetPlatform.IOS).ActiveArchitectures(ProjectFile, TargetName);
@@ -607,10 +608,9 @@ public class IOSPlatform : ApplePlatform
 		IOSExports.GetProvisioningData(InProject, bDistribution, out MobileProvision, out SigningCertificate, out TeamUUID, out bAutomaticSigning);
 	}
 
-	public virtual bool DeployGeneratePList(FileReference ProjectFile, UnrealTargetConfiguration Config, DirectoryReference ProjectDirectory, bool bIsUEGame, string GameName, bool bIsClient, string ProjectName, DirectoryReference InEngineDir, DirectoryReference AppDirectory, string InExecutablePath)
+	public virtual bool DeployGeneratePList(FileReference ProjectFile, UnrealTargetConfiguration Config, DirectoryReference ProjectDirectory, bool bIsUEGame, string GameName, bool bIsClient, string ProjectName, DirectoryReference InEngineDir, DirectoryReference AppDirectory, TargetReceipt Receipt)
 	{
-		FileReference TargetReceiptFileName = GetTargetReceiptFileName(Config, InExecutablePath, InEngineDir, ProjectDirectory, ProjectFile, bIsUEGame);
-		return IOSExports.GeneratePList(ProjectFile, Config, ProjectDirectory, bIsUEGame, GameName, bIsClient, ProjectName, InEngineDir, AppDirectory, TargetReceiptFileName, Log.Logger);
+		return IOSExports.GeneratePList(ProjectFile, Config, ProjectDirectory, bIsUEGame, GameName, bIsClient, ProjectName, InEngineDir, AppDirectory, Receipt, Log.Logger);
 	}
 
 	protected string MakeIPAFileName(UnrealTargetConfiguration TargetConfiguration, ProjectParams Params, DeploymentContext SC, bool bAllowDistroPrefix)
@@ -1389,7 +1389,7 @@ public class IOSPlatform : ApplePlatform
 							SC.IsCodeBasedProject ? false : Params.Client, // Code based projects will have Client in their executable name already
 							SC.ShortProjectName, DirectoryReference.Combine(SC.LocalRoot, "Engine"),
 							DirectoryReference.Combine((SC.IsCodeBasedProject ? SC.ProjectRoot : DirectoryReference.Combine(SC.LocalRoot, "Engine")), "Binaries", PlatformName, AppleExports.UseModernXcode(SC.RawProjectPath) ? "" : "Payload", (SC.IsCodeBasedProject ? SC.ShortProjectName : "UnrealGame") + ".app"),
-							SC.StageExecutables[0]);
+							SC.StageTargets[0].Receipt);
 
 					if (!AppleExports.UseModernXcode(SC.RawProjectPath))
 					{
diff --git a/Engine/Source/Programs/AutomationTool/TVOS/TVOSPlatform.Automation.cs b/Engine/Source/Programs/AutomationTool/TVOS/TVOSPlatform.Automation.cs
--- a/Engine/Source/Programs/AutomationTool/TVOS/TVOSPlatform.Automation.cs
+++ b/Engine/Source/Programs/AutomationTool/TVOS/TVOSPlatform.Automation.cs
@@ -37,19 +37,9 @@ public class TVOSPlatform : IOSPlatform
 		TVOSExports.GetProvisioningData(InProject, bDistribution, out MobileProvision, out SigningCertificate, out Team, out bAutomaticSigning);
     }
 
-	public override bool DeployGeneratePList(FileReference ProjectFile, UnrealTargetConfiguration Config, DirectoryReference ProjectDirectory, bool bIsUEGame, string GameName, bool bIsClient, string ProjectName, DirectoryReference InEngineDir, DirectoryReference AppDirectory, string InExecutablePath)
+	public override bool DeployGeneratePList(FileReference ProjectFile, UnrealTargetConfiguration Config, DirectoryReference ProjectDirectory, bool bIsUEGame, string GameName, bool bIsClient, string ProjectName, DirectoryReference InEngineDir, DirectoryReference AppDirectory, TargetReceipt Receipt)
 	{
-		string TargetName = Path.GetFileNameWithoutExtension(InExecutablePath).Split("-".ToCharArray())[0];
-		FileReference TargetReceiptFileName;
-		if (bIsUEGame)
-		{
-			TargetReceiptFileName = TargetReceipt.GetDefaultPath(InEngineDir, "UnrealGame", UnrealTargetPlatform.TVOS, Config, null);
-		}
-		else
-		{
-			TargetReceiptFileName = TargetReceipt.GetDefaultPath(ProjectDirectory, TargetName, UnrealTargetPlatform.TVOS, Config, null);
-		}
-		return TVOSExports.GeneratePList(ProjectFile, Config, ProjectDirectory, bIsUEGame, GameName, bIsClient, ProjectName, InEngineDir, AppDirectory, TargetReceiptFileName, Log.Logger);
+		return TVOSExports.GeneratePList(ProjectFile, Config, ProjectDirectory, bIsUEGame, GameName, bIsClient, ProjectName, InEngineDir, AppDirectory, Receipt, Log.Logger);
 	}
 
     public override string GetCookPlatform(bool bDedicatedServer, bool bIsClientOnly)
@@ -132,7 +122,7 @@ public class TVOSPlatform : IOSPlatform
 						"TVOS",
 						"Payload",
 						(SC.IsCodeBasedProject ? SC.ShortProjectName : "UnrealGame") + ".app"),
-						SC.StageExecutables[0]);
+						SC.StageTargets[0].Receipt);
                 }
 
 
diff --git a/Engine/Source/Programs/UnrealBuildTool/Platform/IOS/IOSExports.cs b/Engine/Source/Programs/UnrealBuildTool/Platform/IOS/IOSExports.cs
--- a/Engine/Source/Programs/UnrealBuildTool/Platform/IOS/IOSExports.cs
+++ b/Engine/Source/Programs/UnrealBuildTool/Platform/IOS/IOSExports.cs
@@ -97,12 +97,11 @@ namespace UnrealBuildTool
 		/// <param name="ProjectName"></param>
 		/// <param name="InEngineDir"></param>
 		/// <param name="AppDirectory"></param>
-		/// <param name="BuildReceiptFileName"></param>
+		/// <param name="Receipt"></param>
 		/// <param name="Logger">Logger for output</param>
 		/// <returns></returns>
-		public static bool GeneratePList(FileReference ProjectFile, UnrealTargetConfiguration Config, DirectoryReference ProjectDirectory, bool bIsUnrealGame, string GameName, bool bIsClient, string ProjectName, DirectoryReference InEngineDir, DirectoryReference AppDirectory, FileReference BuildReceiptFileName, ILogger Logger)
+		public static bool GeneratePList(FileReference ProjectFile, UnrealTargetConfiguration Config, DirectoryReference ProjectDirectory, bool bIsUnrealGame, string GameName, bool bIsClient, string ProjectName, DirectoryReference InEngineDir, DirectoryReference AppDirectory, TargetReceipt Receipt, ILogger Logger)
 		{
-			TargetReceipt Receipt = TargetReceipt.Read(BuildReceiptFileName);
 			return new UEDeployIOS(Logger).GeneratePList(ProjectFile, Config, ProjectDirectory.FullName, bIsUnrealGame, GameName, bIsClient, ProjectName, InEngineDir.FullName, AppDirectory.FullName, Receipt);
 		}
 
diff --git a/Engine/Source/Programs/UnrealBuildTool/Platform/TVOS/TVOSExports.cs b/Engine/Source/Programs/UnrealBuildTool/Platform/TVOS/TVOSExports.cs
--- a/Engine/Source/Programs/UnrealBuildTool/Platform/TVOS/TVOSExports.cs
+++ b/Engine/Source/Programs/UnrealBuildTool/Platform/TVOS/TVOSExports.cs
@@ -91,12 +91,11 @@ namespace UnrealBuildTool
 		/// <param name="ProjectName"></param>
 		/// <param name="InEngineDir"></param>
 		/// <param name="AppDirectory"></param>
-		/// <param name="BuildReceiptFileName"></param>
+		/// <param name="Receipt"></param>
 		/// <param name="Logger"></param>
 		/// <returns></returns>
-		public static bool GeneratePList(FileReference ProjectFile, UnrealTargetConfiguration Config, DirectoryReference ProjectDirectory, bool bIsUnrealGame, string GameName, bool bIsClient, string ProjectName, DirectoryReference InEngineDir, DirectoryReference AppDirectory, FileReference BuildReceiptFileName, ILogger Logger)
+		public static bool GeneratePList(FileReference ProjectFile, UnrealTargetConfiguration Config, DirectoryReference ProjectDirectory, bool bIsUnrealGame, string GameName, bool bIsClient, string ProjectName, DirectoryReference InEngineDir, DirectoryReference AppDirectory, TargetReceipt Receipt, ILogger Logger)
 		{
-			TargetReceipt Receipt = TargetReceipt.Read(BuildReceiptFileName);
 			return new UEDeployTVOS(Logger).GeneratePList(ProjectFile, Config, ProjectDirectory.FullName, bIsUnrealGame, GameName, bIsClient, ProjectName, InEngineDir.FullName, AppDirectory.FullName, Receipt);
 		}
 	}
