Add a config option -ini:Engine:[Ias]:CompressionFormatOverride that sets a compression format to apply to OnDemand containers if the target platform has compression disabled.

#rnx

- The OnDemand data we upload uses the same compression as the container that we extract the data from. If a target platform has no compression applied at container level then the chunks we upload will be uncompressed and large.
- This new config value allows us to declare a compression format that should be applied to OnDemand containers anyway in these cases. All other compression settings will be taken from the usual packaging settings.
- This is a temporary measure and should be replaced at some point with a system that properly supports per contain compression settings.
- We do not apply this compression if the '-ForceUncompressed' option is being used. It is assumed that forcing compressiong to be disabled should be applied to all staged of the container creation.
- We do not apply this compression if the target platform is set to already apply compression to containers, in that case we allow the global settings to be applied.
- Moved the parsing of engine config files earlier in ::CreatePaks so that we can read from it.

diff --git a/Engine/Source/Programs/AutomationTool/Scripts/CopyBuildToStagingDirectory.Automation.cs b/Engine/Source/Programs/AutomationTool/Scripts/CopyBuildToStagingDirectory.Automation.cs
--- a/Engine/Source/Programs/AutomationTool/Scripts/CopyBuildToStagingDirectory.Automation.cs
+++ b/Engine/Source/Programs/AutomationTool/Scripts/CopyBuildToStagingDirectory.Automation.cs
@@ -3307,6 +3307,9 @@ namespace AutomationScripts
 					SC.CustomConfig,
 					Params.ConfigOverrideParams.ToArray());
 
+			ConfigHierarchy PlatformEngineConfig;
+			bool bHasPlatformEngineConfig = Params.EngineConfigs.TryGetValue(SC.StageTargetPlatform.PlatformType, out PlatformEngineConfig);
+
 			// in standard runs, ProjectPackagingSettings/bCompressed is read by the program invoking this script
 			//	and used to pass "-compressed" on the command line
 			// so typically "-compressed" == ProjectPackagingSettings/bCompressed
@@ -3394,6 +3397,33 @@ namespace AutomationScripts
 				bAnyCompressed = bAnyCompressed || PakParams.bCompressed;
 			}
 
+			// Only consider enabling the OnDemand override if we are NOT forcing uncompressed data but compression has not been
+			// enabled by the normal means.
+
+			string OnDemandCompressionFormat;
+			PlatformEngineConfig.GetString("Ias", "CompressionFormatOverride", out OnDemandCompressionFormat);
+			
+			if (!Params.ForceUncompressed && !Params.Compressed && !string.IsNullOrWhiteSpace(OnDemandCompressionFormat))
+			{
+				bool bHasOnDemandContainers = false;
+				foreach (CreatePakParams PakParams in PakParamsList)
+				{
+					if (PakParams.bOnDemand)
+					{
+						PakParams.bCompressed = true;
+						bHasOnDemandContainers = true;
+					}
+				}
+
+				if (bHasOnDemandContainers)
+				{
+					bAnyCompressed = true;
+					CompressionFormats = OnDemandCompressionFormat;
+					Logger.LogInformation("Overriding to CompressionFormatOverride = {IasCompressionFormat}", OnDemandCompressionFormat);
+					// Note that the compression method will still be set by ini:Game:[/Script/UnrealEd.ProjectPackagingSettings]:PackageCompressionMethod
+				}
+			}
+
 			if (bAnyCompressed && !string.IsNullOrWhiteSpace(CompressionFormats))
 			{
 				CompressionFormats = " -compressionformats=" + CompressionFormats;
@@ -3534,9 +3564,6 @@ namespace AutomationScripts
 				}
 			}
 
-			ConfigHierarchy PlatformEngineConfig;
-			bool bHasPlatformEngineConfig = false;
-			bHasPlatformEngineConfig = Params.EngineConfigs.TryGetValue(SC.StageTargetPlatform.PlatformType, out PlatformEngineConfig);
 			string BulkOption = "";
 			if (bHasPlatformEngineConfig)
 			{
