If no iOS udid is specified, don't include "-u" option. Needed for standard ushell ".deploy" command

#rnx

diff --git a/Engine/Source/Programs/AutomationTool/IOS/IOSPlatform.Automation.cs b/Engine/Source/Programs/AutomationTool/IOS/IOSPlatform.Automation.cs
--- a/Engine/Source/Programs/AutomationTool/IOS/IOSPlatform.Automation.cs
+++ b/Engine/Source/Programs/AutomationTool/IOS/IOSPlatform.Automation.cs
@@ -1927,7 +1927,11 @@ Console.WriteLine("target receipt executable path: {0}", InExecutablePath);
 		// deploy the .ipa
 		var DeviceInstaller = GetPathToLibiMobileDeviceTool("ideviceinstaller");
 
-		string LibimobileDeviceArguments =  "-u " + Params.DeviceNames[0] + " -i " + "\"" +  Path.GetFullPath(AppToDeploy) + "\"";
+		string LibimobileDeviceArguments =  "-i " + "\"" +  Path.GetFullPath(AppToDeploy) + "\"";
+		if (Params.DeviceNames[0].Length > 1)
+		{
+			LibimobileDeviceArguments = " -u " + Params.DeviceNames[0];
+		}
 		LibimobileDeviceArguments = GetLibimobileDeviceNetworkedArgument(LibimobileDeviceArguments, Params.DeviceNames[0]);
 
 		// If we deploying to a Simulator, use "xcrun simctl" instead
