Fix for dropping the bundle id when installing via QuickLaunch

#rnx

diff --git a/Engine/Source/Programs/AutomationTool/IOS/IOSPlatform.Automation.cs b/Engine/Source/Programs/AutomationTool/IOS/IOSPlatform.Automation.cs
--- a/Engine/Source/Programs/AutomationTool/IOS/IOSPlatform.Automation.cs
+++ b/Engine/Source/Programs/AutomationTool/IOS/IOSPlatform.Automation.cs
@@ -1930,7 +1930,7 @@ Console.WriteLine("target receipt executable path: {0}", InExecutablePath);
 		string LibimobileDeviceArguments =  "-i " + "\"" +  Path.GetFullPath(AppToDeploy) + "\"";
 		if (Params.DeviceNames[0].Length > 1)
 		{
-			LibimobileDeviceArguments = " -u " + Params.DeviceNames[0];
+			LibimobileDeviceArguments += " -u " + Params.DeviceNames[0];
 		}
 		LibimobileDeviceArguments = GetLibimobileDeviceNetworkedArgument(LibimobileDeviceArguments, Params.DeviceNames[0]);
 
