Reconciling CommandUtils differences

Adding a new class and class loader to assist with code signing resiliance when code signing is now backed by a service for Windows

#rnx

diff --git a/Engine/Source/Programs/AutomationTool/AutomationUtils/CommandUtils.cs b/Engine/Source/Programs/AutomationTool/AutomationUtils/CommandUtils.cs
--- a/Engine/Source/Programs/AutomationTool/AutomationUtils/CommandUtils.cs
+++ b/Engine/Source/Programs/AutomationTool/AutomationUtils/CommandUtils.cs
@@ -10,6 +10,7 @@ using System.Diagnostics;
 using System.IO;
 using System.IO.Compression;
 using System.Linq;
+using System.Reflection;
 using System.Runtime.Serialization;
 using System.Runtime.Versioning;
 using System.Text;
@@ -3064,7 +3065,7 @@ namespace AutomationTool
 		[SupportedOSPlatform("windows")]
 		public static void Sign(List<FileReference> Files, SignatureType SignatureType, bool AllowMultipleSignatures = true, string Description = null, bool RunInParallel = false)
 		{
-			const int MaxAttempts = 6;
+			const int MaxAttempts = 10;
 			string SignToolPath = GetSignToolPath();
 			string DescriptionArg = String.IsNullOrEmpty(Description) ? "" : $"/d \"{Description}\"";
 			string SpecificStoreArg = bUseMachineStoreForCertificates ? " /sm" : "";
@@ -3072,6 +3073,10 @@ namespace AutomationTool
 			string SHA1TimestampArg = " /tr";
 			string SHA256TimestampArg = " /tr";
 
+			ICodeSignHelper SigningHelper = CodeSignHelper.FindHelper(SignToolPath);
+
+			SigningHelper?.Validate(Logger);
+
 			List<Task> SignTaskList = new List<Task>();
 
 			for (int FileIdx = 0; FileIdx < Files.Count;)
@@ -3109,6 +3114,7 @@ namespace AutomationTool
 
 				Action SignToolAction = new Action(() =>
 				{
+					double TaskTimerLimit = 180.0;
 					Stopwatch TaskTimer = Stopwatch.StartNew();
 					for (int Attempts = 0; ; ++Attempts)
 					{
@@ -3139,7 +3145,17 @@ namespace AutomationTool
 							break;
 						}
 
-						if (TaskTimer.Elapsed.TotalMinutes > 3.0 && Attempts >= MaxAttempts)
+						if (SigningHelper != null && !SigningHelper.IsHealthy(Result, Logger))
+						{
+							Stopwatch FixTimer = Stopwatch.StartNew();
+
+							SigningHelper.Fix(Attempts, Logger);
+
+							//  Push out the max time.
+							TaskTimerLimit += FixTimer.Elapsed.TotalSeconds;
+						}
+
+						if (TaskTimer.Elapsed.TotalSeconds > TaskTimerLimit && Attempts > MaxAttempts)
 						{
 							throw new AutomationException($"Failed to sign files {Attempts} times over a period of {TaskTimer.Elapsed}");
 						}
@@ -3192,6 +3208,68 @@ namespace AutomationTool
 		}
 	}
 
+	public class CodeSignHelper
+	{
+		private static ICodeSignHelper HelperClass = null;
+
+		public static ICodeSignHelper FindHelper(string SignToolPath)
+		{
+			if (HelperClass == null)
+			{
+				// Find and load the first helper class
+				Assembly[] LoadedAssemblies = AppDomain.CurrentDomain.GetAssemblies();
+				foreach (var Dll in LoadedAssemblies)
+				{
+					Type[] AllTypes = Dll.SafeGetLoadedTypes();
+					foreach (var PotentialConfigType in AllTypes.Where(x => !x.IsAbstract && x != typeof(ICodeSignHelper) && typeof(ICodeSignHelper).IsAssignableFrom(x)))
+					{
+						try
+						{
+							ICodeSignHelper NewHelper = Activator.CreateInstance(PotentialConfigType) as ICodeSignHelper;
+							if (NewHelper != null)
+							{
+								NewHelper.Initialize(SignToolPath);
+								HelperClass = NewHelper;
+								break;
+							}
+						}
+						catch (Exception Ex)
+						{
+							var Inner = Ex.InnerException;
+							while (null != Inner)
+							{
+								Logger.LogWarning("Exception encountered creating HelperClass [{Name}] with error: {Message}", PotentialConfigType.Name, Inner.Message);
+								Inner = Inner.InnerException;
+							}
+							Logger.LogWarning("Unable to create HelperClass [{Name}] with error: {Message} \n {Trace}", PotentialConfigType.Name, Ex.Message, Ex.StackTrace);
+						}
+
+						if (HelperClass != null)
+						{
+							break;
+						}
+					}
+				}
+			}
+
+			return HelperClass;
+		}
+	}
+
+	/// <summary>
+	/// Interface helper class handle signing as a service environments
+	/// 
+	/// Implement a custom class for any provides
+	/// </summary>
+	public interface ICodeSignHelper
+	{
+		public void Initialize(string SignToolPath);
+		public bool Validate(ILogger Logger);
+		public bool IsHealthy(IProcessResult Result, ILogger Logger);
+		public bool IsFixing();
+		public bool Fix(int SignAttempt, ILogger Logger);
+	}
+
 	/// <summary>
 	/// Code signing
 	/// </summary>
