Removed take of screenshot on each DoWhile/Util function end in Gauntlet.RpcLibrary.cs

#rnx

diff --git a/Engine/Source/Programs/AutomationTool/Gauntlet/Framework/RpcFramework/Gauntlet.RpcLibrary.cs b/Engine/Source/Programs/AutomationTool/Gauntlet/Framework/RpcFramework/Gauntlet.RpcLibrary.cs
--- a/Engine/Source/Programs/AutomationTool/Gauntlet/Framework/RpcFramework/Gauntlet.RpcLibrary.cs
+++ b/Engine/Source/Programs/AutomationTool/Gauntlet/Framework/RpcFramework/Gauntlet.RpcLibrary.cs
@@ -210,7 +210,6 @@ namespace Gauntlet
 					// Do while RpcCall() returns true
 					if (!RpcCall())
 					{
-						OnBeforeTimeout?.Invoke();
 						return true; // Return true to caller to signal we are done
 					}
 				}
