Move metallib to paks
Memmap metallibs to not run oom
Removed most of the non-UFS code that was used only for metallibs
Removed bundle install logic that was for metallibs/IOS as it is no longer needed

diff --git a/Engine/Config/BaseGame.ini b/Engine/Config/BaseGame.ini
--- a/Engine/Config/BaseGame.ini
+++ b/Engine/Config/BaseGame.ini
@@ -134,6 +134,7 @@ WriteBackMetadataToAssetRegistry=Disabled
 +EarlyDownloaderPakFileFiles=...\*.upluginmanifest
 +EarlyDownloaderPakFileFiles=...\*.uproject
 +EarlyDownloaderPakFileFiles=...\global_sf*.metalmap
++EarlyDownloaderPakFileFiles=...\global_sf*.metallib
 
 +IniKeyDenylist=KeyStorePassword
 +IniKeyDenylist=KeyPassword
diff --git a/Engine/Config/IOS/BaseIOSInstallBundle.ini b/Engine/Config/IOS/BaseIOSInstallBundle.ini
--- a/Engine/Config/IOS/BaseIOSInstallBundle.ini
+++ b/Engine/Config/IOS/BaseIOSInstallBundle.ini
@@ -8,5 +8,4 @@ bOnlyGatherPaksInBulkData=true
 
 [InstallBundleManager.Shaderlib]
 ; List of extensions that will get added to NonUFSShaderLibPaths
-!ExtensionList=ClearArray
-ExtensionList=metallib
\ No newline at end of file
+!ExtensionList=ClearArray
\ No newline at end of file
diff --git a/Engine/Config/IOS/IOSGame.ini b/Engine/Config/IOS/IOSGame.ini
--- a/Engine/Config/IOS/IOSGame.ini
+++ b/Engine/Config/IOS/IOSGame.ini
@@ -1 +1,7 @@
 [Staging]
+
+[/Script/UnrealEd.ProjectPackagingSettings]
++FileExtensionsToSkipCompressionInPak=".metalmap"
++FileExtensionsToSkipCompressionInPak=".metallib"
++FileExtensionsToSkipEncryptionInPak=".metalmap"
++FileExtensionsToSkipEncryptionInPak=".metallib"
\ No newline at end of file
diff --git a/Engine/Config/Mac/MacGame.ini b/Engine/Config/Mac/MacGame.ini
--- a/Engine/Config/Mac/MacGame.ini
+++ b/Engine/Config/Mac/MacGame.ini
@@ -1 +1,7 @@
 [Staging]
+
+[/Script/UnrealEd.ProjectPackagingSettings]
++FileExtensionsToSkipCompressionInPak=".metalmap"
++FileExtensionsToSkipCompressionInPak=".metallib"
++FileExtensionsToSkipEncryptionInPak=".metalmap"
++FileExtensionsToSkipEncryptionInPak=".metallib"
\ No newline at end of file
diff --git a/Engine/Platforms/VisionOS/Source/Programs/AutomationTool/VisionOSPlatform.Automation.cs b/Engine/Platforms/VisionOS/Source/Programs/AutomationTool/VisionOSPlatform.Automation.cs
--- a/Engine/Platforms/VisionOS/Source/Programs/AutomationTool/VisionOSPlatform.Automation.cs
+++ b/Engine/Platforms/VisionOS/Source/Programs/AutomationTool/VisionOSPlatform.Automation.cs
@@ -142,7 +142,6 @@ public class VisionOSPlatform : IOSPlatform
         }
 
         {
-            // Stage any *.metallib files as NonUFS.
             // Get the final output directory for cooked data
             DirectoryReference CookOutputDir;
             if (!String.IsNullOrEmpty(Params.CookOutputDir))
@@ -157,14 +156,6 @@ public class VisionOSPlatform : IOSPlatform
             {
                 CookOutputDir = DirectoryReference.Combine(SC.ProjectRoot, "Saved", "Cooked", SC.CookPlatform);
             }
-            if (DirectoryReference.Exists(CookOutputDir))
-            {
-                List<FileReference> CookedFiles = DirectoryReference.EnumerateFiles(CookOutputDir, "*.metallib", SearchOption.AllDirectories).ToList();
-                foreach (FileReference CookedFile in CookedFiles)
-                {
-                    SC.StageFile(StagedFileType.NonUFS, CookedFile, new StagedFileReference(CookedFile.MakeRelativeTo(CookOutputDir)));
-                }
-            }
         }
     }
 }
diff --git a/Engine/Plugins/Experimental/DefaultInstallBundleManager/Source/Private/InstallBundleSourceBulk.cpp b/Engine/Plugins/Experimental/DefaultInstallBundleManager/Source/Private/InstallBundleSourceBulk.cpp
--- a/Engine/Plugins/Experimental/DefaultInstallBundleManager/Source/Private/InstallBundleSourceBulk.cpp
+++ b/Engine/Plugins/Experimental/DefaultInstallBundleManager/Source/Private/InstallBundleSourceBulk.cpp
@@ -274,12 +274,6 @@ void FInstallBundleSourceBulk::AsyncInit_MakeBundlesForBulkBuild()
 		{
 			FPlatformFileManager::Get().GetPlatformPhysical().FindFilesRecursively(*FoundFiles, *SearchDir, FileExtension);
 		}
-
-#if PLATFORM_IOS
-		// Only scan the root folder on IOS for shaderlibs.  Running this on windows is very expensive
-		// if the content dir contains loose assets which is common during development.
-		FPlatformFileManager::Get().GetPlatformPhysical().FindFiles(*FoundFiles, *ContentDir, TEXT(".metallib"));
-#endif // PLATFORM_IOS
 	},
 	[this, FoundFiles]()
 	{
diff --git a/Engine/Source/Programs/AutomationTool/IOS/IOSPlatform.Automation.cs b/Engine/Source/Programs/AutomationTool/IOS/IOSPlatform.Automation.cs
--- a/Engine/Source/Programs/AutomationTool/IOS/IOSPlatform.Automation.cs
+++ b/Engine/Source/Programs/AutomationTool/IOS/IOSPlatform.Automation.cs
@@ -1419,7 +1419,6 @@ Console.WriteLine("target receipt executable path: {0}", InExecutablePath);
 		}
 
 		{
-			// Stage any *.metallib files as NonUFS.
 			// Get the final output directory for cooked data
 			DirectoryReference CookOutputDir;
 			if (!String.IsNullOrEmpty(Params.CookOutputDir))
@@ -1434,14 +1433,6 @@ Console.WriteLine("target receipt executable path: {0}", InExecutablePath);
 			{
 				CookOutputDir = DirectoryReference.Combine(SC.ProjectRoot, "Saved", "Cooked", SC.CookPlatform);
 			}
-			if (DirectoryReference.Exists(CookOutputDir))
-			{
-				List<FileReference> CookedFiles = DirectoryReference.EnumerateFiles(CookOutputDir, "*.metallib", SearchOption.AllDirectories).ToList();
-				foreach (FileReference CookedFile in CookedFiles)
-				{
-					SC.StageFile(StagedFileType.NonUFS, CookedFile, new StagedFileReference(CookedFile.MakeRelativeTo(CookOutputDir)));
-				}
-			}
 		}
 
 		{
diff --git a/Engine/Source/Programs/AutomationTool/Mac/MacPlatform.Automation.cs b/Engine/Source/Programs/AutomationTool/Mac/MacPlatform.Automation.cs
--- a/Engine/Source/Programs/AutomationTool/Mac/MacPlatform.Automation.cs
+++ b/Engine/Source/Programs/AutomationTool/Mac/MacPlatform.Automation.cs
@@ -260,7 +260,6 @@ public class MacPlatform : ApplePlatform
 		}
 
 		{
-			// Stage any *.metallib files as NonUFS.
 			// Get the final output directory for cooked data
 			DirectoryReference CookOutputDir;
 			if(!String.IsNullOrEmpty(Params.CookOutputDir))
@@ -275,14 +274,6 @@ public class MacPlatform : ApplePlatform
 			{
 				CookOutputDir = DirectoryReference.Combine(SC.ProjectRoot, "Saved", "Cooked", SC.CookPlatform);
 			}
-			if(DirectoryReference.Exists(CookOutputDir))
-			{
-				List<FileReference> CookedFiles = DirectoryReference.EnumerateFiles(CookOutputDir, "*.metallib", SearchOption.AllDirectories).ToList();
-				foreach(FileReference CookedFile in CookedFiles)
-				{
-					SC.StageFile(StagedFileType.NonUFS, CookedFile, new StagedFileReference(CookedFile.MakeRelativeTo(CookOutputDir)));
-				}
-			}
 		}
 	}
 
diff --git a/Engine/Source/Programs/AutomationTool/Scripts/CopyBuildToStagingDirectory.Automation.cs b/Engine/Source/Programs/AutomationTool/Scripts/CopyBuildToStagingDirectory.Automation.cs
--- a/Engine/Source/Programs/AutomationTool/Scripts/CopyBuildToStagingDirectory.Automation.cs
+++ b/Engine/Source/Programs/AutomationTool/Scripts/CopyBuildToStagingDirectory.Automation.cs
@@ -220,6 +220,23 @@ namespace AutomationScripts
 			return null;
 		}
 
+		/// <summary>
+		/// returns a hash set containing file extensions that will not be encrypted within pak files
+		/// or null if no settings could be found.
+		/// </summary>
+		private static HashSet<string> GetFileExtensionsToSkipEncryption(ConfigHierarchy PlatformGameConfig)
+		{
+			if (PlatformGameConfig != null)
+			{
+				List<string> FileExtensionsToSkipEncryptionList = null;
+				PlatformGameConfig.GetArray("/Script/UnrealEd.ProjectPackagingSettings", "FileExtensionsToSkipEncryptionInPak", out FileExtensionsToSkipEncryptionList);
+				HashSet<string> FileExtensionsToSkipEncryption = (FileExtensionsToSkipEncryptionList is not null) ? new HashSet<string>(FileExtensionsToSkipEncryptionList) : null;
+				return FileExtensionsToSkipEncryption;
+			}
+
+			return null;
+		}
+
 		/// <summary>
 		/// Writes a pak response file to disk
 		/// </summary>
@@ -229,7 +246,7 @@ namespace AutomationScripts
 		/// <param name="RehydrateAssets"></param>
 		/// <param name="CryptoSettings"></param>
 		/// <param name="bForceFullEncryption"></param>
-		private static void WritePakResponseFile(string Filename, Dictionary<string, string> ResponseFile, bool Compressed, bool RehydrateAssets, EncryptionAndSigning.CryptoSettings CryptoSettings, bool bForceFullEncryption, HashSet<string> FileExtensionsToSkipCompression)
+		private static void WritePakResponseFile(string Filename, Dictionary<string, string> ResponseFile, bool Compressed, bool RehydrateAssets, EncryptionAndSigning.CryptoSettings CryptoSettings, bool bForceFullEncryption, HashSet<string> FileExtensionsToSkipCompression, HashSet<string> FileExtensionsToSkipEncryption)
 		{
 			using (var Writer = new StreamWriter(Filename, false, new System.Text.UTF8Encoding(true)))
 			{
@@ -253,7 +270,8 @@ namespace AutomationScripts
 						Line += " -rehydrate";
 					}
 
-					if (CryptoSettings != null)
+					bool bSkipEncryption = FileExtensionsToSkipEncryption != null && FileExtensionsToSkipEncryption.Contains(Path.GetExtension(Entry.Key));
+					if (CryptoSettings != null && !bSkipEncryption)
 					{
 						bool bEncryptFile = bForceFullEncryption || CryptoSettings.bEnablePakFullAssetEncryption;
 						bEncryptFile = bEncryptFile || (CryptoSettings.bEnablePakUAssetEncryption && Extension.Contains(".uasset"));
@@ -321,7 +339,7 @@ namespace AutomationScripts
 			return CmdLine.ToString();
 		}
 
-		static private string GetPakFileSpecificUnrealPakArguments(Dictionary<string, string> UnrealPakResponseFile, FileReference OutputLocation, string AdditionalOptions, bool Compressed, bool RehydrateAssets,  EncryptionAndSigning.CryptoSettings CryptoSettings, String PatchSourceContentPath, string EncryptionKeyGuid, HashSet<string> FileExtensionsToSkipCompression)
+		static private string GetPakFileSpecificUnrealPakArguments(Dictionary<string, string> UnrealPakResponseFile, FileReference OutputLocation, string AdditionalOptions, bool Compressed, bool RehydrateAssets,  EncryptionAndSigning.CryptoSettings CryptoSettings, String PatchSourceContentPath, string EncryptionKeyGuid, HashSet<string> FileExtensionsToSkipCompression, HashSet<string> FileExtensionsToSkipEncryption)
 		{
 			StringBuilder CmdLine = new StringBuilder(MakePathSafeToUseWithCommandLine(OutputLocation.FullName));
 
@@ -332,7 +350,7 @@ namespace AutomationScripts
 			string ResponseFilesPath = CombinePaths(CmdEnv.EngineSavedFolder, "ResponseFiles");
 			InternalUtils.SafeCreateDirectory(ResponseFilesPath);
 			string UnrealPakResponseFileName = CombinePaths(ResponseFilesPath, "PakList_" + PakName + ".txt");
-			WritePakResponseFile(UnrealPakResponseFileName, UnrealPakResponseFile, Compressed, RehydrateAssets, CryptoSettings, bForceEncryption, FileExtensionsToSkipCompression);
+			WritePakResponseFile(UnrealPakResponseFileName, UnrealPakResponseFile, Compressed, RehydrateAssets, CryptoSettings, bForceEncryption, FileExtensionsToSkipCompression, FileExtensionsToSkipEncryption);
 			CmdLine.AppendFormat(" -create={0}", CommandUtils.MakePathSafeToUseWithCommandLine(UnrealPakResponseFileName));
 
 			if (!String.IsNullOrEmpty(PatchSourceContentPath))
@@ -364,9 +382,10 @@ namespace AutomationScripts
 		{
 			bool RehydrateAssets = false;
 			HashSet<string> FileExtensionsToSkipCompression = null;
+			HashSet<string> FileExtensionsToSkipEncryption = null;
 
 			StringBuilder CmdLine = new StringBuilder();
-			CmdLine.Append(GetPakFileSpecificUnrealPakArguments(UnrealPakResponseFile, OutputLocation, "", Compressed, RehydrateAssets, CryptoSettings, PatchSourceContentPath, EncryptionKeyGuid, FileExtensionsToSkipCompression));
+			CmdLine.Append(GetPakFileSpecificUnrealPakArguments(UnrealPakResponseFile, OutputLocation, "", Compressed, RehydrateAssets, CryptoSettings, PatchSourceContentPath, EncryptionKeyGuid, FileExtensionsToSkipCompression, FileExtensionsToSkipEncryption));
 			CmdLine.Append(" ");
 			CmdLine.Append(GetCommonUnrealPakArguments(PakOrderFileLocations, AdditionalOptions, CryptoSettings, CryptoKeysCacheFilename, SecondaryPakOrderFileLocations, bUnattended));
 			return CmdLine.ToString();
@@ -385,7 +404,8 @@ namespace AutomationScripts
 			string PatchSourceContentPath,
 			bool bGenerateDiffPatch,
 			bool bIsDLC,
-			HashSet<string> FileExtensionsToSkipCompression)
+			HashSet<string> FileExtensionsToSkipCompression,
+			HashSet<string> FileExtensionsToSkipEncryption)
 		{
 			StringBuilder CmdLine = new StringBuilder();
 			CmdLine.AppendFormat("-Output={0}", MakePathSafeToUseWithCommandLine(Path.ChangeExtension(PakOutputLocation.FullName, ".utoc")));
@@ -416,7 +436,7 @@ namespace AutomationScripts
 			string ResponseFilesPath = CombinePaths(CmdEnv.EngineSavedFolder, "ResponseFiles");
 			InternalUtils.SafeCreateDirectory(ResponseFilesPath);
 			string UnrealPakResponseFileName = CombinePaths(ResponseFilesPath, "PakListIoStore_" + ContainerName + ".txt");
-			WritePakResponseFile(UnrealPakResponseFileName, UnrealPakResponseFile, bCompressed, RehydrateAssets, CryptoSettings, bForceEncryption, FileExtensionsToSkipCompression);
+			WritePakResponseFile(UnrealPakResponseFileName, UnrealPakResponseFile, bCompressed, RehydrateAssets, CryptoSettings, bForceEncryption, FileExtensionsToSkipCompression, FileExtensionsToSkipEncryption);
 			CmdLine.AppendFormat(" -ResponseFile={0}", CommandUtils.MakePathSafeToUseWithCommandLine(UnrealPakResponseFileName));
 
 			if (CryptoSettings != null && CryptoSettings.bDataCryptoRequired)
@@ -1086,12 +1106,6 @@ namespace AutomationScripts
 				return null;
 			}
 
-			// metallib files cannot *currently* be staged as UFS as the Metal API needs to mmap them from files on disk in order to function efficiently
-			if (File.HasExtension(".metallib"))
-			{
-				return StagedFileType.NonUFS;
-			}
-
 			// Cannot stage
 			if (File.HasExtension(".utoc") || File.HasExtension(".ucas") || File.HasExtension(".uondemandtoc"))
 			{
@@ -3322,6 +3336,7 @@ namespace AutomationScripts
 			PlatformGameConfig.GetBool("/Script/UnrealEd.ProjectPackagingSettings", "bForceUseProjectCompressionFormatIgnoreHardwareOverride", out bForceUseProjectCompressionFormatIgnoreHardwareOverride);
 
 			HashSet<string> FileExtensionsToSkipCompressionInPaks = GetFileExtensionsToSkipCompression(PlatformGameConfig);
+			HashSet<string> FileExtensionsToSkipEncryptionInPaks = GetFileExtensionsToSkipEncryption(PlatformGameConfig);
 
 			string HardwareCompressionFormat = null;
 			if ( ! bForceUseProjectCompressionFormatIgnoreHardwareOverride )
@@ -3981,7 +3996,8 @@ namespace AutomationScripts
 								ContainerPatchSourcePath,
 								bGenerateDiffPatch,
 								Params.HasDLCName,
-								FileExtensionsToSkipCompressionInPaks));
+								FileExtensionsToSkipCompressionInPaks,
+								FileExtensionsToSkipEncryptionInPaks));
 						}
 
 						if (!PakParams.bStageLoose)
@@ -4001,7 +4017,8 @@ namespace AutomationScripts
 								Params.SkipEncryption ? null : CryptoSettings,
 								PatchSourceContentPath,
 								Params.SkipEncryption ? "" : PakParams.EncryptionKeyGuid,
-								FileExtensionsToSkipCompressionInPaks));
+								FileExtensionsToSkipCompressionInPaks,
+								FileExtensionsToSkipEncryptionInPaks));
 						}
 						LogNames.Add(OutputLocation.GetFileNameWithoutExtension());
 					}
diff --git a/Engine/Source/Programs/AutomationTool/TVOS/TVOSPlatform.Automation.cs b/Engine/Source/Programs/AutomationTool/TVOS/TVOSPlatform.Automation.cs
--- a/Engine/Source/Programs/AutomationTool/TVOS/TVOSPlatform.Automation.cs
+++ b/Engine/Source/Programs/AutomationTool/TVOS/TVOSPlatform.Automation.cs
@@ -159,7 +159,6 @@ public class TVOSPlatform : IOSPlatform
         }
 
         {
-            // Stage any *.metallib files as NonUFS.
             // Get the final output directory for cooked data
             DirectoryReference CookOutputDir;
             if (!String.IsNullOrEmpty(Params.CookOutputDir))
@@ -174,14 +173,6 @@ public class TVOSPlatform : IOSPlatform
             {
                 CookOutputDir = DirectoryReference.Combine(SC.ProjectRoot, "Saved", "Cooked", SC.CookPlatform);
             }
-            if (DirectoryReference.Exists(CookOutputDir))
-            {
-                List<FileReference> CookedFiles = DirectoryReference.EnumerateFiles(CookOutputDir, "*.metallib", SearchOption.AllDirectories).ToList();
-                foreach (FileReference CookedFile in CookedFiles)
-                {
-                    SC.StageFile(StagedFileType.NonUFS, CookedFile, new StagedFileReference(CookedFile.MakeRelativeTo(CookOutputDir)));
-                }
-            }
         }
     }
 }
diff --git a/Engine/Source/Runtime/Apple/MetalRHI/Private/MetalDynamicRHI_Shaders.cpp b/Engine/Source/Runtime/Apple/MetalRHI/Private/MetalDynamicRHI_Shaders.cpp
--- a/Engine/Source/Runtime/Apple/MetalRHI/Private/MetalDynamicRHI_Shaders.cpp
+++ b/Engine/Source/Runtime/Apple/MetalRHI/Private/MetalDynamicRHI_Shaders.cpp
@@ -151,24 +151,29 @@ FRHIShaderLibraryRef FMetalDynamicRHI::RHICreateShaderLibrary(EShaderPlatform Pl
 #endif
 	};
 
+	auto MapShaderLibraryFile = [](const FString& FilePath, TUniquePtr<FMetalShaderLibrary::FShaderLibDataOwner>& MemOwner)
+	{
+		if (FPlatformProperties::SupportsMemoryMappedFiles())
+		{
+			FOpenMappedResult Res = FPlatformFileManager::Get().GetPlatformFile().OpenMappedEx(*FilePath);
+			MemOwner->MappedCacheFile = Res.HasError() ? nullptr : Res.StealValue();
+			if (MemOwner->MappedCacheFile.IsValid())
+			{
+				MemOwner->MappedRegion = TUniquePtr<IMappedFileRegion>(MemOwner->MappedCacheFile->MapRegion(0, MemOwner->MappedCacheFile->GetFileSize()));
+			}
+		}
+	};
+
  #if !USE_MMAPPED_SHADERARCHIVE
     FArchive* BinaryShaderAr = IFileManager::Get().CreateFileReader(*BinaryShaderFile);
 #else
 	FStaticMemoryReader* BinaryShaderAr = nullptr;
     TUniquePtr<FMetalShaderLibrary::FShaderLibDataOwner> MemOwner = MakeUnique<FMetalShaderLibrary::FShaderLibDataOwner>();
-	if (FPlatformProperties::SupportsMemoryMappedFiles())
+	MapShaderLibraryFile(BinaryShaderFile, MemOwner);
+	if (MemOwner->MappedRegion.IsValid())
 	{
-		FOpenMappedResult Res = FPlatformFileManager::Get().GetPlatformFile().OpenMappedEx(*BinaryShaderFile);
-		MemOwner->MappedCacheFile = Res.HasError() ? nullptr : Res.StealValue();
-		if (MemOwner->MappedCacheFile.IsValid())
-		{
-			MemOwner->MappedRegion = TUniquePtr<IMappedFileRegion>(MemOwner->MappedCacheFile->MapRegion(0, MemOwner->MappedCacheFile->GetFileSize()));
-			if(MemOwner->MappedRegion.IsValid())
-			{
-				UE_LOG(LogMetal, Display, TEXT("mmapping %s, %d bytes"), *BinaryShaderFile, MemOwner->MappedCacheFile->GetFileSize());
-				BinaryShaderAr = new FStaticMemoryReader(MemOwner->MappedRegion->GetMappedPtr(), MemOwner->MappedCacheFile->GetFileSize());
-			}
-		}
+		UE_LOG(LogMetal, Display, TEXT("mmapping %s, %d bytes"), *BinaryShaderFile, MemOwner->MappedCacheFile->GetFileSize());
+		BinaryShaderAr = new FStaticMemoryReader(MemOwner->MappedRegion->GetMappedPtr(), MemOwner->MappedCacheFile->GetFileSize());
 	}
 
 	if(BinaryShaderAr == nullptr)
@@ -198,57 +203,51 @@ FRHIShaderLibraryRef FMetalDynamicRHI::RHICreateShaderLibrary(EShaderPlatform Pl
         {
             check(((SerializedShaders.GetNumShaders() + Header.NumShadersPerLibrary - 1) / Header.NumShadersPerLibrary) == Header.NumLibraries);
 
-            TArray<MTLLibraryPtr> Libraries;
-            Libraries.Empty(Header.NumLibraries);
+			TArray<TUniquePtr<FMetalShaderLibrary::FShaderLibDataOwner>> LibraryMemOwner;
+
+			LibraryMemOwner.Empty(Header.NumLibraries);
 
             for (uint32 i = 0; i < Header.NumLibraries; i++)
             {
-                FString MetalLibraryFilePath = (FilePath / LibName) + FString::Printf(TEXT(".%d.metallib"), i);
-                FString MetalLibraryAbsoluteFilePath = IFileManager::Get().ConvertToAbsolutePathForExternalAppForRead(*MetalLibraryFilePath);
-
-                METAL_GPUPROFILE(FScopedMetalCPUStats CPUStat(FString::Printf(TEXT("NewLibraryFile: %s"), *MetalLibraryAbsoluteFilePath)));
-                NS::Error* Error;
-                NS::String* MetalLibraryFilePathNSString = FStringToNSString(MetalLibraryAbsoluteFilePath);
-                NS::URL *metalLibraryURL = NS::URL::fileURLWithPath(MetalLibraryFilePathNSString);
-                MTLLibraryPtr Library = NS::TransferPtr(Device->GetDevice()->newLibrary(metalLibraryURL, &Error));
-
-                if (Library.get() == nullptr)
-                {
-                    // Metallib not found. Is Zen server being used?
-                    static const bool bRunningWithZenStore = FPlatformFileManager::Get().FindPlatformFile(TEXT("StorageServer")) != nullptr;
-                    if( bRunningWithZenStore )
-                    {
-                        TArray<uint8> LibraryData;
-                        if (FFileHelper::LoadFileToArray(LibraryData, *MetalLibraryFilePath))
-                        {
-                            dispatch_data_t data = dispatch_data_create(LibraryData.GetData(), LibraryData.Num(), nil, DISPATCH_DATA_DESTRUCTOR_DEFAULT);
-                            Library = NS::TransferPtr(Device->GetDevice()->newLibrary(data, &Error));
-                        }
-                        else
-                        {
-                            UE_LOG(LogMetal, Warning, TEXT("Metallib '%s' unable to be read from ZenStore."), *MetalLibraryFilePath);
-                            return nullptr;
-                        }
-                    }
-                    else
-                    {
-                        UE_LOG(LogMetal, Warning, TEXT("Metallib '%s' not found and ZenStore is not being used."), *MetalLibraryFilePath);
-                        return nullptr;
-                    }
-                }
-
-                if (Library.get() != nullptr)
-                {
-                    Libraries.Add(Library);
-                }
-                else
-                {
-                    UE_LOG(LogMetal, Display, TEXT("Failed to create library: %s"), *NSStringToFString(Error->description()));
-                    return nullptr;
-                }
+				FString MetalLibraryFilePath = (FilePath / LibName) + FString::Printf(TEXT(".%d.metallib"), i);
+				NS::Error* Error;
+
+				MTLLibraryPtr Library;
+				TUniquePtr<FMetalShaderLibrary::FShaderLibDataOwner> MemOwnerLibrary = MakeUnique<FMetalShaderLibrary::FShaderLibDataOwner>();
+				MapShaderLibraryFile(MetalLibraryFilePath, MemOwnerLibrary);
+				if (MemOwnerLibrary->MappedRegion.IsValid())
+				{
+					UE_LOG(LogMetal, Display, TEXT("mmapping %s, %d bytes"), *MetalLibraryFilePath, MemOwnerLibrary->MappedCacheFile->GetFileSize());
+
+					dispatch_data_t data = dispatch_data_create(MemOwnerLibrary->MappedRegion->GetMappedPtr(), MemOwnerLibrary->MappedRegion->GetMappedSize(), nil, DISPATCH_DATA_DESTRUCTOR_MUNMAP);
+					MemOwnerLibrary->Library = NS::TransferPtr(Device->GetDevice()->newLibrary(data, &Error));
+					dispatch_release(data);
+				}
+
+				if (MemOwnerLibrary->Library.get() == nullptr)
+				{
+					TArray<uint8>& FileData = MemOwnerLibrary->Mem;
+					if (FFileHelper::LoadFileToArray(FileData, *MetalLibraryFilePath))
+					{
+						UE_LOG(LogMetal, Display, TEXT("emulating mmapping %s, %d bytes!"), *MetalLibraryFilePath, FileData.Num());
+						dispatch_data_t data = dispatch_data_create(FileData.GetData(), FileData.Num(), nil, DISPATCH_DATA_DESTRUCTOR_MUNMAP);
+						MemOwnerLibrary->Library = NS::TransferPtr(Device->GetDevice()->newLibrary(data, &Error));
+						dispatch_release(data);
+					}
+				}
+
+				if ((MemOwnerLibrary != nullptr) && (MemOwnerLibrary->Library.get() != nullptr))
+				{
+					LibraryMemOwner.Add(MoveTemp(MemOwnerLibrary));
+				}
+				else
+				{
+					UE_LOG(LogMetal, Display, TEXT("Failed to create library: %s"), *MetalLibraryFilePath);
+					return nullptr;
+				}
             }
 
-            FMetalShaderLibrary* MtlLib = new FMetalShaderLibrary(*Device, Platform, Name, BinaryShaderFile, Header, MoveTemp(SerializedShaders), MoveTemp(ShaderCode), Libraries
+            FMetalShaderLibrary* MtlLib = new FMetalShaderLibrary(*Device, Platform, Name, BinaryShaderFile, Header, MoveTemp(SerializedShaders), MoveTemp(ShaderCode), MoveTemp(LibraryMemOwner)
 #if USE_MMAPPED_SHADERARCHIVE
 				,MoveTemp(MemOwner)
 #endif
diff --git a/Engine/Source/Runtime/Apple/MetalRHI/Private/Shaders/MetalShaderLibrary.cpp b/Engine/Source/Runtime/Apple/MetalRHI/Private/Shaders/MetalShaderLibrary.cpp
--- a/Engine/Source/Runtime/Apple/MetalRHI/Private/Shaders/MetalShaderLibrary.cpp
+++ b/Engine/Source/Runtime/Apple/MetalRHI/Private/Shaders/MetalShaderLibrary.cpp
@@ -52,15 +52,15 @@ FMetalShaderLibrary::FMetalShaderLibrary(FMetalDevice& MetalDevice,
 										 const FString& InShaderLibraryFilename,
 										 const FMetalShaderLibraryHeader& InHeader,
 										 FSerializedShaderArchive&& InSerializedShaders,
-                                         FShaderCodeArrayType&& InShaderCode,
-										 const TArray<MTLLibraryPtr>& InLibrary)
+										 FShaderCodeArrayType&& InShaderCode,
+										 TArray<TUniquePtr<FShaderLibDataOwner>>&& InLibraryMemOwner)
 	: FRHIShaderLibrary(Platform, Name)
 	, Device(MetalDevice)
 	, ShaderLibraryFilename(InShaderLibraryFilename)
-	, Library(InLibrary)
 	, Header(InHeader)
 	, SerializedShaders(MoveTemp(InSerializedShaders))
 	, ShaderCode(MoveTemp(InShaderCode))
+	, LibraryMemOwner(MoveTemp(InLibraryMemOwner))
 {
 #if !UE_BUILD_SHIPPING
 	DebugFile = nullptr;
@@ -142,16 +142,16 @@ TRefCountPtr<FRHIShader> FMetalShaderLibrary::CreateShader(int32 Index, bool bRe
 	switch (ShaderEntry.Frequency)
 	{
 		case SF_Vertex:
-			Shader = CreateMetalShader<FMetalVertexShader>(Device, Code, Library[LibraryIndex]);
+			Shader = CreateMetalShader<FMetalVertexShader>(Device, Code, LibraryMemOwner[LibraryIndex]->Library);
 			break;
 
 		case SF_Pixel:
-			Shader = CreateMetalShader<FMetalPixelShader>(Device, Code, Library[LibraryIndex]);
+			Shader = CreateMetalShader<FMetalPixelShader>(Device, Code, LibraryMemOwner[LibraryIndex]->Library);
 			break;
  
  		case SF_Geometry:
 #if PLATFORM_SUPPORTS_GEOMETRY_SHADERS
-            Shader = CreateMetalShader<FMetalGeometryShader>(Device, Code, Library[LibraryIndex]);
+            Shader = CreateMetalShader<FMetalGeometryShader>(Device, Code, LibraryMemOwner[LibraryIndex]->Library);
 #else
             checkf(false, TEXT("Geometry shaders not supported"));
 #endif
@@ -160,7 +160,7 @@ TRefCountPtr<FRHIShader> FMetalShaderLibrary::CreateShader(int32 Index, bool bRe
 
         case SF_Mesh:
 #if PLATFORM_SUPPORTS_MESH_SHADERS
-            Shader = CreateMetalShader<FMetalMeshShader>(Device, Code, Library[LibraryIndex]);
+            Shader = CreateMetalShader<FMetalMeshShader>(Device, Code, LibraryMemOwner[LibraryIndex]->Library);
 #else
 			checkf(false, TEXT("Mesh shaders not supported"));
 #endif
@@ -168,14 +168,14 @@ TRefCountPtr<FRHIShader> FMetalShaderLibrary::CreateShader(int32 Index, bool bRe
 
         case SF_Amplification:
 #if PLATFORM_SUPPORTS_MESH_SHADERS
-            Shader = CreateMetalShader<FMetalAmplificationShader>(Device, Code, Library[LibraryIndex]);
+            Shader = CreateMetalShader<FMetalAmplificationShader>(Device, Code, LibraryMemOwner[LibraryIndex]->Library);
 #else
 			checkf(false, TEXT("Amplification shaders not supported"));
 #endif
             break;
 
 		case SF_Compute:
-			Shader = CreateMetalShader<FMetalComputeShader>(Device, Code, Library[LibraryIndex]);
+			Shader = CreateMetalShader<FMetalComputeShader>(Device, Code, LibraryMemOwner[LibraryIndex]->Library);
 			break;
 
 		default:
diff --git a/Engine/Source/Runtime/Apple/MetalRHI/Private/Shaders/MetalShaderLibrary.h b/Engine/Source/Runtime/Apple/MetalRHI/Private/Shaders/MetalShaderLibrary.h
--- a/Engine/Source/Runtime/Apple/MetalRHI/Private/Shaders/MetalShaderLibrary.h
+++ b/Engine/Source/Runtime/Apple/MetalRHI/Private/Shaders/MetalShaderLibrary.h
@@ -18,16 +18,16 @@ class FMetalShaderLibrary final : public FRHIShaderLibrary
 {
 public:
 
+	struct FShaderLibDataOwner
+	{
+		TArray<uint8> Mem;
+		TUniquePtr<IMappedFileHandle> MappedCacheFile;
+		TUniquePtr<IMappedFileRegion> MappedRegion;
+		MTLLibraryPtr Library;
+	};
 #if !USE_MMAPPED_SHADERARCHIVE
     using FShaderCodeArrayType = TArray<uint8>;
-#else
-    struct FShaderLibDataOwner
-    {
-        TArray<uint8> Mem;
-        TUniquePtr<IMappedFileHandle> MappedCacheFile;
-        TUniquePtr<IMappedFileRegion> MappedRegion;
-    };
-    
+#else 
 	using FShaderCodeArrayType = TArrayView<const uint8>;
 private:
     TUniquePtr<FShaderLibDataOwner> MemOwner;
@@ -40,10 +40,10 @@ public:
         const FMetalShaderLibraryHeader& InHeader,
         FSerializedShaderArchive&& InSerializedShaders,
         FShaderCodeArrayType&& InShaderCode,
-        const TArray<MTLLibraryPtr>& InLibrary,
+		TArray<TUniquePtr<FShaderLibDataOwner>>&& InLibraryMemOwner,
         TUniquePtr<FShaderLibDataOwner>&& InMemOwner)
         :
-        FMetalShaderLibrary(Device, Platform, Name, InShaderLibraryFilename, InHeader, MoveTemp(InSerializedShaders), MoveTemp(InShaderCode), InLibrary)
+        FMetalShaderLibrary(Device, Platform, Name, InShaderLibraryFilename, InHeader, MoveTemp(InSerializedShaders), MoveTemp(InShaderCode), MoveTemp(InLibraryMemOwner))
         {
             MemOwner = MoveTemp(InMemOwner);
         }
@@ -59,7 +59,7 @@ public:
 						const FMetalShaderLibraryHeader& InHeader,
 						FSerializedShaderArchive&& InSerializedShaders,
 						FShaderCodeArrayType&& InShaderCode,
-						const TArray<MTLLibraryPtr>& InLibrary);
+						TArray<TUniquePtr<FShaderLibDataOwner>>&& InLibraryMemOwner);
 
 	virtual ~FMetalShaderLibrary();
 
@@ -86,10 +86,10 @@ public:
 private:
 	FMetalDevice& Device;
 	FString ShaderLibraryFilename;
-	TArray<MTLLibraryPtr> Library;
 	FMetalShaderLibraryHeader Header;
 	FSerializedShaderArchive SerializedShaders;
 	FShaderCodeArrayType ShaderCode;
+	TArray<TUniquePtr<FShaderLibDataOwner>> LibraryMemOwner;
 #if !UE_BUILD_SHIPPING
 	class FMetalShaderDebugZipFile* DebugFile;
 #endif
