[Gauntlet] Allow longer retry when reserving devices

#rnx

diff --git a/Engine/Source/Programs/AutomationTool/AutomationUtils/DeviceReservation.cs b/Engine/Source/Programs/AutomationTool/AutomationUtils/DeviceReservation.cs
--- a/Engine/Source/Programs/AutomationTool/AutomationUtils/DeviceReservation.cs
+++ b/Engine/Source/Programs/AutomationTool/AutomationUtils/DeviceReservation.cs
@@ -67,15 +67,7 @@ namespace AutomationTool.DeviceReservation
 			// This blocks until the reservation is successful.
 			ActiveReservation = Reservation.Create(ReservationBaseUri, InDeviceTypes, ReserveTime, RetryMax, PoolID);
 
-			// Resolve the device IPs
-			ReservedDevices = new List<Device>();
-			foreach (var DeviceName in ActiveReservation.DeviceNames)
-			{
-				ReservedDevices.Add(Device.Get(ReservationBaseUri, DeviceName));
-			}
-
-			RenewThread = new Thread(DoAutoRenew);
-			RenewThread.Start();
+			StartRenewThread();
 		}
 
 		/// <summary>
@@ -88,7 +80,11 @@ namespace AutomationTool.DeviceReservation
 		{
 			ReservationBaseUri = new Uri(InReservationBaseUri);
 			ActiveReservation = Reservation;
+			StartRenewThread();
+		}
 
+		private void StartRenewThread()
+		{
 			// Resolve the device IPs
 			ReservedDevices = new List<Device>();
 			foreach (string DeviceName in ActiveReservation.DeviceNames)
@@ -325,7 +321,8 @@ namespace AutomationTool.DeviceReservation
 					if (RetryCount == RetryMax)
 					{
 						Console.WriteLine("Device reservation unsuccessful");
-						string FinalMessage = "Device reservation unsuccessful";
+						string PoolName = string.IsNullOrEmpty(PoolID) ? "default" : PoolID;
+						string FinalMessage = $"Device reservation for {string.Join(", ", DeviceTypes)} from {PoolName} pool unsuccessful";
 						string MaxRetryWebExMessage = GetReservationErrorMessage(WebEx);
 						if (!string.IsNullOrEmpty(MaxRetryWebExMessage))
 						{
diff --git a/Engine/Source/Programs/AutomationTool/Gauntlet/Framework/Gauntlet.TestExecutor.cs b/Engine/Source/Programs/AutomationTool/Gauntlet/Framework/Gauntlet.TestExecutor.cs
--- a/Engine/Source/Programs/AutomationTool/Gauntlet/Framework/Gauntlet.TestExecutor.cs
+++ b/Engine/Source/Programs/AutomationTool/Gauntlet/Framework/Gauntlet.TestExecutor.cs
@@ -4,6 +4,7 @@ using System;
 using System.Collections.Generic;
 using System.IO;
 using System.Linq;
+using System.Net;
 using System.Threading;
 using EpicGames.Core;
 using Gauntlet.Utils;
@@ -242,11 +243,21 @@ namespace Gauntlet
 							}
 							catch (System.Exception ex)
 							{
-								Log.Error(KnownLogEvents.Gauntlet_TestEvent, "Test {Name} threw an exception during ready check. Ex: {Exception}", Node, ex);
-								Node.AddTestEvent(new UnrealTestEvent(EventSeverity.Error, "Test Failed to Start", new List<string> {ex.Message}));
-								PendingTests[i] = null;
-								NodeInfo.TimeSetupBegan = NodeInfo.TimeSetupEnded = NodeInfo.TimeTestEnded = DateTime.Now;
-								CompletedTests.Add(NodeInfo);
+								if (ex.InnerException != null && ex.InnerException is WebException
+									&& (ex.InnerException as WebException).Response != null
+									&& ((ex.InnerException as WebException).Response as HttpWebResponse).StatusCode == HttpStatusCode.Conflict)
+								{
+									Log.Info(ex.Message);
+									// In case of conflict server response, let the test execution loop handles it (resolved through timeout option).
+								}
+								else
+								{
+									Log.Error(KnownLogEvents.Gauntlet_TestEvent, "Test {Name} threw an exception during ready check. Ex: {Exception}", Node, ex);
+									Node.AddTestEvent(new UnrealTestEvent(EventSeverity.Error, "Test Failed to Start", new List<string> { ex.Message }));
+									PendingTests[i] = null;
+									NodeInfo.TimeSetupBegan = NodeInfo.TimeSetupEnded = NodeInfo.TimeTestEnded = DateTime.Now;
+									CompletedTests.Add(NodeInfo);
+								}
 							}
 
 							if (IsTestReady)
@@ -276,11 +287,12 @@ namespace Gauntlet
 									double TimeWaiting = (DateTime.Now - NodeInfo.FirstReadyCheckTime).TotalSeconds;
 									if (TimeWaiting >= Options.Wait)
 									{
-										Log.Error(KnownLogEvents.Gauntlet_DeviceEvent, "Test {TestName} has been waiting to run resource-free for {Time:00} seconds. Removing from wait list", Node, TimeWaiting);
+										string Pool = string.IsNullOrEmpty(Globals.DevicePoolId)? string.Empty : $" from {Globals.DevicePoolId} pool";
+										Log.Error(KnownLogEvents.Gauntlet_DeviceEvent, "Test {TestName} has been waiting to run resource-free for {Time:00} seconds{Pool}. Removing from wait list", Node, TimeWaiting, Pool);
 										Node.AddTestEvent(new UnrealTestEvent(EventSeverity.Error, "Insufficient devices found", new List<string> {string.Format("Test {0} was unable to find enough devices after trying for {1:00} seconds.", Node, TimeWaiting), "This is not a test-related failure."}));
 										PendingTests[i] = null;
 										NodeInfo.TimeSetupBegan = NodeInfo.TimeSetupEnded = NodeInfo.TimeTestEnded = DateTime.Now;
-										NodeInfo.Result = TestExecutionInfo.ExecutionResult.Failed;
+										NodeInfo.Result = TestExecutionInfo.ExecutionResult.NotStarted;
 										CompletedTests.Add(NodeInfo);
 									}
 								}
