- Added a check for missing DDPI, which can happen if you have a precompiled UBT but not the platform config files

diff --git a/Engine/Source/Programs/AutomationTool/Turnkey/TurnkeyManifest.cs b/Engine/Source/Programs/AutomationTool/Turnkey/TurnkeyManifest.cs
--- a/Engine/Source/Programs/AutomationTool/Turnkey/TurnkeyManifest.cs
+++ b/Engine/Source/Programs/AutomationTool/Turnkey/TurnkeyManifest.cs
@@ -170,7 +170,8 @@ namespace Turnkey
 				// also manually create local "FileSource" objects specified via code
 				foreach (UnrealTargetPlatform Platform in UnrealTargetPlatform.GetValidPlatforms())
 				{
-					if (DataDrivenPlatformInfo.GetDataDrivenInfoForPlatform(Platform.ToString()).bIsEnabled)
+					DataDrivenPlatformInfo.ConfigDataDrivenPlatformInfo Info = DataDrivenPlatformInfo.GetDataDrivenInfoForPlatform(Platform.ToString());
+					if (Info != null && Info.bIsEnabled)
 					{
 						// these are usually going to be empty
 						foreach (string SpecifiedVersion in AutomationTool.Platform.GetPlatform(Platform).GetCodeSpecifiedSdkVersions())
