LaunchMultiServer command

diff --git a/Engine/Source/Programs/AutomationTool/Scripts/LaunchMultiServer.Automation.cs b/Engine/Source/Programs/AutomationTool/Scripts/LaunchMultiServer.Automation.cs
--- a/Engine/Source/Programs/AutomationTool/Scripts/LaunchMultiServer.Automation.cs
+++ b/Engine/Source/Programs/AutomationTool/Scripts/LaunchMultiServer.Automation.cs
@@ -298,6 +298,7 @@ namespace AutomationTool
 					if (LaunchType == ProcessLaunchType.Server)
 					{
 						Args += String.Format("-MultiServerLocalId={0} -MultiServerListenPort={1} -port={2} ", ServerInfo.Definition.LocalId, ServerInfo.Definition.ListenPort, ServerInfo.GameListenPort);
+						Args += String.Format("-ConsoleTitle=\"{0} Server ID {1}\" ", Params.ShortProjectName, ServerInfo.Definition.LocalId);
 						Args += Params.ServerCommandline + " ";
 					}
 					else if (LaunchType == ProcessLaunchType.DirectClient)
