AutomationTool: Fix incorrect UEProjectRoot when debugging

#rnx

diff --git a/Engine/Source/Programs/AutomationTool/AutomationUtils/CommandEnvironment.cs b/Engine/Source/Programs/AutomationTool/AutomationUtils/CommandEnvironment.cs
--- a/Engine/Source/Programs/AutomationTool/AutomationUtils/CommandEnvironment.cs
+++ b/Engine/Source/Programs/AutomationTool/AutomationUtils/CommandEnvironment.cs
@@ -87,7 +87,7 @@ namespace AutomationTool
 			LocalRoot = CommandUtils.GetEnvVar(EnvVarNames.LocalRoot);
 			if(String.IsNullOrEmpty(LocalRoot))
 			{
-				LocalRoot = CommandUtils.ConvertSeparators(PathSeparator.Slash, Path.GetFullPath(Path.Combine(Path.GetDirectoryName(AutomationToolDll), "..", "..", "..", "..")));
+				LocalRoot = CommandUtils.ConvertSeparators(PathSeparator.Slash, Unreal.RootDirectory.FullName);
 				CommandUtils.ConditionallySetEnvVar(EnvVarNames.LocalRoot, LocalRoot);
 			}
 
