Include dump_syms tool in Mac editor build
Fix GeneratePsyms command where input files have spaces in path

diff --git a/Engine/Build/InstalledEngineFilters.xml b/Engine/Build/InstalledEngineFilters.xml
--- a/Engine/Build/InstalledEngineFilters.xml
+++ b/Engine/Build/InstalledEngineFilters.xml
@@ -412,6 +412,9 @@
 		Engine/Source/ThirdParty/Catch2/Catch2.build.cs
 		Engine/Source/ThirdParty/Catch2/Catch2.tps
 		Engine/Source/ThirdParty/Catch2/Catch2Extras.Build.cs
+		
+		<!-- debug tools -->
+		Engine/Source/ThirdParty/Breakpad/src/tools/mac/binaries/dump_syms
 	</Property>
 
 	<Property Name="CopyEditorExceptionsMac" Value="$(CopyEditorExceptions)">
diff --git a/Engine/Source/Programs/AutomationTool/BuildGraph/Tasks/SymGenTask.cs b/Engine/Source/Programs/AutomationTool/BuildGraph/Tasks/SymGenTask.cs
--- a/Engine/Source/Programs/AutomationTool/BuildGraph/Tasks/SymGenTask.cs
+++ b/Engine/Source/Programs/AutomationTool/BuildGraph/Tasks/SymGenTask.cs
@@ -203,7 +203,7 @@ namespace BuildGraph.Tasks
 						}
 					}
 
-					dumpSymsArgs = extraOptions + sourceFile.FullName;
+					dumpSymsArgs = extraOptions + "\"" + sourceFile.FullName + "\"";
 				}
 
 				IProcessResult result = CommandUtils.Run(symbolDumperExeForFile, dumpSymsArgs, null, CommandUtils.ERunOptions.AppMustExist, null, FilterSpew, null, workingDirectory);
