Convert Gauntlet.RpcLibrary C# wrappers - AutomationControllerSummarizeReports-> AutomationControllerSummarizeReports (move C++ endpoint from old lib to a new one)
Convert Gauntlet.RpcLibrary C# wrappers - ExecuteConsoleCommand-> ValkExecuteConsoleCommand
Convert Gauntlet.RpcLibrary C# wrappers - AutomationControllerClearReports-> AutomationControllerClearReports (move C++ endpoint from old lib to a new one)

#rnx

diff --git a/Engine/Source/Programs/AutomationTool/Gauntlet/Framework/RpcFramework/Gauntlet.RpcLibrary.cs b/Engine/Source/Programs/AutomationTool/Gauntlet/Framework/RpcFramework/Gauntlet.RpcLibrary.cs
--- a/Engine/Source/Programs/AutomationTool/Gauntlet/Framework/RpcFramework/Gauntlet.RpcLibrary.cs
+++ b/Engine/Source/Programs/AutomationTool/Gauntlet/Framework/RpcFramework/Gauntlet.RpcLibrary.cs
@@ -329,18 +329,6 @@ namespace Gauntlet
 	        return true;
         }
 
-		public static bool ExecuteConsoleCommand(RpcTarget InTarget, string ConsoleCommandString)
-		{
-			if (InTarget == null)
-			{
-				return false;
-			}
-			Dictionary<string, object> RequestArgs = new Dictionary<string, object>();
-			RequestArgs.Add("Command", ConsoleCommandString);
-			RpcExecutor.CallRpc(InTarget, "ExecuteConsoleCommand", 5000, RequestArgs);
-			return true;
-		}
-
 		public static bool ExitClient(RpcTarget InTarget)
 		{
 			if (InTarget == null)
@@ -629,7 +617,42 @@ namespace Gauntlet
 	{
 		public static bool AutomationControllerClearReports(this RpcTarget InTarget)
 		{
-			return RpcLibrary.ExecuteRpcCheckSuccess(InTarget, "AutomationControllerClearReports");
+			const string RPCName = "AutomationControllerClearReports";
+			using (HttpResponseMessage RpcResponse = RpcExecutor.CallRpc(InTarget, RPCName))
+			{
+				RpcLibrary.SimpleResponse Response = RpcLibrary.GetSimpleResponse(RpcResponse);
+				if (!Response.Succeeded)
+				{
+					throw new TestException($"RPC '{RPCName}' failed: {Response.Value}");
+				}
+
+				return true;
+			}
+		}
+
+		public static bool AutomationControllerSummarizeReports(this RpcTarget InTarget, out Dictionary<string, string> OutResponse)
+		{
+			const string RPCName = "AutomationControllerSummarizeReports";
+			try
+			{
+				using (HttpResponseMessage RpcResponse = RpcExecutor.CallRpc(InTarget, RPCName))
+				{
+					if (RpcResponse != null)
+					{
+						OutResponse = RpcLibrary.GetResponseAsObject<Dictionary<string, string>>(RpcResponse);
+						return true;
+					}
+
+					OutResponse = null;
+					return false;
+				}
+			}
+			catch (TestException Exception)
+			{
+				Log.Error($"Automation test caused {Exception}");
+				OutResponse = null;
+				return false;
+			}
 		}
 	}
 
