[Gauntlet] iOS - Set shorter timeout per files deployed

#rnx

diff --git a/Engine/Source/Programs/AutomationTool/Gauntlet/Platform/IOS/Gauntlet.TargetDeviceIOS.cs b/Engine/Source/Programs/AutomationTool/Gauntlet/Platform/IOS/Gauntlet.TargetDeviceIOS.cs
--- a/Engine/Source/Programs/AutomationTool/Gauntlet/Platform/IOS/Gauntlet.TargetDeviceIOS.cs
+++ b/Engine/Source/Programs/AutomationTool/Gauntlet/Platform/IOS/Gauntlet.TargetDeviceIOS.cs
@@ -122,6 +122,9 @@ namespace Gauntlet
 		[AutoParam(60 * 60 * 2)] // iOS installs can be slow... 2 hour default
 		public int MaxInstallTime { get; protected set; }
 
+		[AutoParam(60 * 15)]
+		public int MaxPerFileInstallTime { get; protected set; }
+
 		/// <summary>
 		/// Temp path we use to push/pull things from the device
 		/// </summary>
@@ -559,16 +562,18 @@ namespace Gauntlet
 					string FileName = Path.GetFileName(ContentFile.Key);
 					Log.Info("({Current}/{Count}) Installing {FileName}...", CurrentFileIndex, BulkContentFileCount, FileName);
 
+					double RemainingTimePerFile = Math.Min(RemainingTime, MaxPerFileInstallTime);
 					string QualifiedSourcePath = Path.GetFullPath(Path.Combine(AppPath, ContentFile.Key));
-					IProcessResult ContentCopyResult = CopyFileToDevice(QualifiedSourcePath, ContentFile.Value, Build.PackageName, (int)RemainingTime, true);
+					IProcessResult ContentCopyResult = CopyFileToDevice(QualifiedSourcePath, ContentFile.Value, Build.PackageName, (int)RemainingTimePerFile, true);
 
-					RemainingTime -= (DateTime.Now - StartTime).TotalSeconds;
+					double DeltaTime = (DateTime.Now - StartTime).TotalSeconds;
+					RemainingTime -= DeltaTime;
 
 					if (ContentCopyResult.ExitCode != 0)
 					{
-						if(RemainingTime <= 0)
+						if(RemainingTimePerFile <= DeltaTime)
 						{
-							throw new AutomationException("Install timed out after {MaxInstallTime} seconds", MaxInstallTime);
+							throw new AutomationException("Install timed out after {MaxPerFileInstallTime} seconds", RemainingTimePerFile);
 						}
 						else
 						{
