[Gauntlet] Fix an issue where failing to start app on the device will lead to try to reserve another device

#rnx

diff --git a/Engine/Source/Programs/AutomationTool/Gauntlet/Unreal/Base/Gauntlet.UnrealSession.cs b/Engine/Source/Programs/AutomationTool/Gauntlet/Unreal/Base/Gauntlet.UnrealSession.cs
--- a/Engine/Source/Programs/AutomationTool/Gauntlet/Unreal/Base/Gauntlet.UnrealSession.cs
+++ b/Engine/Source/Programs/AutomationTool/Gauntlet/Unreal/Base/Gauntlet.UnrealSession.cs
@@ -10,6 +10,7 @@ using System.Text.RegularExpressions;
 using System.Threading;
 using UnrealBuildTool;
 using Gauntlet.Utils;
+using System.Net;
 
 namespace Gauntlet
 {
@@ -714,7 +715,7 @@ namespace Gauntlet
 		/// <returns></returns>
 		public bool TryReserveDevices()
 		{
-			if(ShouldRetainDevices && HasAcquiredDevices())
+			if (ShouldRetainDevices && HasAcquiredDevices())
 			{
 				return true;
 			}
@@ -724,9 +725,9 @@ namespace Gauntlet
 			IEnumerable<UnrealSessionRole> RolesThatRequireDevice = SessionRoles.Where(R => !R.IsNullRole());
 
 			// Get a count of the number of devices required for each platform
-			foreach(UnrealSessionRole Role in RolesThatRequireDevice)
+			foreach (UnrealSessionRole Role in RolesThatRequireDevice)
 			{
-				if(RequiredDeviceTypes.ContainsKey(Role.Constraint))
+				if (RequiredDeviceTypes.ContainsKey(Role.Constraint))
 				{
 					++RequiredDeviceTypes[Role.Constraint];
 				}
@@ -736,7 +737,24 @@ namespace Gauntlet
 				}
 			}
 
-			return UnrealDeviceReservation.TryReserveDevices(RequiredDeviceTypes, RolesThatRequireDevice.Count());
+			try
+			{
+				return UnrealDeviceReservation.TryReserveDevices(RequiredDeviceTypes, RolesThatRequireDevice.Count());
+			}
+			catch (Exception ex)
+			{
+				if (ex.InnerException != null && ex.InnerException is WebException
+					&& (ex.InnerException as WebException).Response != null
+					&& ((ex.InnerException as WebException).Response as HttpWebResponse).StatusCode == HttpStatusCode.Conflict)
+				{
+					// Reservation failed for lack of devices
+					// Return false to let attempt loop to kick in
+					Log.Info(ex.Message);
+					return false;
+				}
+
+				throw;
+			}
 		}
 
 		public bool TryReserveDevices(int Attempts)
