Jira UE-230638
Adapt UAT Localize Command for Mac and Linux

#tests Run a localize command on Windows, Mac or Linux and verify in the logs that the executable path is correct

diff --git a/Engine/Source/Programs/AutomationTool/Scripts/Localisation.Automation.cs b/Engine/Source/Programs/AutomationTool/Scripts/Localisation.Automation.cs
--- a/Engine/Source/Programs/AutomationTool/Scripts/Localisation.Automation.cs
+++ b/Engine/Source/Programs/AutomationTool/Scripts/Localisation.Automation.cs
@@ -805,13 +805,28 @@ class Localize : BuildCommand
 
 	private string GetEditorExePath()
 	{
-		var EditorExe = CombinePaths(CmdEnv.LocalRoot, @"Engine/Binaries/Win64/UnrealEditor-Cmd.exe");
+		var EditorExe = HostPlatform.Current.GetUnrealExePath("UnrealEditor-Cmd.exe");
+		Logger.LogInformation($"Looking for editor executable: {EditorExe}");
 		if (!File.Exists(EditorExe))
 		{
-			// Try using the debug .exe instead 
-			EditorExe = CombinePaths(CmdEnv.LocalRoot, @"Engine/Binaries/Win64/UnrealEditor-Win64-Debug-Cmd.exe");
+			// Try using the debug executable instead
+			// (which is supposed to have a different name depending on the platform)
+			var CurrentOS = RuntimePlatform.Current;
+			if (CurrentOS == RuntimePlatform.Type.Mac)
+			{
+				EditorExe = "UnrealEditor-Mac-Debug-Cmd";
+			}
+			else if (CurrentOS == RuntimePlatform.Type.Linux)
+			{
+				EditorExe = "UnrealEditor-Linux-Debug-Cmd";
+			}
+			else
+			{
+				EditorExe = "UnrealEditor-Win64-Debug-Cmd.exe";
+			}
+			Logger.LogInformation($"Editor executable not found, now looking for: {EditorExe}");
+			EditorExe = HostPlatform.Current.GetUnrealExePath(EditorExe);
 		}
-
 		return EditorExe;
 	}
 
