Fix formatting string

#rnx

diff --git a/Engine/Source/Programs/AutomationTool/Gauntlet/Platform/IOS/Gauntlet.TargetDeviceIOS.cs b/Engine/Source/Programs/AutomationTool/Gauntlet/Platform/IOS/Gauntlet.TargetDeviceIOS.cs
--- a/Engine/Source/Programs/AutomationTool/Gauntlet/Platform/IOS/Gauntlet.TargetDeviceIOS.cs
+++ b/Engine/Source/Programs/AutomationTool/Gauntlet/Platform/IOS/Gauntlet.TargetDeviceIOS.cs
@@ -573,7 +573,7 @@ namespace Gauntlet
 					{
 						if(RemainingTimePerFile <= DeltaTime)
 						{
-							throw new AutomationException("Install timed out after {MaxPerFileInstallTime} seconds", RemainingTimePerFile);
+							throw new AutomationException("Install timed out after {0} seconds", RemainingTimePerFile);
 						}
 						else
 						{
