Make sure that IasTool is compiled when the command line option '-ApplyIoStoreOnDemand' is used.

#rnx

diff --git a/Engine/Source/Programs/AutomationTool/Scripts/BuildProjectCommand.Automation.cs b/Engine/Source/Programs/AutomationTool/Scripts/BuildProjectCommand.Automation.cs
--- a/Engine/Source/Programs/AutomationTool/Scripts/BuildProjectCommand.Automation.cs
+++ b/Engine/Source/Programs/AutomationTool/Scripts/BuildProjectCommand.Automation.cs
@@ -116,6 +116,12 @@ namespace AutomationScripts
 				{
 					ProjectAgenda.AddTargets(new string[] { "UnrealPak" }, HostPlatform.Current.HostEditorPlatform, UnrealTargetConfiguration.Development, Params.CodeBasedUprojectPath, InAddArgs: ProgramBuildArgs);
 				}
+
+				// If IoStoreOnDemand is being forced onto a project then we will also need to compile IasTool for use during staging
+				if (Params.ApplyIoStoreOnDemand && (!Params.HasEditorTargets || Params.EditorTargets.Contains("IasTool") == false))
+				{
+					NoProjectAgenda.AddTargets(new string[] { "IasTool" }, HostPlatform.Current.HostEditorPlatform, UnrealTargetConfiguration.Development, InAddArgs: ProgramBuildArgs);
+				}
 			}
 
 			// Additional compile arguments
