Allow '-ApplyIoStoreOnDemand' to provide a default upload target

#rnx

- Previous if a user ran BuildCookRun with '-ApplyIoStoreOnDemand' but without '-Upload=???' the ondemand containers would be created but nothing would get uploaded and the ondemand system will not initialize at runtime.
- Now if this is the case, the presence of the '-ApplyIoStoreOnDemand' will force the upload target to 'LocalZen' rather than skipping the upload phase, so the command will just work.
- I've opted to put the new logic checking the flag in it's own code block rather than adding addition AND/OR checks to the upload code as keeping the 'ApplyIoStoreOnDemand' code isolated will make it easier to remove in the future (if we choose to do so) and makes the intent more obvious.

diff --git a/Engine/Source/Programs/AutomationTool/Scripts/CopyBuildToStagingDirectory.Automation.cs b/Engine/Source/Programs/AutomationTool/Scripts/CopyBuildToStagingDirectory.Automation.cs
--- a/Engine/Source/Programs/AutomationTool/Scripts/CopyBuildToStagingDirectory.Automation.cs
+++ b/Engine/Source/Programs/AutomationTool/Scripts/CopyBuildToStagingDirectory.Automation.cs
@@ -4192,6 +4192,14 @@ namespace AutomationScripts
 				RunIoStore(Params, SC, IoStoreCommandsFileName, FinalOrderFiles, AdditionalArgs);
 			}
 
+			// When the ApplyIoStoreOnDemand flag is present then we need to make sure that ondemand data is uploaded.
+			// If no upload target has been specified we will default to using the local zen server which is the equivalent
+			// of providing '-Upload=LocalZen' as a commandline arg.
+			if (Params.ApplyIoStoreOnDemand && string.IsNullOrWhiteSpace(Params.Upload))
+			{
+				Params.Upload = "localzen";
+			}
+
 			if (!string.IsNullOrWhiteSpace(Params.Upload))
 			{
 				StringBuilder UploadArgs = new StringBuilder();
