Cooker: IncrementalCook: CopyBuildToStagingDirectory needs to only the ReferencedSet when staging, to match the behavior in UnrealPak. Use of the ReferencedSet ensures that previously cooked packages that are no longer referenced are not staged.
#rnx

diff --git a/Engine/Source/Programs/AutomationTool/Scripts/CopyBuildToStagingDirectory.Automation.cs b/Engine/Source/Programs/AutomationTool/Scripts/CopyBuildToStagingDirectory.Automation.cs
--- a/Engine/Source/Programs/AutomationTool/Scripts/CopyBuildToStagingDirectory.Automation.cs
+++ b/Engine/Source/Programs/AutomationTool/Scripts/CopyBuildToStagingDirectory.Automation.cs
@@ -919,7 +919,10 @@ namespace AutomationScripts
 			for (; ; )
 			{
 				HttpClient HttpClient = new HttpClient();
-				using var Request = new HttpRequestMessage(HttpMethod.Get, string.Format("http://{0}:{1}/prj/{2}/oplog/{3}/entries", Host, Port, ProjectId, OplogId));
+				string Uri = string.Format("http://{0}:{1}/prj/{2}/oplog/{3}/entries", Host, Port, ProjectId, OplogId);
+				Uri += "?trim_by_referencedset=true";
+
+				using var Request = new HttpRequestMessage(HttpMethod.Get, Uri);
 				Request.Headers.Add("Accept", "application/x-ue-cb");
 				try
 				{
