Handle device factory CreateDevice returning null

#rnx

diff --git a/Engine/Source/Programs/AutomationTool/Gauntlet/Framework/Devices/Gauntlet.DevicePool.cs b/Engine/Source/Programs/AutomationTool/Gauntlet/Framework/Devices/Gauntlet.DevicePool.cs
--- a/Engine/Source/Programs/AutomationTool/Gauntlet/Framework/Devices/Gauntlet.DevicePool.cs
+++ b/Engine/Source/Programs/AutomationTool/Gauntlet/Framework/Devices/Gauntlet.DevicePool.cs
@@ -1108,6 +1108,12 @@ namespace Gauntlet
 					NewDevice = Factory.CreateDevice(Def.Address, ClientTempDir, Def.DeviceData);
 				}
 
+				if (NewDevice == null)
+				{
+					Log.Info(KnownLogEvents.Gauntlet_DeviceEvent, "Failed to create device {Name}. Device could not be connected.", NewDevice);
+					return null;
+				}
+
 				if (NewDevice.IsAvailable == false)
 				{
 					Log.Info(KnownLogEvents.Gauntlet_DeviceEvent, "Assigned device {Name} reports unavailable. Requesting a forced disconnect", NewDevice);
@@ -1115,7 +1121,7 @@ namespace Gauntlet
 
 					if (NewDevice.IsAvailable == false)
 					{
-						Log.Info(KnownLogEvents.Gauntlet_DeviceEvent, "Assigned device {Name} still  unavailable. Requesting a reboot", NewDevice);
+						Log.Info(KnownLogEvents.Gauntlet_DeviceEvent, "Assigned device {Name} still unavailable. Requesting a reboot", NewDevice);
 						NewDevice.Reboot();
 					}
 				}
