Fix name reference in log when device could not be created

#rnx

diff --git a/Engine/Source/Programs/AutomationTool/Gauntlet/Framework/Devices/Gauntlet.DevicePool.cs b/Engine/Source/Programs/AutomationTool/Gauntlet/Framework/Devices/Gauntlet.DevicePool.cs
--- a/Engine/Source/Programs/AutomationTool/Gauntlet/Framework/Devices/Gauntlet.DevicePool.cs
+++ b/Engine/Source/Programs/AutomationTool/Gauntlet/Framework/Devices/Gauntlet.DevicePool.cs
@@ -1110,7 +1110,7 @@ namespace Gauntlet
 
 				if (NewDevice == null)
 				{
-					Log.Info(KnownLogEvents.Gauntlet_DeviceEvent, "Failed to create device {Name}. Device could not be connected.", NewDevice);
+					Log.Info(KnownLogEvents.Gauntlet_DeviceEvent, "Failed to create device {Name}. Device could not be connected.", Def.ToString());
 					return null;
 				}
 
