Save and restore window positions for proxy clients in the same way as direct clients.

diff --git a/Engine/Source/Programs/AutomationTool/Scripts/LaunchMultiServer.Automation.cs b/Engine/Source/Programs/AutomationTool/Scripts/LaunchMultiServer.Automation.cs
--- a/Engine/Source/Programs/AutomationTool/Scripts/LaunchMultiServer.Automation.cs
+++ b/Engine/Source/Programs/AutomationTool/Scripts/LaunchMultiServer.Automation.cs
@@ -258,9 +258,9 @@ namespace AutomationTool
 							Logger.LogInformation("Starting proxy server connected to MultiServer instances");
 							StartProcessesForServerInfos(UniqueServerInfos, ProcessLaunchType.ProxyServer, "");
 
-							for (int i = 0; i < ClientCount; i++)
+							for (int ClientIndex = 0; ClientIndex < ClientCount; ClientIndex++)
 							{
-								StartProcessesForServerInfos(UniqueServerInfos, ProcessLaunchType.ProxyClient, "");
+								StartProcessesForServerInfos(UniqueServerInfos, ProcessLaunchType.ProxyClient, " -SaveWinPos="+(ClientIndex + 1));
 							}
 						}
 					}
@@ -324,9 +324,9 @@ namespace AutomationTool
 
 						if (LaunchType == ProcessLaunchType.Server)
 						{
-						Args += String.Format("-MultiServerLocalId={0} -MultiServerListenPort={1} -port={2} ", ServerInfo.Definition.LocalId, ServerInfo.Definition.ListenPort, ServerInfo.GameListenPort);
-						Args += String.Format("-ConsoleTitle=\"{0} Server ID {1}\" ", Params.ShortProjectName, ServerInfo.Definition.LocalId);
-						Args += Params.ServerCommandline + " ";
+							Args += String.Format("-MultiServerLocalId={0} -MultiServerListenPort={1} -port={2} ", ServerInfo.Definition.LocalId, ServerInfo.Definition.ListenPort, ServerInfo.GameListenPort);
+							Args += String.Format("-ConsoleTitle=\"{0} Server ID {1}\" ", Params.ShortProjectName, ServerInfo.Definition.LocalId);
+							Args += Params.ServerCommandline + " ";
 						}
 						else if (LaunchType == ProcessLaunchType.DirectClient)
 						{
@@ -403,7 +403,8 @@ namespace AutomationTool
 
 					Args += "127.0.0.1:8000 " + (Params.Cook ? "" : "-game ");
 					Args += Params.ClientCommandline + " ";
-					Args += "-windowed";
+					Args += "-windowed ";
+					Args += AdditionalParameters;
 
 					var NewProcess = Run(App, Args, null, ERunOptions.Default | ERunOptions.NoWaitForExit | ERunOptions.NoStdOutRedirect);
 
