Increase Install Build timeout to 2.5 hours to support long running iOS installs

diff --git a/Engine/Source/Programs/AutomationTool/Gauntlet/Platform/IOS/Gauntlet.TargetDeviceIOS.cs b/Engine/Source/Programs/AutomationTool/Gauntlet/Platform/IOS/Gauntlet.TargetDeviceIOS.cs
--- a/Engine/Source/Programs/AutomationTool/Gauntlet/Platform/IOS/Gauntlet.TargetDeviceIOS.cs
+++ b/Engine/Source/Programs/AutomationTool/Gauntlet/Platform/IOS/Gauntlet.TargetDeviceIOS.cs
@@ -119,7 +119,7 @@ namespace Gauntlet
 		public bool IsAvailable => true;
 		#endregion
 
-		[AutoParam(60 * 60 * 2)] // iOS installs can be slow... 2 hour default
+		[AutoParam(60 * 60 * 2.5)] // iOS installs can be slow... 2.5 hour default
 		public int MaxInstallTime { get; protected set; }
 
 		[AutoParam(60 * 15)]
