Add a UAT command to streamline APT hookup.

- First pass, but is workable.

diff --git a/Engine/Plugins/Performance/AutomatedPerfTesting/Resources/TemplatesForHookup/AutoPerfTests.xml b/Engine/Plugins/Performance/AutomatedPerfTesting/Resources/TemplatesForHookup/AutoPerfTests.xml
new file mode 100644
--- /dev/null
+++ b/Engine/Plugins/Performance/AutomatedPerfTesting/Resources/TemplatesForHookup/AutoPerfTests.xml
@@ -0,0 +1,33 @@
+<?xml version='1.0' ?>
+<BuildGraph xmlns="http://www.epicgames.com/BuildGraph" xmlns:xsi="http://www.w3.org/2001/XMLSchema-instance" xsi:schemaLocation="http://www.epicgames.com/BuildGraph ../Schema.xsd" >
+
+<!-- Project Info -->
+<Property Name="ProjectName" Value="**REPLACE_PROJECTNAME**" />
+<Property Name="ProjectPath" Value="**REPLACE_PROJECTPATH**" />
+
+<!-- Stick to the defaults unless the project needs to change this -->
+<Property Name="WithBATDefaults" Value="true" />
+
+<!-- If a build machine, turn off all defaults. Settings should be specified via the website UI -->
+<Do If="$(IsBuildMachine)">
+	<Property Name="DefaultEditorPlatforms" Value="Win64" />
+	<Property Name="DefaultTargetPlatforms" Value="Win64" />
+	<Property Name="DefaultTargetConfigurations" Value="Development" />
+	<Property Name="DefaultEditorTestList" Value="" />
+	<Property Name="DefaultTargetTestList" Value="" />
+</Do>
+
+<!-- Configure AutomatedPerfTest Sequence test -->
+<Include Script="$(RootDir)/Engine/Plugins/Performance/AutomatedPerfTesting/Build/Inc/AutomatedPerfTestProjectSettings.xml" />
+<Do If="'$(MapSequenceComboNames)' == ''">
+	<Property Name="MapSequenceComboNames" Value="PerfSequence" />
+</Do>
+
+
+<!-- Common Setting Overrides -->
+<Property Name="GauntletSettingsFile" Value="$(RootDir)/$(ProjectPath)/Build/Inc/GauntletSettings.xml" />
+
+<!-- This will declare an aggregate called BuildAndTest NaniteTest -->
+<Include Script="$(RootDir)/Engine/Build/Graph/Tasks/BuildAndTestProject.xml" />
+
+</BuildGraph>
diff --git a/Engine/Plugins/Performance/AutomatedPerfTesting/Resources/TemplatesForHookup/GauntletSettings.xml b/Engine/Plugins/Performance/AutomatedPerfTesting/Resources/TemplatesForHookup/GauntletSettings.xml
new file mode 100644
--- /dev/null
+++ b/Engine/Plugins/Performance/AutomatedPerfTesting/Resources/TemplatesForHookup/GauntletSettings.xml
@@ -0,0 +1,14 @@
+<?xml version='1.0' ?>
+<BuildGraph xmlns="http://www.epicgames.com/BuildGraph" xmlns:xsi="http://www.w3.org/2001/XMLSchema-instance" xsi:schemaLocation="http://www.epicgames.com/BuildGraph ../../Schema.xsd" >
+
+	<!-- Gauntlet additional arguments - use dx12 RHI -->
+	<!-- Using BuildNamePath specifically as we want the preflight suffix if it exists (cargo cult?). -->
+	<Option Name="ExtraGauntletCommonArgs" DefaultValue="-BuildName=$(BuildNamePath)" Description="Additional Common arguments passed to Gauntlet command line"/>
+
+	<!-- Common AutoTestPerf Gauntlet settings -->
+	<Include Script="$(RootDir)/Engine/Plugins/Performance/AutomatedPerfTesting/Build/Inc/AutomatedPerfTestCommonSettings.xml" />
+
+	<!-- Common Gauntlet settings -->
+	<Include Script="$(RootDir)/Engine/Build/Graph/Tasks/Inc/GauntletSettings.xml" />
+
+</BuildGraph>
\ No newline at end of file
diff --git a/Engine/Source/Programs/AutomationTool/Scripts/AddAutomatedPerfTestToProject.cs b/Engine/Source/Programs/AutomationTool/Scripts/AddAutomatedPerfTestToProject.cs
new file mode 100644
--- /dev/null
+++ b/Engine/Source/Programs/AutomationTool/Scripts/AddAutomatedPerfTestToProject.cs
@@ -0,0 +1,229 @@
+// Copyright Epic Games, Inc. All Rights Reserved.
+using System;
+using System.Collections.Generic;
+using System.Text;
+using System.IO;
+using AutomationTool;
+using UnrealBuildTool;
+using UnrealBuildBase;
+using EpicGames.Core;
+using System.Linq;
+using System.Text.Json;
+using Microsoft.Extensions.Logging;
+using Microsoft.VisualStudio.OLE.Interop;
+using System.Threading.Tasks;
+using System.Security.AccessControl;
+
+[Help("UAT command to enable PerfTesting hook ups for a project.")]
+[Help("-Project", "Path to the project to be instrumented.")]
+class AddAutomatedPerfTestToProject : BuildCommand
+{
+	/// <summary>
+	/// The plugin name itself
+	/// </summary>
+	const string AptPluginName = "AutomatedPerfTesting";
+
+	/// <summary>
+	/// Path to the AutomatePerfTesting plugin
+	/// </summary>
+	const string AptPathFromRoot = "Engine/Plugins/Performance/AutomatedPerfTesting";
+
+	/// <summary>
+	/// Path to the AutomatePerfTesting templates
+	/// </summary>
+	const string AptTemplatesFromPluginRoot = "Resources/TemplatesForHookup";
+
+	/// <summary>
+	/// Name of the Gauntlet settings that needs to be copied verbatim
+	/// </summary>
+	const string AptGauntletIncludeFile = "GauntletSettings.xml";
+
+	/// <summary>
+	/// Name of the template file with project setings that will be modified
+	/// </summary>
+	const string AptProjectBuildGraphTemplateFile = "AutoPerfTests.xml";
+
+	/// <summary>
+	/// Name of the config section for APT
+	/// </summary>
+	const string AptConfigSectionName = "/Script/AutomatedPerfTesting.AutomatedSequencePerfTestProjectSettings";
+
+	protected bool EnablePlugin(string ProjectFilePath)
+	{
+		Logger.LogInformation("Enabling the plugin for the project.");
+
+		JsonObject ProjectJson = JsonObject.Read(new FileReference(ProjectFilePath));
+
+		JsonObject AptReference = new JsonObject();
+		AptReference.AddOrSetFieldValue("Name", AptPluginName);
+		AptReference.AddOrSetFieldValue("Enabled", true);
+		List<JsonObject> NewPluginArray = new List<JsonObject>();
+		if(ProjectJson.TryGetObjectArrayField("Plugins", out JsonObject[] Plugins))
+		{
+			Logger.LogInformation("Project has some plugins configured for it, checking if APT is among them.");
+
+			foreach(JsonObject PluginDesc in Plugins) 
+			{
+				string PluginName;
+				if (PluginDesc.TryGetStringField("Name",out PluginName) && PluginName == AptPluginName) 
+				{
+					Logger.LogInformation("Project already has AutomatedPerfTest plugin enabled.");
+					return true;
+				}
+				NewPluginArray.Add(PluginDesc);
+			}
+		}
+		else
+		{
+			Logger.LogInformation("Project had no plugins configured for it, adding a block.");
+		}
+
+		Logger.LogInformation("Adding APT to the list of project plugins.");
+		NewPluginArray.Add(AptReference);
+		ProjectJson.AddOrSetFieldValue("Plugins", NewPluginArray.ToArray());
+
+		string TempFile = ProjectFilePath + ".temp";	// unlikely to have more than one upgraders running at the same time
+		File.WriteAllText(TempFile,ProjectJson.ToJsonString());
+
+		ProjectJson = null;
+		File.Delete(ProjectFilePath);
+		File.Move(TempFile, ProjectFilePath);
+		return true;
+	}
+	protected bool ModifyTargetCSFiles(string Project)
+	{
+		Logger.LogInformation("Modifying target files for the game/client configuration.");
+
+		// TODO: find either ProjectName.Target.cs or ProjectNameClient.Target.cs  (which needs 'Game' to be removed) and modify accordingly.
+		// Ideally just nix this step by making logging in Shipping and LLM in Test engine defaults
+
+		return true;
+	}
+
+	protected bool ModifyProjectConfigFiles(string ProjectFilePath)
+	{
+		Logger.LogInformation("Modifying project settings.");
+
+		string ProjectDir = Path.GetDirectoryName(ProjectFilePath);
+
+		// load ini file
+		string DefaultEngineIniFile = Path.Combine(ProjectDir,"Config","DefaultEngine.ini");
+		if(!File.Exists(DefaultEngineIniFile)) 
+		{
+			Logger.LogInformation("Project does not seem to have DefaultEngine.ini.");
+			return true;
+		}
+
+		ConfigFile DefaultEngineIni = new ConfigFile(new FileReference(DefaultEngineIniFile));
+
+		ConfigFileSection Section = null;
+		if (DefaultEngineIni.TryGetSection(AptConfigSectionName, out Section))
+		{
+			Logger.LogInformation("Project already has APT settings in its DefaultEngine.ini.");
+			return true;
+		}
+
+		Section = DefaultEngineIni.FindOrAddSection(AptConfigSectionName);
+		Section.Lines.Add(new ConfigLine(ConfigLineAction.Add, "MapsAndSequencesToTest", "ComboName = \"PerfSeqeunce\", Map = \"/...\", Sequence = \"...\""));
+		Section.Lines.Add(new ConfigLine(ConfigLineAction.Set,"SequenceStartDelay", "5.0000"));
+
+		DefaultEngineIni.Write(new FileReference(DefaultEngineIniFile));
+
+		;// +MapsAndSequencesToTest = (ComboName = "PerfSequence", Map = "/Game/Maps/LV_Main.LV_Main", Sequence = "/Game/Cinematics/LS_ReplaySequence.LS_ReplaySequence")
+		;// SequenceStartDelay = 5.000000
+
+		return true;
+	}
+
+	/// <summary>
+	/// Return a project name from the path to .uproject
+	/// </summary>
+	protected string GetProjectName(string ProjectFilePath)
+	{
+		return Path.GetFileNameWithoutExtension(ProjectFilePath);
+	}
+
+	/// <summary>
+	/// Return a project path relative to the current engine root, with forward slashes.
+	/// </summary>
+	protected string GetProjectRelativePath(string ProjectFilePath) 
+	{
+		string ProjectPath = Path.GetDirectoryName(ProjectFilePath);
+		ProjectPath = ProjectPath.Replace(Unreal.EngineDirectory.FullName, "");
+		ProjectPath = CommandUtils.ConvertSeparators(PathSeparator.Slash, ProjectPath);
+		return ProjectPath;
+	}
+
+	protected bool ModifyBuildGraphFiles(string ProjectFilePath)
+	{
+		// work out the target directory
+		string ProjectDir = Path.GetDirectoryName(ProjectFilePath);
+
+		string BuildGraphTemplates = Path.Combine(AptPathFromRoot, AptTemplatesFromPluginRoot);
+
+		string BuildGraphDest = Path.Combine(ProjectDir, "Build");
+		string BuildGraphGauntletIncDest = Path.Combine(BuildGraphDest, "Inc");
+
+		Directory.CreateDirectory(BuildGraphDest);
+		// copy Gauntlet settings file
+		{
+			Directory.CreateDirectory(BuildGraphGauntletIncDest);
+
+			string DestFile = Path.Combine(BuildGraphGauntletIncDest,AptGauntletIncludeFile);
+			if(File.Exists(DestFile))
+			{
+				Logger.LogInformation("Project already has file {0}.",DestFile);
+			}
+			else
+			{
+				string SrcFile = Path.Combine(BuildGraphTemplates,AptGauntletIncludeFile);
+				File.Copy(SrcFile,DestFile);
+			}
+		}
+
+		// copy example per-project BuildGraph file
+		{
+			string DestFile = Path.Combine(BuildGraphDest,AptProjectBuildGraphTemplateFile);
+			if(File.Exists(DestFile))
+			{
+				Logger.LogInformation("Project already has file {0}.",DestFile);
+			}
+			else
+			{
+				string SrcFile = Path.Combine(BuildGraphTemplates,AptProjectBuildGraphTemplateFile);
+				string Template = File.ReadAllText(SrcFile);
+				Template = Template.Replace("**REPLACE_PROJECTNAME**", GetProjectName(ProjectFilePath));
+				Template = Template.Replace("**REPLACE_PROJECTPATH**",GetProjectRelativePath(ProjectFilePath));
+				File.WriteAllText(DestFile, Template);
+			}
+		}
+
+		return true;
+	}
+
+	public override void ExecuteBuild()
+	{
+		Logger.LogInformation("Adding AutomatedPerfTest hookups for a project.");
+
+		// Parse the project
+		string Project = ParseParamValue("Project");
+		if (Project == null)
+		{
+			throw new AutomationException("Project to instrument was not specified, use -Project=<path_to_uproject>. Lookup via Default.uprojectdirs is not supported");
+		}
+
+		if (!File.Exists(Project))
+		{
+			throw new AutomationException("Specified project file '{0}' was not found", Project);
+		}
+
+		Logger.LogInformation("Project {0} will be configured for AutomatedPerfTest plugin.", Project);
+		
+		EnablePlugin(Project);
+		ModifyProjectConfigFiles(Project);
+		ModifyTargetCSFiles(Project);
+		ModifyBuildGraphFiles(Project);
+
+		Logger.LogInformation("Successfully added.");
+	}
+}
