Add ability to export zen snapshot to builds API in cloud storage.

diff --git a/Engine/Source/Programs/AutomationTool/BuildGraph/Tasks/ZenExportSnapshotTask.cs b/Engine/Source/Programs/AutomationTool/BuildGraph/Tasks/ZenExportSnapshotTask.cs
--- a/Engine/Source/Programs/AutomationTool/BuildGraph/Tasks/ZenExportSnapshotTask.cs
+++ b/Engine/Source/Programs/AutomationTool/BuildGraph/Tasks/ZenExportSnapshotTask.cs
@@ -14,6 +14,7 @@ using System.Xml;
 using EpicGames.Core;
 using EpicGames.Horde;
 using EpicGames.ProjectStore;
+using EpicGames.Serialization;
 using Microsoft.Extensions.Logging;
 
 namespace AutomationTool.Tasks
@@ -28,10 +29,14 @@ namespace AutomationTool.Tasks
 		/// </summary>
 		Invalid,
 		/// <summary>
-		/// Snapshot stored in cloud repositories such as Unreal Cloud DDC.
+		/// Snapshot stored in cloud repository such as Unreal Cloud DDC.
 		/// </summary>
 		Cloud,
 		/// <summary>
+		/// Snapshot stored in builds repository such as Unreal Cloud DDC.
+		/// </summary>
+		Builds,
+		/// <summary>
 		/// Snapshot stored in a zenserver.
 		/// </summary>
 		Zen,
@@ -58,6 +63,12 @@ namespace AutomationTool.Tasks
 		[TaskParameter(Optional = true)]
 		public string Platform { get; set; }
 
+		/// <summary>
+		/// The metadata to associate with the snapshot
+		/// </summary>
+		[TaskParameter(Optional = true)]
+		public string Metadata { get; set; }
+
 		/// <summary>
 		/// A file to read with information about the snapshot that should be used as a base when exporting this new snapshot
 		/// </summary>
@@ -322,6 +333,18 @@ namespace AutomationTool.Tasks
 			return name.Replace('/', '_').Replace(' ', '_').Replace('+', '_').Replace('-', '_');
 		}
 
+		private string GetCloudBucketName()
+		{
+			string bucketName = _parameters.DestinationCloudBucket;
+			string projectNameAsBucketName = _parameters.Project.GetFileNameWithoutAnyExtensions().ToLowerInvariant();
+			if (String.IsNullOrEmpty(bucketName))
+			{
+				bucketName = projectNameAsBucketName;
+			}
+			bucketName = SanitizeBucketName(bucketName);
+			return bucketName;
+		}
+
 		private void WriteExportSource(JsonWriter writer, SnapshotStorageType destinationStorageType, ExportSourceData exportSource, string name)
 		{
 			string targetPlatform = _parameters.SnapshotDescriptorPlatform;
@@ -332,14 +355,9 @@ namespace AutomationTool.Tasks
 			writer.WriteObjectStart();
 			switch (destinationStorageType)
 			{
+				case SnapshotStorageType.Builds:
 				case SnapshotStorageType.Cloud:
-					string bucketName = _parameters.DestinationCloudBucket;
-					string projectNameAsBucketName = _parameters.Project.GetFileNameWithoutAnyExtensions().ToLowerInvariant();
-					if (String.IsNullOrEmpty(bucketName))
-					{
-						bucketName = projectNameAsBucketName;
-					}
-					bucketName = SanitizeBucketName(bucketName);
+					string bucketName = GetCloudBucketName();
 
 					string hostName = _parameters.SnapshotDescriptorCloudHost;
 					if (String.IsNullOrEmpty(hostName))
@@ -353,9 +371,17 @@ namespace AutomationTool.Tasks
 						httpVersion = _parameters.DestinationCloudHttpVersion;
 					}
 
+					string StorageTypeName = "cloud";
+					string StorageIdentifierName = "key";
+					if (destinationStorageType == SnapshotStorageType.Builds)
+					{
+						StorageTypeName = "builds";
+						StorageIdentifierName = "builds-id";
+					}
+
 					IoHash destinationKeyHash = IoHash.Compute(Encoding.UTF8.GetBytes(name));
 					writer.WriteValue("name", name);
-					writer.WriteValue("type", "cloud");
+					writer.WriteValue("type", StorageTypeName);
 					writer.WriteValue("targetplatform", targetPlatform);
 					writer.WriteValue("host", hostName);
 					if (!String.IsNullOrEmpty(httpVersion) && !httpVersion.Equals("None", StringComparison.OrdinalIgnoreCase))
@@ -364,7 +390,7 @@ namespace AutomationTool.Tasks
 					}
 					writer.WriteValue("namespace", _parameters.DestinationCloudNamespace);
 					writer.WriteValue("bucket", bucketName);
-					writer.WriteValue("key", destinationKeyHash.ToString().ToLowerInvariant());
+					writer.WriteValue(StorageIdentifierName, destinationKeyHash.ToString().ToLowerInvariant());
 					break;
 				case SnapshotStorageType.Zen:
 					string projectName = _parameters.Project.GetFileNameWithoutAnyExtensions().ToLowerInvariant() + ".oplog";
@@ -520,6 +546,7 @@ namespace AutomationTool.Tasks
 
 			switch (destinationStorageType)
 			{
+				case SnapshotStorageType.Builds:
 				case SnapshotStorageType.Cloud:
 					if (String.IsNullOrEmpty(_parameters.DestinationCloudHost))
 					{
@@ -534,15 +561,16 @@ namespace AutomationTool.Tasks
 						throw new AutomationException("Missing destination identifier when exporting to cloud");
 					}
 
-					string bucketName = _parameters.DestinationCloudBucket;
-					string projectNameAsBucketName = projectFile.GetFileNameWithoutAnyExtensions().ToLowerInvariant();
-					if (String.IsNullOrEmpty(bucketName))
+					string bucketName = GetCloudBucketName();
+
+					string StorageTypeName = "cloud";
+					string StorageIdentifierName = "key";
+					if (destinationStorageType == SnapshotStorageType.Builds)
 					{
-						bucketName = projectNameAsBucketName;
+						StorageTypeName = "builds";
+						StorageIdentifierName = "builds-id";
 					}
-					bucketName = SanitizeBucketName(bucketName);
-
-					oplogExportCommandline.AppendFormat(" --cloud {0} --namespace {1} --bucket {2}", _parameters.DestinationCloudHost, _parameters.DestinationCloudNamespace, bucketName);
+					oplogExportCommandline.AppendFormat(" --{0} {1} --namespace {2} --bucket {3}", StorageTypeName, _parameters.DestinationCloudHost, _parameters.DestinationCloudNamespace, bucketName);
 
 					if (!String.IsNullOrEmpty(_parameters.DestinationCloudHttpVersion))
 					{
@@ -562,7 +590,7 @@ namespace AutomationTool.Tasks
 						string hostUrlArg = String.Format("--hosturl http://{0}:{1}", exportSource._isLocalHost ? "localhost" : exportSource._hostName, exportSource._hostPort);
 
 						string baseKeyArg = String.Empty;
-						if ((exportSource._snapshotBaseDescriptor != null) && !String.IsNullOrEmpty(exportSource._snapshotBaseDescriptor.Key))
+						if ((destinationStorageType == SnapshotStorageType.Cloud) && (exportSource._snapshotBaseDescriptor != null) && !String.IsNullOrEmpty(exportSource._snapshotBaseDescriptor.Key))
 						{
 							if (exportSource._snapshotBaseDescriptor.Type == SnapshotStorageType.Cloud)
 							{
@@ -577,12 +605,28 @@ namespace AutomationTool.Tasks
 						StringBuilder exportSingleSourceCommandline = new StringBuilder(oplogExportCommandline.Length);
 						exportSingleSourceCommandline.Append(oplogExportCommandline);
 
-						StringBuilder destinationKeyBuilder = new StringBuilder();
-						destinationKeyBuilder.AppendFormat("{0}.{1}.{2}", projectNameAsBucketName, _parameters.DestinationIdentifier, exportSource._oplogId);
-						exportNames[exportIndex] = destinationKeyBuilder.ToString().ToLowerInvariant();
-						IoHash destinationKeyHash = IoHash.Compute(Encoding.UTF8.GetBytes(exportNames[exportIndex]));
+						StringBuilder destinationNameBuilder = new StringBuilder();
+						destinationNameBuilder.AppendFormat("{0}.{1}.{2}", bucketName, _parameters.DestinationIdentifier, exportSource._oplogId);
+						exportNames[exportIndex] = destinationStorageType.ToString().ToLowerInvariant();
+
+						string destinationId;
+						if (destinationStorageType == SnapshotStorageType.Builds)
+						{
+							if (!string.IsNullOrEmpty(_parameters.Metadata))
+							{
+								exportSingleSourceCommandline.AppendFormat(" --builds-metadata \"{0}\"", _parameters.Metadata);
+							}
+							IoHash destinationKeyHash = IoHash.Compute(Encoding.UTF8.GetBytes(exportNames[exportIndex]));
+							CbObjectId objectId = new CbObjectId(destinationKeyHash.ToByteArray());
+							destinationId = objectId.ToString().ToLowerInvariant();
+						}
+						else
+						{
+							IoHash destinationKeyHash = IoHash.Compute(Encoding.UTF8.GetBytes(exportNames[exportIndex]));
+							destinationId = destinationKeyHash.ToString().ToLowerInvariant();
+						}
 
-						exportSingleSourceCommandline.AppendFormat(" {0} --key {1} {2} {3} {4}", hostUrlArg, destinationKeyHash.ToString().ToLowerInvariant(), baseKeyArg, exportSource._projectId, exportSource._oplogId);
+						exportSingleSourceCommandline.AppendFormat(" {0} --{1} {2} {3} {4} {5}", hostUrlArg, StorageIdentifierName, destinationId, baseKeyArg, exportSource._projectId, exportSource._oplogId);
 						if (_parameters.SkipExport || TryExportOplogCommand(zenExe.FullName, exportSingleSourceCommandline.ToString()))
 						{
 							successfullyExportedSources.Add(exportSource);
