Fix for bug in the way ids are written into snapshot descriptors for snapshots exported via the builds API.

diff --git a/Engine/Source/Programs/AutomationTool/BuildGraph/Tasks/ZenExportSnapshotTask.cs b/Engine/Source/Programs/AutomationTool/BuildGraph/Tasks/ZenExportSnapshotTask.cs
--- a/Engine/Source/Programs/AutomationTool/BuildGraph/Tasks/ZenExportSnapshotTask.cs
+++ b/Engine/Source/Programs/AutomationTool/BuildGraph/Tasks/ZenExportSnapshotTask.cs
@@ -371,15 +371,24 @@ namespace AutomationTool.Tasks
 						httpVersion = _parameters.DestinationCloudHttpVersion;
 					}
 
+					string destinationId;
 					string StorageTypeName = "cloud";
 					string StorageIdentifierName = "key";
 					if (destinationStorageType == SnapshotStorageType.Builds)
 					{
 						StorageTypeName = "builds";
 						StorageIdentifierName = "builds-id";
+
+						IoHash destinationKeyHash = IoHash.Compute(Encoding.UTF8.GetBytes(name));
+						CbObjectId objectId = new CbObjectId(destinationKeyHash.ToByteArray());
+						destinationId = objectId.ToString().ToLowerInvariant();
+					}
+					else
+					{
+						IoHash destinationKeyHash = IoHash.Compute(Encoding.UTF8.GetBytes(name));
+						destinationId = destinationKeyHash.ToString().ToLowerInvariant();
 					}
 
-					IoHash destinationKeyHash = IoHash.Compute(Encoding.UTF8.GetBytes(name));
 					writer.WriteValue("name", name);
 					writer.WriteValue("type", StorageTypeName);
 					writer.WriteValue("targetplatform", targetPlatform);
@@ -390,7 +399,7 @@ namespace AutomationTool.Tasks
 					}
 					writer.WriteValue("namespace", _parameters.DestinationCloudNamespace);
 					writer.WriteValue("bucket", bucketName);
-					writer.WriteValue(StorageIdentifierName, destinationKeyHash.ToString().ToLowerInvariant());
+					writer.WriteValue(StorageIdentifierName, destinationId);
 					break;
 				case SnapshotStorageType.Zen:
 					string projectName = _parameters.Project.GetFileNameWithoutAnyExtensions().ToLowerInvariant() + ".oplog";
