Fix naming convention violation in ZenExportSnapshotTask.cs

diff --git a/Engine/Source/Programs/AutomationTool/BuildGraph/Tasks/ZenExportSnapshotTask.cs b/Engine/Source/Programs/AutomationTool/BuildGraph/Tasks/ZenExportSnapshotTask.cs
--- a/Engine/Source/Programs/AutomationTool/BuildGraph/Tasks/ZenExportSnapshotTask.cs
+++ b/Engine/Source/Programs/AutomationTool/BuildGraph/Tasks/ZenExportSnapshotTask.cs
@@ -372,12 +372,12 @@ namespace AutomationTool.Tasks
 					}
 
 					string destinationId;
-					string StorageTypeName = "cloud";
-					string StorageIdentifierName = "key";
+					string storageTypeName = "cloud";
+					string storageIdentifierName = "key";
 					if (destinationStorageType == SnapshotStorageType.Builds)
 					{
-						StorageTypeName = "builds";
-						StorageIdentifierName = "builds-id";
+						storageTypeName = "builds";
+						storageIdentifierName = "builds-id";
 
 						IoHash destinationKeyHash = IoHash.Compute(Encoding.UTF8.GetBytes(name));
 						CbObjectId objectId = new CbObjectId(destinationKeyHash.ToByteArray());
@@ -390,7 +390,7 @@ namespace AutomationTool.Tasks
 					}
 
 					writer.WriteValue("name", name);
-					writer.WriteValue("type", StorageTypeName);
+					writer.WriteValue("type", storageTypeName);
 					writer.WriteValue("targetplatform", targetPlatform);
 					writer.WriteValue("host", hostName);
 					if (!String.IsNullOrEmpty(httpVersion) && !httpVersion.Equals("None", StringComparison.OrdinalIgnoreCase))
@@ -399,7 +399,7 @@ namespace AutomationTool.Tasks
 					}
 					writer.WriteValue("namespace", _parameters.DestinationCloudNamespace);
 					writer.WriteValue("bucket", bucketName);
-					writer.WriteValue(StorageIdentifierName, destinationId);
+					writer.WriteValue(storageIdentifierName, destinationId);
 					break;
 				case SnapshotStorageType.Zen:
 					string projectName = _parameters.Project.GetFileNameWithoutAnyExtensions().ToLowerInvariant() + ".oplog";
@@ -572,14 +572,14 @@ namespace AutomationTool.Tasks
 
 					string bucketName = GetCloudBucketName();
 
-					string StorageTypeName = "cloud";
-					string StorageIdentifierName = "key";
+					string storageTypeName = "cloud";
+					string storageIdentifierName = "key";
 					if (destinationStorageType == SnapshotStorageType.Builds)
 					{
-						StorageTypeName = "builds";
-						StorageIdentifierName = "builds-id";
+						storageTypeName = "builds";
+						storageIdentifierName = "builds-id";
 					}
-					oplogExportCommandline.AppendFormat(" --{0} {1} --namespace {2} --bucket {3}", StorageTypeName, _parameters.DestinationCloudHost, _parameters.DestinationCloudNamespace, bucketName);
+					oplogExportCommandline.AppendFormat(" --{0} {1} --namespace {2} --bucket {3}", storageTypeName, _parameters.DestinationCloudHost, _parameters.DestinationCloudNamespace, bucketName);
 
 					if (!String.IsNullOrEmpty(_parameters.DestinationCloudHttpVersion))
 					{
@@ -635,7 +635,7 @@ namespace AutomationTool.Tasks
 							destinationId = destinationKeyHash.ToString().ToLowerInvariant();
 						}
 
-						exportSingleSourceCommandline.AppendFormat(" {0} --{1} {2} {3} {4} {5}", hostUrlArg, StorageIdentifierName, destinationId, baseKeyArg, exportSource._projectId, exportSource._oplogId);
+						exportSingleSourceCommandline.AppendFormat(" {0} --{1} {2} {3} {4} {5}", hostUrlArg, storageIdentifierName, destinationId, baseKeyArg, exportSource._projectId, exportSource._oplogId);
 						if (_parameters.SkipExport || TryExportOplogCommand(zenExe.FullName, exportSingleSourceCommandline.ToString()))
 						{
 							successfullyExportedSources.Add(exportSource);
