Added WindowsGFN and WindowsLuna into the MCPPlatform list to address "Unable to resolve MCP platform for Label" warnings in horde.

diff --git a/Engine/Source/Programs/AutomationTool/AutomationUtils/MCPPublic.cs b/Engine/Source/Programs/AutomationTool/AutomationUtils/MCPPublic.cs
--- a/Engine/Source/Programs/AutomationTool/AutomationUtils/MCPPublic.cs
+++ b/Engine/Source/Programs/AutomationTool/AutomationUtils/MCPPublic.cs
@@ -251,6 +251,16 @@ namespace EpicGames.MCP.Automation
 		/// </summary>
 		public static MCPPlatform WindowsCN = FindOrAddByName("WindowsCN");
 
+		/// <summary>
+		/// WindowsGFN Platform.
+		/// </summary>
+		public static MCPPlatform WindowsGFN = FindOrAddByName("WindowsGFN");
+
+		/// <summary>
+		/// WindowsLuna Platform.
+		/// </summary>
+		public static MCPPlatform WindowsLuna = FindOrAddByName("WindowsLuna");
+
 		/// <summary>
 		/// IOSCN Platform.
 		/// </summary>
