Copy .collection files into development cooked editor builds so developers can view collections in the cooked editor.
Disable editing of these collections in the cooked editor.

diff --git a/Engine/Source/Programs/AutomationTool/CookedEditor/CookedEditor.Automation.cs b/Engine/Source/Programs/AutomationTool/CookedEditor/CookedEditor.Automation.cs
--- a/Engine/Source/Programs/AutomationTool/CookedEditor/CookedEditor.Automation.cs
+++ b/Engine/Source/Programs/AutomationTool/CookedEditor/CookedEditor.Automation.cs
@@ -86,6 +86,7 @@ public class ModifyStageContext
 	public bool bStagePlatformDirs = true;
 	public bool bStageUAT = false;
 	public bool bIsForExternalDistribution = false;
+	public bool bStageCollections = false;
 	public bool bStagePython = false;
 	public bool bStageTargetFiles = false;
 
@@ -118,6 +119,7 @@ public class ModifyStageContext
 		bStagePlatformDirs = ConfigHelper.GetBool("bStagePlatformDirs");
 		bStageUAT = ConfigHelper.GetBool("bStageUAT");
 		bIsForExternalDistribution = ConfigHelper.GetBool("bIsForExternalDistribution");
+		bStageCollections = ConfigHelper.GetBool("bStageCollections");
 		bStagePython = ConfigHelper.GetBool("bStagePython");
 
 		// cache some useful properties
@@ -491,6 +493,29 @@ public class MakeCookedEditor : BuildCommand
 			}		
 		}
 
+		// .collection files
+		if (!bIsCookedCooker && Context.bStageCollections)
+		{
+			DirectoryReference ProjectCollectionsDirectory = DirectoryReference.Combine(Context.ProjectDirectory, "Content", "Collections");
+			if (DirectoryReference.Exists(ProjectCollectionsDirectory))
+			{
+				Context.NonUFSFilesToStage.AddRange(DirectoryReference.EnumerateFiles(ProjectCollectionsDirectory, "*.collection", SearchOption.TopDirectoryOnly));
+			}
+
+			DirectoryReference ProjectDevelopersDirectory = DirectoryReference.Combine(Context.ProjectDirectory, "Content", "Developers");
+			if (DirectoryReference.Exists(ProjectDevelopersDirectory))
+			{
+				foreach (DirectoryReference DeveloperDirectory in DirectoryReference.EnumerateDirectories(ProjectDevelopersDirectory))
+				{
+					DirectoryReference DeveloperCollectionsDirectory = DirectoryReference.Combine(DeveloperDirectory, "Collections");
+					if (DirectoryReference.Exists(DeveloperCollectionsDirectory))
+					{
+						Context.NonUFSFilesToStage.AddRange(DirectoryReference.EnumerateFiles(DeveloperCollectionsDirectory, "*.collection", SearchOption.TopDirectoryOnly));
+					}
+				}
+			}
+		}
+
 		// Python scripts
 		if (Context.bStagePython)
 		{
