[Gauntlet][RPC] Add lock on IncomingMessageQueue to GetNewMessages to prevent exception if collection is modified during iteration

#rnx

diff --git a/Engine/Source/Programs/AutomationTool/Gauntlet/Framework/RpcFramework/Gauntlet.RpcExecutor.cs b/Engine/Source/Programs/AutomationTool/Gauntlet/Framework/RpcFramework/Gauntlet.RpcExecutor.cs
--- a/Engine/Source/Programs/AutomationTool/Gauntlet/Framework/RpcFramework/Gauntlet.RpcExecutor.cs
+++ b/Engine/Source/Programs/AutomationTool/Gauntlet/Framework/RpcFramework/Gauntlet.RpcExecutor.cs
@@ -334,9 +334,14 @@ namespace Gauntlet
 			{
 				return NewMessages;
 			}
-			foreach (GauntletIncomingRpcMessage NewMessage in IncomingMessageQueue)
+
+			// We need to lock the original IncomingMessageQueue to prevent modifications to the collection that would throw an exception during iteration in the foreach loop
+			lock (MessageQueueLock)
 			{
-				NewMessages.Add(new GauntletIncomingRpcMessage(NewMessage));
+				foreach (GauntletIncomingRpcMessage NewMessage in IncomingMessageQueue)
+				{
+					NewMessages.Add(new GauntletIncomingRpcMessage(NewMessage));
+				}
 			}
 
 			if (bClearMessageQueue)
