Change snapshot export to use a generated object id instead of trying to make one based on the identifier string.

diff --git a/Engine/Source/Programs/AutomationTool/BuildGraph/Tasks/ZenExportSnapshotTask.cs b/Engine/Source/Programs/AutomationTool/BuildGraph/Tasks/ZenExportSnapshotTask.cs
--- a/Engine/Source/Programs/AutomationTool/BuildGraph/Tasks/ZenExportSnapshotTask.cs
+++ b/Engine/Source/Programs/AutomationTool/BuildGraph/Tasks/ZenExportSnapshotTask.cs
@@ -625,8 +625,7 @@ namespace AutomationTool.Tasks
 							{
 								exportSingleSourceCommandline.AppendFormat(" --builds-metadata \"{0}\"", _parameters.Metadata);
 							}
-							IoHash destinationKeyHash = IoHash.Compute(Encoding.UTF8.GetBytes(exportNames[exportIndex]));
-							CbObjectId objectId = new CbObjectId(destinationKeyHash.ToByteArray());
+							CbObjectId objectId = CbObjectId.NewObjectId();
 							destinationId = objectId.ToString().ToLowerInvariant();
 						}
 						else
