AutomatedPerfTesting - Updates to setup script and templates
* Changing relative path to use Unreal.RootDirectory for project path
+ Adding .csproj to ensure APT test nodes are included for project
+ Adding local testing batch and shell scripts

diff --git a/Engine/Plugins/Performance/AutomatedPerfTesting/Resources/TemplatesForHookup/AutoPerfTests.xml b/Engine/Plugins/Performance/AutomatedPerfTesting/Resources/TemplatesForHookup/AutoPerfTests.xml
--- a/Engine/Plugins/Performance/AutomatedPerfTesting/Resources/TemplatesForHookup/AutoPerfTests.xml
+++ b/Engine/Plugins/Performance/AutomatedPerfTesting/Resources/TemplatesForHookup/AutoPerfTests.xml
@@ -27,7 +27,7 @@
 <!-- Common Setting Overrides -->
 <Property Name="GauntletSettingsFile" Value="$(RootDir)/$(ProjectPath)/Build/Inc/GauntletSettings.xml" />
 
-<!-- This will declare an aggregate called BuildAndTest NaniteTest -->
+<!-- This will declare an aggregate called BuildAndTest **REPLACE_PROJECTNAME** -->
 <Include Script="$(RootDir)/Engine/Build/Graph/Tasks/BuildAndTestProject.xml" />
 
 </BuildGraph>
diff --git a/Engine/Plugins/Performance/AutomatedPerfTesting/Resources/TemplatesForHookup/PROJECT.AutoPerfTest.csproj b/Engine/Plugins/Performance/AutomatedPerfTesting/Resources/TemplatesForHookup/PROJECT.AutoPerfTest.csproj
new file mode 100644
--- /dev/null
+++ b/Engine/Plugins/Performance/AutomatedPerfTesting/Resources/TemplatesForHookup/PROJECT.AutoPerfTest.csproj
@@ -0,0 +1,31 @@
+﻿<Project Sdk="Microsoft.NET.Sdk">
+  <Import Project="**REPLACE_PROJECTNAME**.AutoPerfTest.csproj.props" Condition="Exists('**REPLACE_PROJECTNAME**.AutoPerfTest.csproj.props')" />
+  <Import Project="$(EngineDir)\Source\Programs\Shared\UnrealEngine.csproj.props" />
+  <PropertyGroup>
+    <TargetFramework>net8.0</TargetFramework>
+    <OutputType>Library</OutputType>
+    <AppendTargetFrameworkToOutputPath>false</AppendTargetFrameworkToOutputPath>
+    <GenerateTargetFrameworkAttribute>false</GenerateTargetFrameworkAttribute>
+    <Configurations>Debug;Release;Development</Configurations>
+    <RootNamespace>**REPLACE_PROJECTNAME**.AutoPerfTest</RootNamespace>
+    <AssemblyName>**REPLACE_PROJECTNAME**.AutoPerfTest</AssemblyName>
+    <OutputPath>..\..\Binaries\DotNET\AutomationTool\AutomationScripts</OutputPath>
+    <TreatWarningsAsErrors>true</TreatWarningsAsErrors>
+    <WarningsNotAsErrors>612,618</WarningsNotAsErrors>
+    <DebugType>pdbonly</DebugType>
+  </PropertyGroup>
+  <ItemGroup>
+    <Compile Include="$(EngineDir)\Plugins\Performance\AutomatedPerfTesting\Build\Scripts\AutomatedPerfTestConfig.cs" />
+    <Compile Include="$(EngineDir)\Plugins\Performance\AutomatedPerfTesting\Build\Scripts\AutomatedPerfTestNode.cs" />
+  </ItemGroup>
+  <ItemGroup>
+    <ProjectReference Include="$(EngineDir)\Source\Programs\AutomationTool\AutomationUtils\AutomationUtils.Automation.csproj" />
+    <ProjectReference Include="$(EngineDir)\Source\Programs\Shared\EpicGames.Core\EpicGames.Core.csproj" PrivateAssets="All">
+      <Private>false</Private>
+    </ProjectReference>
+    <ProjectReference Include="$(EngineDir)\Source\Programs\AutomationTool\Gauntlet\Gauntlet.Automation.csproj" PrivateAssets="All">
+      <Private>false</Private>
+    </ProjectReference>
+    <ProjectReference Include="$(EngineDir)\Source\Programs\UnrealBuildTool\UnrealBuildTool.csproj" />
+  </ItemGroup>
+</Project>
diff --git a/Engine/Plugins/Performance/AutomatedPerfTesting/Resources/TemplatesForHookup/RunLocalTests.bat b/Engine/Plugins/Performance/AutomatedPerfTesting/Resources/TemplatesForHookup/RunLocalTests.bat
new file mode 100755
--- /dev/null
+++ b/Engine/Plugins/Performance/AutomatedPerfTesting/Resources/TemplatesForHookup/RunLocalTests.bat
@@ -0,0 +1,3 @@
+call **RUN_UAT_RELATIVEPATH**\RunUAT.bat BuildGraph -Script=**REPLACE_PROJECTPATH**/Build/AutoPerfTests.xml -Target="BuildAndTest **REPLACE_PROJECTNAME**" -UseLocalBuildStorage -set:TargetPlatforms=Win64 -set:TargetConfigurations=Test -set:WithAPTSequenceTests=true -set:DoPerf=true -set:DoLLM=true -set:GenerateLocalReports=true -set:GenerateLocalReports=true -set:IterationsPerf=1 -set:IterationsLLM=1 -set:IterationsInsights=1
+pause
+
diff --git a/Engine/Plugins/Performance/AutomatedPerfTesting/Resources/TemplatesForHookup/RunLocalTests.sh b/Engine/Plugins/Performance/AutomatedPerfTesting/Resources/TemplatesForHookup/RunLocalTests.sh
new file mode 100644
--- /dev/null
+++ b/Engine/Plugins/Performance/AutomatedPerfTesting/Resources/TemplatesForHookup/RunLocalTests.sh
@@ -0,0 +1,9 @@
+if [ "$(uname)" = "Darwin" ]; then
+	**RUN_UAT_RELATIVEPATH**/RunUAT.sh BuildGraph -Script=**REPLACE_PROJECTPATH**/Build/AutoPerfTests.xml -Target="BuildAndTest **REPLACE_PROJECTNAME**" -UseLocalBuildStorage -set:TargetPlatforms=Mac -set:WithAPTSequenceTests=true -set:IsBuildMachine=true -set:TargetConfigurations=Test -set:DoPerf=true -set:DoLLM=true -set:GenerateLocalReports=true -set:IterationsPerf=1 -set:IterationsLLM=1 -set:IterationsInsights=1
+fi
+
+if [ "$(uname)" = "Linux" ]; then
+	**RUN_UAT_RELATIVEPATH**/RunUAT.sh BuildGraph -Script=**REPLACE_PROJECTPATH**/Build/AutoPerfTests.xml -Target="BuildAndTest **REPLACE_PROJECTNAME**" -UseLocalBuildStorage -set:TargetPlatforms=Linux -set:WithAPTSequenceTests=true -set:TargetConfigurations=Test -set:DoPerf=true -set:DoLLM=true -set:GenerateLocalReports=true -set:GenerateLocalReports=true -set:IterationsPerf=1 -set:IterationsLLM=1 -set:IterationsInsights=1
+fi
+
+popd
diff --git a/Engine/Source/Programs/AutomationTool/Scripts/AddAutomatedPerfTestToProject.cs b/Engine/Source/Programs/AutomationTool/Scripts/AddAutomatedPerfTestToProject.cs
--- a/Engine/Source/Programs/AutomationTool/Scripts/AddAutomatedPerfTestToProject.cs
+++ b/Engine/Source/Programs/AutomationTool/Scripts/AddAutomatedPerfTestToProject.cs
@@ -13,9 +13,11 @@ using Microsoft.Extensions.Logging;
 using Microsoft.VisualStudio.OLE.Interop;
 using System.Threading.Tasks;
 using System.Security.AccessControl;
+using Microsoft.Build.Evaluation.Context;
 
 [Help("UAT command to enable PerfTesting hook ups for a project.")]
 [Help("-Project", "Path to the project to be instrumented.")]
+[Help("-Overwrite", "If set, will overwrite any template files that already exist.")]
 class AddAutomatedPerfTestToProject : BuildCommand
 {
 	/// <summary>
@@ -43,6 +45,26 @@ class AddAutomatedPerfTestToProject : BuildCommand
 	/// </summary>
 	const string AptProjectBuildGraphTemplateFile = "AutoPerfTests.xml";
 
+	/// <summary>
+	/// Name of the template file for running tests locally on Windows
+	/// </summary>
+	const string AptLocalTestBatchFile = "RunLocalTests.bat";
+	
+	/// <summary>
+	/// Name of the template file for running tests locally on Mac/Linux
+	/// </summary>
+	const string AptLocalTestShellFile = "RunLocalTests.sh";
+	
+	/// <summary>
+	/// Name of the dummy C# project file to ensure APT classes are included in builds
+	/// </summary>
+	const string AptCSProj = "PROJECT.AutoPerfTest.csproj";
+	
+	/// <summary>
+	/// Path to the engine's build/batchfiles directory relative to the engine root
+	/// </summary>
+	const string EngineBatchFilesFromRoot = "Engine/Build/BatchFiles";
+	
 	/// <summary>
 	/// Name of the config section for APT
 	/// </summary>
@@ -90,15 +112,6 @@ class AddAutomatedPerfTestToProject : BuildCommand
 		File.Move(TempFile, ProjectFilePath);
 		return true;
 	}
-	protected bool ModifyTargetCSFiles(string Project)
-	{
-		Logger.LogInformation("Modifying target files for the game/client configuration.");
-
-		// TODO: find either ProjectName.Target.cs or ProjectNameClient.Target.cs  (which needs 'Game' to be removed) and modify accordingly.
-		// Ideally just nix this step by making logging in Shipping and LLM in Test engine defaults
-
-		return true;
-	}
 
 	protected bool ModifyProjectConfigFiles(string ProjectFilePath)
 	{
@@ -149,12 +162,23 @@ class AddAutomatedPerfTestToProject : BuildCommand
 	protected string GetProjectRelativePath(string ProjectFilePath) 
 	{
 		string ProjectPath = Path.GetDirectoryName(ProjectFilePath);
-		ProjectPath = ProjectPath.Replace(Unreal.EngineDirectory.FullName, "");
+		ProjectPath = Path.GetRelativePath(Unreal.RootDirectory.ToString(), ProjectPath);
 		ProjectPath = CommandUtils.ConvertSeparators(PathSeparator.Slash, ProjectPath);
 		return ProjectPath;
 	}
+	
+	/// <summary>
+	/// Return the path to RunUAT relative to the input file
+	/// </summary>
+	protected string GetRunUATRelativePath(string InFilePath) 
+	{
+		string InDirectoryPath = Path.GetDirectoryName(InFilePath);
+		string RunUATRelativePath = Path.GetRelativePath(InDirectoryPath, Path.Combine(Unreal.RootDirectory.ToString(), EngineBatchFilesFromRoot));
+		RunUATRelativePath = CommandUtils.ConvertSeparators(PathSeparator.Slash, RunUATRelativePath);
+		return RunUATRelativePath;
+	}
 
-	protected bool ModifyBuildGraphFiles(string ProjectFilePath)
+	protected bool ModifyBuildGraphFiles(string ProjectFilePath, bool Overwrite=false)
 	{
 		// work out the target directory
 		string ProjectDir = Path.GetDirectoryName(ProjectFilePath);
@@ -177,6 +201,10 @@ class AddAutomatedPerfTestToProject : BuildCommand
 			else
 			{
 				string SrcFile = Path.Combine(BuildGraphTemplates,AptGauntletIncludeFile);
+				if (File.Exists(DestFile) && Overwrite)
+				{
+					File.Delete(DestFile);
+				}
 				File.Copy(SrcFile,DestFile);
 			}
 		}
@@ -184,7 +212,7 @@ class AddAutomatedPerfTestToProject : BuildCommand
 		// copy example per-project BuildGraph file
 		{
 			string DestFile = Path.Combine(BuildGraphDest,AptProjectBuildGraphTemplateFile);
-			if(File.Exists(DestFile))
+			if(File.Exists(DestFile) && !Overwrite)
 			{
 				Logger.LogInformation("Project already has file {0}.",DestFile);
 			}
@@ -201,6 +229,72 @@ class AddAutomatedPerfTestToProject : BuildCommand
 		return true;
 	}
 
+	protected bool ModifyCSProj(string ProjectFilePath, bool Overwrite=false)
+	{
+		// work out the target directory
+		string ProjectDir = Path.GetDirectoryName(ProjectFilePath);
+
+		string TemplatesFolder = Path.Combine(AptPathFromRoot, AptTemplatesFromPluginRoot);
+
+		string CSProjDest = Path.Combine(ProjectDir, "Build", "Scripts");
+
+		string CSProjFileName = GetProjectName(ProjectFilePath) + ".AutoPerfTest.csproj";
+		
+		// copy example per-project csproj  file
+		Directory.CreateDirectory(CSProjDest);
+		{
+			string DestFile = Path.Combine(CSProjDest,CSProjFileName);
+			if(File.Exists(DestFile) && !Overwrite)
+			{
+				Logger.LogInformation("Project already has file {0}.",DestFile);
+			}
+			else
+			{
+				string SrcFile = Path.Combine(TemplatesFolder,AptCSProj);
+				string Template = File.ReadAllText(SrcFile);
+				Template = Template.Replace("**REPLACE_PROJECTNAME**", GetProjectName(ProjectFilePath));
+				Template = Template.Replace("**REPLACE_PROJECTPATH**",GetProjectRelativePath(ProjectFilePath));
+				File.WriteAllText(DestFile, Template);
+			}
+		}
+		return true;
+	}
+	
+	protected bool ModifyLocalTestsBatchFile(string ProjectFilePath, bool Overwrite=false)
+	{
+		// work out the target directory
+		string ProjectDir = Path.GetDirectoryName(ProjectFilePath);
+
+		string TemplatesFolder = Path.Combine(AptPathFromRoot, AptTemplatesFromPluginRoot);
+
+		string BatchFileDest = Path.Combine(ProjectDir, "Build", "BatchFiles");
+		
+		// copy example per-project csproj  file
+		Directory.CreateDirectory(BatchFileDest);
+		{
+			List<string> BatchFiles = [AptLocalTestBatchFile, AptLocalTestShellFile];
+			
+			foreach (string DestFileName in BatchFiles)
+			{
+				string DestFile = Path.Combine(BatchFileDest,DestFileName);
+				if(File.Exists(DestFile) && !Overwrite)
+				{
+					Logger.LogInformation("Project already has file {0}.",DestFile);
+				}
+				else
+				{
+					string SrcFile = Path.Combine(TemplatesFolder,DestFileName);
+					string Template = File.ReadAllText(SrcFile);
+					Template = Template.Replace("**REPLACE_PROJECTNAME**", GetProjectName(ProjectFilePath));
+					Template = Template.Replace("**REPLACE_PROJECTPATH**",GetProjectRelativePath(ProjectFilePath));
+					Template = Template.Replace("**RUN_UAT_RELATIVEPATH**",GetRunUATRelativePath(DestFile));
+					File.WriteAllText(DestFile, Template);
+				}	
+			}
+		}
+		return true;
+	}
+
 	public override void ExecuteBuild()
 	{
 		Logger.LogInformation("Adding AutomatedPerfTest hookups for a project.");
@@ -212,17 +306,15 @@ class AddAutomatedPerfTestToProject : BuildCommand
 			throw new AutomationException("Project to instrument was not specified, use -Project=<path_to_uproject>. Lookup via Default.uprojectdirs is not supported");
 		}
 
-		if (!File.Exists(Project))
-		{
-			throw new AutomationException("Specified project file '{0}' was not found", Project);
-		}
+		bool Overwrite = ParseParam("Overwrite");
 
 		Logger.LogInformation("Project {0} will be configured for AutomatedPerfTest plugin.", Project);
 		
 		EnablePlugin(Project);
 		ModifyProjectConfigFiles(Project);
-		ModifyTargetCSFiles(Project);
-		ModifyBuildGraphFiles(Project);
+		ModifyBuildGraphFiles(Project, Overwrite);
+		ModifyCSProj(Project, Overwrite);
+		ModifyLocalTestsBatchFile(Project, Overwrite);
 
 		Logger.LogInformation("Successfully added.");
 	}
