Insights Automation
- Viewer and hub test cases to be run nightly on CIS
- Automate generate traces, can run multiple clients at the same time and multiple platforms for each client as long as they have the same build platform
- Small adjustments to tests so that they can be run both non-editor and editor
#rnx

diff --git a/Engine/Build/InsightsTests.xml b/Engine/Build/InsightsTests.xml
new file mode 100644
--- /dev/null
+++ b/Engine/Build/InsightsTests.xml
@@ -0,0 +1,106 @@
+<?xml version='1.0' ?>
+<BuildGraph xmlns="http://www.epicgames.com/BuildGraph" xmlns:xsi="http://www.w3.org/2001/XMLSchema-instance" xsi:schemaLocation="http://www.epicgames.com/BuildGraph ../../../Build/Graph/Schema.xsd" >
+	<!-- Horde Options -->	
+	<EnvVar Name="UE_HORDE_JOBID"/>
+	<EnvVar Name="UE_HORDE_STEPID"/>
+
+	<Include Script="Graph/Tasks/Inc/SharedProperties.xml" />
+
+	<Option Name="Clients" DefaultValue="" Description="Comma separated list of client specifiers in format [ClientTarget;Platform1+...;ProjectDir;ProjectName;Test]. Project path can be empty." />
+	<Option Name="ViewerMode" DefaultValue="false" Description="Whether to run insights tests in viewer mode"/>
+	<Option Name="HubMode" DefaultValue="false" Description="Whether to run insights tests in hub mode"/>
+	<Option Name="GenerateTraces" DefaultValue="false" Description="Set to true to generate insights traces against a list of clients"/>
+	<Option Name="HostTests" DefaultValue="System.Insights" Description="Tests to run with the Insights app"/>
+	<Option Name="Configuration" DefaultValue="Development" Description="Development, Debug or Shipping"/>
+	<Option Name="Devices" DefaultValue="" Description="A list of devices to test for all clients and platforms, separated by ;"/>
+	<Option Name="SkipCook" DefaultValue="false" Description="Skip the cooking stage"/>
+	<Option Name="SkipBuild" DefaultValue="false" Description="Skip the build stage"/>
+	<Option Name="SkipRun" DefaultValue="false" Description="Skip running the tests"/>
+	<Option Name="GauntletDevicePool" DefaultValue="UE4" Description="Gauntlet Device Pool to use to run target platform tests"/>
+	<Option Name="TestOnWin64" DefaultValue="true" Description="Test using Win64 host" />
+	<Option Name="TestOnMac" DefaultValue="false" Description="Test using Mac host" />
+	<Option Name="TestOnLinux" DefaultValue="false" Description="Test using Linux host" />
+	
+	<Option Name="OverrideAgentWin64" DefaultValue="" Description="Override the Win64 agent to run tests on" />
+	<Option Name="OverrideAgentMac" DefaultValue="" Description="Override the Mac agent to run tests on" />
+	<Option Name="OverrideAgentLinux" DefaultValue="" Description="Override the Linux agent to run tests on" />
+
+	<Property Name="TestNodes" Value=""/>
+	<Property Name="HostAgents" Value="Win64;Linux;Mac" />
+
+	<Property Name="ExtraArgs" Value=""/>
+	<Property Name="ExtraArgs" Value="$(ExtraArgs) -devices=$(Devices)" If="'$(Devices)' != ''"/>
+	<Property Name="ExtraArgs" Value="$(ExtraArgs) -deviceurl=$(GauntletDeviceService) -devicepool=$(GauntletDevicePool)"/>
+
+	<ForEach Name="AgentType" Values="$(HostAgents)" Separator=";">
+		<Property Name="BuildPlatform" Value=""/>
+		<Property Name="BuildPlatform" Value="$(AgentType)"/>
+		<Property Name="OverrideAgent" Value="$(OverrideAgent$(BuildPlatform))" />
+		<Property Name="AgentType" Value="$(OverrideAgent)" If="'$(OverrideAgent)' != ''"/>
+		<Agent Name="Insights Automation Prerequsites - $(BuildPlatform)" Type="$(AgentType)" If="$(TestOn$(BuildPlatform))">
+			<Node Name="Unreal Editor $(BuildPlatform)">
+				<Compile Target="UnrealEditor" Platform="$(BuildPlatform)" Configuration="$(Configuration)"  Arguments="-UBA"/>
+			</Node>
+			<Node Name="Unreal Insights $(BuildPlatform)">
+				<Compile Target="UnrealInsights" Platform="$(BuildPlatform)" Configuration="$(Configuration)" Arguments="-UBA"/>
+			</Node>
+			<Node Name="Build Cook Clients $(BuildPlatform)" Requires="Unreal Insights $(BuildPlatform)" Produces="#ClientExecutables" IgnoreModified="$(RootDir)/...">
+				<ForEach Name="ClientSpec" Values="$(Clients)" Separator=",">
+					<Property Name="ClientTarget" Value=""/>
+					<Property Name="PlatformList" Value=""/>
+					<Property Name="ProjectDir" Value=""/>
+					<Property Name="ProjectName" Value=""/>
+					<Property Name="ClientTest" Value=""/>
+					<Property Name="ClientTraces" Value=""/>
+					<Regex Pattern="(.*);(.*);(.*);(.*);(.*);(.*)" Capture="ClientTarget;PlatformList;ProjectDir;ProjectName;ClientTest;ClientTraces" Input="$(ClientSpec)"/>
+					<ForEach Name="Platform" Values="$(PlatformList)" Separator="+">
+						<Property Name="CookOption" Value=""/>
+						<Property Name="CookOption" Value="-Cook" If="!$(SkipCook)"/>
+						<Property Name="CookOption" Value="-SkipCook" If="$(SkipCook)"/>
+						<Property Name="BuildOption" Value=""/>
+						<Property Name="BuildOption" Value="-Build" If="!$(SkipBuild)"/>
+						<Property Name="BuildOption" Value="-SkipBuild" If="$(SkipBuild)"/>
+						<Command Name="BuildCookRun" Arguments="$(BuildOption) $(CookOption) -Stage -Pak -NoSubmit -NoCodeSign -notools -project=$(RootDir)/$(ProjectDir)/$(ProjectName).uproject -configuration=$(Configuration) -platform=$(Platform) -UBA"/>
+					</ForEach>
+					<Copy From="$(RootDir)/$(ProjectDir)/Saved/StagedBuilds/..." To="$(RootDir)/Client_Binaries/$(ProjectDir)" Tag="#ClientExecutables" ErrorIfNotFound="true"/>
+				</ForEach>
+			</Node>
+		</Agent>
+
+		<!-- UnrealInsights can only be launched on Test-* typed agents. -->
+		<Property Name="TestAgentType" Value="TestWin64" If="'$(AgentType)' == 'Win64'"/>
+		<Property Name="TestAgentType" Value="TestLinux" If="'$(AgentType)' == 'Linux'"/>
+		<Property Name="TestAgentType" Value="TestMac" If="'$(AgentType)' == 'Mac'"/>
+
+		<Agent Name="Insights Automation $(BuildPlatform)" Type="$(TestAgentType)" If="$(TestOn$(BuildPlatform))">
+			<Node Name="Insights Generate Traces $(BuildPlatform)" Requires="Unreal Insights $(BuildPlatform);Build Cook Clients $(BuildPlatform);#ClientExecutables" IgnoreModified="$(RootDir)/...">
+				<ForEach Name="ClientSpec" Values="$(Clients)" Separator=",">
+					<Property Name="ClientTarget" Value=""/>
+					<Property Name="PlatformList" Value=""/>
+					<Property Name="ProjectDir" Value=""/>
+					<Property Name="ProjectName" Value=""/>
+					<Property Name="ClientTest" Value=""/>
+					<Property Name="ClientTraces" Value=""/>
+					<Regex Pattern="(.*);(.*);(.*);(.*);(.*);(.*)" Capture="ClientTarget;PlatformList;ProjectDir;ProjectName;ClientTest;ClientTraces" Input="$(ClientSpec)"/>
+					<Copy From="$(RootDir)/Client_Binaries/$(ProjectDir)/..." To="$(RootDir)/$(ProjectDir)/Saved/StagedBuilds" ErrorIfNotFound="true" Overwrite="True"/>
+				</ForEach>
+				<Property Name="TestNodes" Value="$(TestNodes);Insights Generate Traces $(BuildPlatform)" If="$(GenerateTraces)"/>
+				<Property Name="GenerateTracesExtraArgs" Value=""/>
+				<Property Name="GenerateTracesExtraArgs" Value=" -WriteTraceFileOnly" If="'$(BuildPlatform)' == 'Linux'"/>
+				<Command Name="RunInsightsTests" Arguments="-test=InsightsTests -clients=$(Clients) -configuration=$(Configuration) $(ExtraArgs) $(GenerateTracesExtraArgs) -Verbose" If="!$(SkipRun)"/>
+			</Node>
+			<Node Name="Insights Viewer Mode $(BuildPlatform)" Requires="Unreal Editor $(BuildPlatform);Unreal Insights $(BuildPlatform)">
+				<Property Name="ViewerModeTests" Value="StartsWith:System.Insights.Analysis+StartsWith:System.Insights.Trace"/>
+				<Property Name="TestNodes" Value="$(TestNodes);Insights Viewer Mode $(BuildPlatform)" If="$(ViewerMode)"/>
+				<Command Name="RunInsightsTests" Arguments="-test=InsightsTests -mode=Viewer -HostTests=$(ViewerModeTests) -configuration=$(Configuration) $(ExtraArgs) -Verbose"/>
+			</Node>
+			<Node Name="Insights Hub Mode $(BuildPlatform)" Requires="Unreal Editor $(BuildPlatform);Unreal Insights $(BuildPlatform)">
+				<Property Name="HubModeTests" Value="StartsWith:System.Insights.Hub"/>
+				<Property Name="TestNodes" Value="$(TestNodes);Insights Hub Mode $(BuildPlatform)" If="$(HubMode)"/>
+				<Command Name="RunInsightsTests" Arguments="-test=InsightsTests -mode=Hub -HostTests=$(HubModeTests) -configuration=$(Configuration) $(ExtraArgs) -Verbose"/>
+			</Node>
+		</Agent>
+	</ForEach>
+
+	<Aggregate Name="Insights Tests" Requires="$(TestNodes)"/>
+</BuildGraph>
diff --git a/Engine/Source/Developer/TraceInsights/Private/Insights/Tests/FunctionalTests/ExportCommandsTests.cpp b/Engine/Source/Developer/TraceInsights/Private/Insights/Tests/FunctionalTests/ExportCommandsTests.cpp
--- a/Engine/Source/Developer/TraceInsights/Private/Insights/Tests/FunctionalTests/ExportCommandsTests.cpp
+++ b/Engine/Source/Developer/TraceInsights/Private/Insights/Tests/FunctionalTests/ExportCommandsTests.cpp
@@ -66,12 +66,12 @@ bool FCommandsExportThreadsDataTest::RunTest(const FString& Parameters)
 {
 	IPlatformFile& PlatformFile = FPlatformFileManager::Get().GetPlatformFile();
 	const FString SourceTracePath = FPaths::RootDir() / TEXT("EngineTest/SourceAssets/Utrace/CommandsExportTest_5.4.utrace");
-	const FString StoreTracePath = TEXT("/tmp/CommandsExportTest_5.4.utrace");
-	const FString TestResultsDirPath = TEXT("/tmp/TestResults");
+	const FString StoreTracePath = TEXT("CommandsExportTest_5.4.utrace");
+	const FString TestResultsDirPath = TEXT("TestResults");
 	const FString CmdThreadLogPath = TestResultsDirPath / TEXT("Logs/cmd_treads.log");
 	const FString ExportThreadsReportPath = TestResultsDirPath / TEXT("SingleCommand/Threads.csv");
 	const FString ExportThreadsTask = TEXT("TimingInsights.ExportThreads ") + ExportThreadsReportPath;
-	double Timeout = 10.0;
+	double Timeout = 30.0;
 
 	UTEST_TRUE(TEXT("Trace in project exists"), PlatformFile.FileExists(*SourceTracePath));
 	PlatformFile.CopyFile(*StoreTracePath, *SourceTracePath);
@@ -97,12 +97,12 @@ bool FCommandsExportTimersDataTest::RunTest(const FString& Parameters)
 {
 	IPlatformFile& PlatformFile = FPlatformFileManager::Get().GetPlatformFile();
 	const FString SourceTracePath = FPaths::RootDir() / TEXT("EngineTest/SourceAssets/Utrace/CommandsExportTest_5.4.utrace");
-	const FString StoreTracePath = TEXT("/tmp/CommandsExportTest_5.4.utrace");
-	const FString TestResultsDirPath = TEXT("/tmp/TestResults");
+	const FString StoreTracePath = TEXT("CommandsExportTest_5.4.utrace");
+	const FString TestResultsDirPath = TEXT("TestResults");
 	const FString CmdTimersLogPath = TestResultsDirPath / TEXT("Logs/cmd_timers.log");
 	const FString ExportTimersReportPath = TestResultsDirPath / TEXT("SingleCommand/Timers.csv");
 	const FString ExportTimersTask = TEXT("TimingInsights.ExportTimers ") + ExportTimersReportPath;
-	double Timeout = 10.0;
+	double Timeout = 30.0;
 
 	UTEST_TRUE(TEXT("Trace in project exists"), PlatformFile.FileExists(*SourceTracePath));
 	PlatformFile.CopyFile(*StoreTracePath, *SourceTracePath);
@@ -128,12 +128,12 @@ bool FCommandsExportTimingEventsDataTest::RunTest(const FString& Parameters)
 {
 	IPlatformFile& PlatformFile = FPlatformFileManager::Get().GetPlatformFile();
 	const FString SourceTracePath = FPaths::RootDir() / TEXT("EngineTest/SourceAssets/Utrace/CommandsExportTest_5.4.utrace");
-	const FString StoreTracePath = TEXT("/tmp/CommandsExportTest_5.4.utrace");
-	const FString TestResultsDirPath = TEXT("/tmp/TestResults");
+	const FString StoreTracePath = TEXT("CommandsExportTest_5.4.utrace");
+	const FString TestResultsDirPath = TEXT("TestResults");
 	const FString CmdTimingEventsLogPath = TestResultsDirPath / TEXT("Logs/cmd_timing_events.log");
 	const FString ExportTimingEventsReportPath = TestResultsDirPath / TEXT("SingleCommand/TimingEvents.txt");
 	const FString ExportTimingEventsTask = TEXT("TimingInsights.ExportTimingEvents " + ExportTimingEventsReportPath);
-	double Timeout = 10.0;
+	double Timeout = 30.0;
 
 	UTEST_TRUE(TEXT("Trace in project exists"), PlatformFile.FileExists(*SourceTracePath));
 	PlatformFile.CopyFile(*StoreTracePath, *SourceTracePath);
@@ -159,12 +159,12 @@ bool FCommandsExportFilteredTimingEventsDataTest::RunTest(const FString& Paramet
 {
 	IPlatformFile& PlatformFile = FPlatformFileManager::Get().GetPlatformFile();
 	const FString SourceTracePath = FPaths::RootDir() / TEXT("EngineTest/SourceAssets/Utrace/CommandsExportTest_5.4.utrace");
-	const FString StoreTracePath = TEXT("/tmp/CommandsExportTest_5.4.utrace");
-	const FString TestResultsDirPath = TEXT("/tmp/TestResults");
+	const FString StoreTracePath = TEXT("CommandsExportTest_5.4.utrace");
+	const FString TestResultsDirPath = TEXT("TestResults");
 	const FString CmdTimingEventsNonDefaultLogPath = TestResultsDirPath / TEXT("Logs/cmd_timing_non_default.log");
 	const FString ExportTimingEventsNonDefaultReportPath = TestResultsDirPath / TEXT("SingleCommand/TimingEventsNonDefault.csv");
 	const FString ExportTimingEventsBorderedTask = TEXT("TimingInsights.ExportTimingEvents " + ExportTimingEventsNonDefaultReportPath + " -columns=ThreadId,ThreadName,TimerId,TimerName,StartTime,EndTime,Duration,Depth -threads=GameThread -timers=* -startTime=10 -endTime=20");
-	double Timeout = 10.0;
+	double Timeout = 30.0;
 
 	UTEST_TRUE(TEXT("Trace in project exists"), PlatformFile.FileExists(*SourceTracePath));
 	PlatformFile.CopyFile(*StoreTracePath, *SourceTracePath);
@@ -212,11 +212,11 @@ bool FCommandsExportMultipleExportCommandsTest::RunTest(const FString& Parameter
 	IPlatformFile& PlatformFile = FPlatformFileManager::Get().GetPlatformFile();
 	const FString SourceTracePath = FPaths::RootDir() / TEXT("EngineTest/SourceAssets/Utrace/CommandsExportTest_5.4.utrace");
 	const FString SourceExportPath = FPaths::RootDir() / TEXT("EngineTest/SourceAssets/Rsp/export.rsp");
-	const FString StoreTracePath = TEXT("/tmp/CommandsExportTest_5.4.utrace");
-	const FString LogResultExportPath = TEXT("/tmp/export.rsp");
-	const FString TestResultsDirPath = TEXT("/tmp/TestResults");
+	const FString StoreTracePath = TEXT("CommandsExportTest_5.4.utrace");
+	const FString LogResultExportPath = TEXT("export.rsp");
+	const FString TestResultsDirPath = TEXT("TestResults");
 	const FString CmdExportLogPath = TestResultsDirPath / TEXT("Logs/cmd_export.log");
-	double Timeout = 10.0;
+	double Timeout = 30.0;
 
 	UTEST_TRUE(TEXT("Trace in project exists"), PlatformFile.FileExists(*SourceTracePath));
 	UTEST_TRUE(TEXT("Export in project exists"), PlatformFile.FileExists(*SourceExportPath));
@@ -236,21 +236,21 @@ bool FCommandsExportMultipleExportCommandsTest::RunTest(const FString& Parameter
 
 	const TArray<FStringView> ExpectedThreadsElementsRsp =
 	{
-		TEXTVIEW("/tmp/TestResults/RSPtest/CSV/Threads_rsp.csv"),
-		TEXTVIEW("/tmp/TestResults/RSPtest/TSV/Threads_rsp.tsv"),
-		TEXTVIEW("/tmp/TestResults/RSPtest/TXT/Threads_rsp.txt"),
+		TEXTVIEW("/TestResults/RSPtest/CSV/Threads_rsp.csv"),
+		TEXTVIEW("/TestResults/RSPtest/TSV/Threads_rsp.tsv"),
+		TEXTVIEW("/TestResults/RSPtest/TXT/Threads_rsp.txt"),
 	};
 	const TArray<FStringView> ExpectedTimersElementsRsp =
 	{
-		TEXTVIEW("/tmp/TestResults/RSPtest/CSV/Timers_rsp.csv"),
-		TEXTVIEW("/tmp/TestResults/RSPtest/TSV/Timers_rsp.tsv"),
-		TEXTVIEW("/tmp/TestResults/RSPtest/TXT/Timers_rsp.txt"),
+		TEXTVIEW("/TestResults/RSPtest/CSV/Timers_rsp.csv"),
+		TEXTVIEW("/TestResults/RSPtest/TSV/Timers_rsp.tsv"),
+		TEXTVIEW("/TestResults/RSPtest/TXT/Timers_rsp.txt"),
 	};
 	const TArray<FStringView> ExpectedTimingEventsRsp =
 	{
-		TEXTVIEW("/tmp/TestResults/RSPtest/CSV/TimingEvents_rsp.csv"),
-		TEXTVIEW("/tmp/TestResults/RSPtest/TSV/TimingEvents_rsp.tsv"),
-		TEXTVIEW("/tmp/TestResults/RSPtest/TXT/TimingEvents_rsp.txt"),
+		TEXTVIEW("/TestResults/RSPtest/CSV/TimingEvents_rsp.csv"),
+		TEXTVIEW("/TestResults/RSPtest/TSV/TimingEvents_rsp.tsv"),
+		TEXTVIEW("/TestResults/RSPtest/TXT/TimingEvents_rsp.txt"),
 	};
 	const TMap<FStringView, TArray<FStringView>> ExpectedElementsRspMap =
 	{
diff --git a/Engine/Source/Developer/TraceInsights/Private/Insights/Tests/FunctionalTests/MemoryInsightsTests.cpp b/Engine/Source/Developer/TraceInsights/Private/Insights/Tests/FunctionalTests/MemoryInsightsTests.cpp
--- a/Engine/Source/Developer/TraceInsights/Private/Insights/Tests/FunctionalTests/MemoryInsightsTests.cpp
+++ b/Engine/Source/Developer/TraceInsights/Private/Insights/Tests/FunctionalTests/MemoryInsightsTests.cpp
@@ -33,7 +33,39 @@ DECLARE_LOG_CATEGORY_EXTERN(MemoryInsightsTests, Log, All);
 
 #if WITH_AUTOMATION_TESTS
 
-IMPLEMENT_SIMPLE_AUTOMATION_TEST(FMemoryInsightsUploadLLMXMLReportsTraceTest, "System.Insights.Trace.Analysis.MemoryInsights.UploadMemoryInsightsLLMXMLReportsTrace", EAutomationTestFlags::ProgramContext | EAutomationTestFlags::EngineFilter)
+class FMemoryInsightsTestBase : public FAutomationTestBase
+{
+public:
+	FMemoryInsightsTestBase(const FString& InName, const bool bInComplexTask)
+		: FAutomationTestBase(InName, bInComplexTask)
+	{
+	}
+
+	virtual bool CanRunInEnvironment(const FString& TestParams, FString* OutReason, bool* OutWarn) const override
+	{
+		using namespace UE::Insights::MemoryProfiler;
+
+		FMemorySharedState* SharedState = FMemoryProfilerManager::Get()->GetSharedState();
+		if (!SharedState)
+		{
+			if (OutReason)
+			{
+				*OutReason = TEXT("ProfilerWindow should be valid. Please, run this test through Insights Session automation tab.");
+			}
+
+			if (OutWarn)
+			{
+				*OutWarn = true;
+			}
+
+			return false;
+		}
+
+		return true;
+	}
+};
+
+IMPLEMENT_CUSTOM_SIMPLE_AUTOMATION_TEST(FMemoryInsightsUploadLLMXMLReportsTraceTest, FMemoryInsightsTestBase, "System.Insights.Trace.Analysis.MemoryInsights.UploadMemoryInsightsLLMXMLReportsTrace", EAutomationTestFlags::ProgramContext | EAutomationTestFlags::EngineFilter)
 bool FMemoryInsightsUploadLLMXMLReportsTraceTest::RunTest(const FString& Parameters)
 {
 	using namespace UE::Insights::MemoryProfiler;
@@ -42,12 +74,6 @@ bool FMemoryInsightsUploadLLMXMLReportsTraceTest::RunTest(const FString& Paramet
 	const FString LLMReportTypesXMLPath = FPaths::RootDir() / TEXT("EngineTest/SourceAssets/Utrace/LLMReportTypes.xml");
 
 	FMemorySharedState* SharedState = FMemoryProfilerManager::Get()->GetSharedState();
-	if (!SharedState)
-	{
-		AddError("ProfilerWindow should be valid. Please, run this test through Insights Session automation tab");
-		return false;
-	}
-
 	SharedState->RemoveAllMemTagGraphTracks();
 	const int DefaultTracksAmount = SharedState->GetTimingView()->GetAllTracks().Num();
 	SharedState->RemoveAllMemTagGraphTracks();
@@ -219,7 +245,7 @@ bool MemoryInsightsAllocationsQueryTableTest(const FString& Parameters, const TM
 	return true;
 }
 
-IMPLEMENT_COMPLEX_AUTOMATION_TEST(FMemoryInsightsAllocationsQueryTableEditorPackageTest, "System.Insights.Trace.Analysis.MemoryInsights.AllocationsQueryTable.Editor/Package", EAutomationTestFlags::ProgramContext | EAutomationTestFlags::EngineFilter)
+IMPLEMENT_COMPLEX_AUTOMATION_TEST(FMemoryInsightsAllocationsQueryTableEditorPackageTest, "System.Insights.Trace.Analysis.MemoryInsights.AllocationsQueryTable.Editor.Package", EAutomationTestFlags::ProgramContext | EAutomationTestFlags::EngineFilter)
 bool FMemoryInsightsAllocationsQueryTableEditorPackageTest::RunTest(const FString& Parameters)
 {
 	bool bSuccess = MemoryInsightsAllocationsQueryTableTest(Parameters, AllocsTimeMarkerEditorPackageGetterMap, this);
diff --git a/Engine/Source/Developer/TraceInsights/Private/Insights/Tests/FunctionalTests/TimingProfilerTests.cpp b/Engine/Source/Developer/TraceInsights/Private/Insights/Tests/FunctionalTests/TimingProfilerTests.cpp
--- a/Engine/Source/Developer/TraceInsights/Private/Insights/Tests/FunctionalTests/TimingProfilerTests.cpp
+++ b/Engine/Source/Developer/TraceInsights/Private/Insights/Tests/FunctionalTests/TimingProfilerTests.cpp
@@ -143,22 +143,22 @@ bool EnumerateEventsToFile::RunTest(const FString& Parameters)
 		const TraceServices::ITimingProfilerTimerReader* TimerReader;
 		TimingProfilerProvider.ReadTimers([&TimerReader](const TraceServices::ITimingProfilerTimerReader& Out) { TimerReader = &Out; });
 
-		FString EnumeratedEventsStr;
+		FString FilePath = FPaths::ProjectSavedDir() + TEXT("/EnumerateEventsToFile.txt");
+		TUniquePtr<FArchive> Ar = TUniquePtr<FArchive>(IFileManager::Get().CreateFileWriter(*FilePath));
 		TimingProfilerProvider.ReadTimeline(TimelineIndex,
-			[Session, &EnumeratedEventsStr, TimerReader](const TraceServices::ITimingProfilerProvider::Timeline& Timeline)
+			[Session, &Ar, TimerReader](const TraceServices::ITimingProfilerProvider::Timeline& Timeline)
 			{
 				Timeline.EnumerateEvents(-1.0, Session->GetDurationSeconds() + 1.0,
-					[&EnumeratedEventsStr, TimerReader](double EventStartTime, double EventEndTime, uint32 EventDepth, const TraceServices::FTimingProfilerEvent& Event)
+					[&Ar, TimerReader](double EventStartTime, double EventEndTime, uint32 EventDepth, const TraceServices::FTimingProfilerEvent& Event)
 					{
 						const TraceServices::FTimingProfilerTimer* Timer = TimerReader->GetTimer(Event.TimerIndex);
 
-						EnumeratedEventsStr.Appendf(TEXT("%s %f %f %d\n"), Timer->Name, EventStartTime, EventEndTime, EventDepth);
+						(*Ar) << FString::Printf(TEXT("%s %f %f %d\n"), Timer->Name, EventStartTime, EventEndTime, EventDepth).GetCharArray();
 						return TraceServices::EEventEnumerate::Continue;
 					});
 			});
 
-		FString FilePath = FPaths::ProjectSavedDir() + TEXT("/EnumerateEventsToFile.txt");
-		FFileHelper::SaveStringToFile(EnumeratedEventsStr, *FilePath);
+		Ar->Close();
 	}
 
 	return !HasAnyErrors();
@@ -184,22 +184,22 @@ bool EnumerateScopesToFile::RunTest(const FString& Parameters)
 		const TraceServices::ITimingProfilerTimerReader* TimerReader;
 		TimingProfilerProvider.ReadTimers([&TimerReader](const TraceServices::ITimingProfilerTimerReader& Out) { TimerReader = &Out; });
 
-		FString EnumeratedEventsStr;
+		FString FilePath = FPaths::ProjectSavedDir() + TEXT("/EnumerateScopesToFile.txt");
+		TUniquePtr<FArchive> Ar = TUniquePtr<FArchive>(IFileManager::Get().CreateFileWriter(*FilePath));
 		TimingProfilerProvider.ReadTimeline(TimelineIndex,
-			[Session, &EnumeratedEventsStr, TimerReader](const TraceServices::ITimingProfilerProvider::Timeline& Timeline)
+			[Session, &Ar, TimerReader](const TraceServices::ITimingProfilerProvider::Timeline& Timeline)
 			{
 				Timeline.EnumerateEvents(-1.0, Session->GetDurationSeconds() + 1.0,
-					[&EnumeratedEventsStr, TimerReader](bool bStart, double Time, const TraceServices::FTimingProfilerEvent& Event)
+					[&Ar, TimerReader](bool bStart, double Time, const TraceServices::FTimingProfilerEvent& Event)
 					{
 						const TraceServices::FTimingProfilerTimer* Timer = TimerReader->GetTimer(Event.TimerIndex);
 
-						EnumeratedEventsStr.Appendf(TEXT("%s %d %f\n"), Timer->Name, (int32)bStart, Time);
+						(*Ar) << FString::Printf(TEXT("%s %d %f\n"), Timer->Name, (int32)bStart, Time).GetCharArray();
 						return TraceServices::EEventEnumerate::Continue;
 					});
 			});
 
-		FString FilePath = FPaths::ProjectSavedDir() + TEXT("/EnumerateScopesToFile.txt");
-		FFileHelper::SaveStringToFile(EnumeratedEventsStr, *FilePath);
+		Ar->Close();
 	}
 
 	return !HasAnyErrors();
diff --git a/Engine/Source/Developer/TraceInsightsFrontend/Private/InsightsFrontend/Tests/FunctionalTests/MemoryInsightsTests.cpp b/Engine/Source/Developer/TraceInsightsFrontend/Private/InsightsFrontend/Tests/FunctionalTests/MemoryInsightsTests.cpp
--- a/Engine/Source/Developer/TraceInsightsFrontend/Private/InsightsFrontend/Tests/FunctionalTests/MemoryInsightsTests.cpp
+++ b/Engine/Source/Developer/TraceInsightsFrontend/Private/InsightsFrontend/Tests/FunctionalTests/MemoryInsightsTests.cpp
@@ -33,7 +33,7 @@ void FAutomationDriverUnrealInsightsHubMemoryInsightsTest::Define()
 	{
 		AutomationWindow = FSlateApplication::Get().GetActiveTopLevelWindow();
 		const FString AutomationWindowName = TEXT("Automation");
-		if (AutomationWindow->GetTitle().ToString().Contains(AutomationWindowName))
+		if (AutomationWindow && AutomationWindow->GetTitle().ToString().Contains(AutomationWindowName))
 		{
 			AutomationWindow->Minimize();
 		}
@@ -104,7 +104,7 @@ void FAutomationDriverUnrealInsightsHubMemoryInsightsTest::Define()
 				return false;
 			};
 
-			if (!Driver->Wait(Until::Condition(TraceWaiter, FWaitTimeout::InSeconds(10))))
+			if (!Driver->Wait(Until::Condition(TraceWaiter, FWaitTimeout::InSeconds(30))))
 			{
 				AddError("Live trace should appear");
 				FPlatformProcess::TerminateProc(EditorHandle);
@@ -125,7 +125,7 @@ void FAutomationDriverUnrealInsightsHubMemoryInsightsTest::Define()
 			// Test live trace
 			FString TraceParameters = FString::Printf(TEXT("-InsightsTest -ABSLOG=\"%s\" -AutoQuit -ExecOnAnalysisCompleteCmd=\"Automation RunTests System.Insights.Trace.Analysis.MemoryInsights.UploadMemoryInsightsLLMXMLReportsTrace\" -OpenTraceFile=\"%s\""), *TestLogPath, *StoreTracePath);
 			UE::Insights::FMiscUtils::OpenUnrealInsights(*TraceParameters);
-			bool bLineFound = Utils.FileContainsString(TestLogPath, SuccessTestResult, 60.0f);
+			bool bLineFound = Utils.FileContainsString(TestLogPath, SuccessTestResult, 120.0f);
 			TestTrue("Test for live trace should pass", bLineFound);
 
 			IFileManager::Get().DeleteDirectory(*LogDirPath, false, true);
@@ -133,7 +133,7 @@ void FAutomationDriverUnrealInsightsHubMemoryInsightsTest::Define()
 
 			// Test stopped trace 
 			UE::Insights::FMiscUtils::OpenUnrealInsights(*TraceParameters);
-			bLineFound = Utils.FileContainsString(TestLogPath, SuccessTestResult, 60.0f);
+			bLineFound = Utils.FileContainsString(TestLogPath, SuccessTestResult, 120.0f);
 			TestTrue("Test for stopped trace should pass", bLineFound);
 
 			IFileManager::Get().DeleteDirectory(*LogDirPath, false, true);
@@ -145,8 +145,11 @@ void FAutomationDriverUnrealInsightsHubMemoryInsightsTest::Define()
 	{
 		Driver.Reset();
 		IAutomationDriverModule::Get().Disable();
-		AutomationWindow->Restore();
-		AutomationWindow.Reset();
+		if (AutomationWindow)
+		{
+			AutomationWindow->Restore();
+			AutomationWindow.Reset();
+		}
 	});
 }
 
diff --git a/Engine/Source/Developer/TraceInsightsFrontend/Private/InsightsFrontend/Tests/FunctionalTests/TraceStoreWindowTests.cpp b/Engine/Source/Developer/TraceInsightsFrontend/Private/InsightsFrontend/Tests/FunctionalTests/TraceStoreWindowTests.cpp
--- a/Engine/Source/Developer/TraceInsightsFrontend/Private/InsightsFrontend/Tests/FunctionalTests/TraceStoreWindowTests.cpp
+++ b/Engine/Source/Developer/TraceInsightsFrontend/Private/InsightsFrontend/Tests/FunctionalTests/TraceStoreWindowTests.cpp
@@ -163,8 +163,11 @@ void FAutomationDriverUnrealInsightsSessionBrowserTest::Define()
 	{
 		Driver.Reset();
 		IAutomationDriverModule::Get().Disable();
-		AutomationWindow->Restore();
-		AutomationWindow.Reset();
+		if (AutomationWindow)
+		{
+			AutomationWindow->Restore();
+			AutomationWindow.Reset();
+		}
 	});
 }
 
diff --git a/Engine/Source/Developer/TraceInsightsFrontend/Private/InsightsFrontend/Tests/FunctionalTests/UnrealTraceServerTests.cpp b/Engine/Source/Developer/TraceInsightsFrontend/Private/InsightsFrontend/Tests/FunctionalTests/UnrealTraceServerTests.cpp
--- a/Engine/Source/Developer/TraceInsightsFrontend/Private/InsightsFrontend/Tests/FunctionalTests/UnrealTraceServerTests.cpp
+++ b/Engine/Source/Developer/TraceInsightsFrontend/Private/InsightsFrontend/Tests/FunctionalTests/UnrealTraceServerTests.cpp
@@ -77,8 +77,8 @@ bool FUnrealTraceServerStartingTest::RunTest(const FString& Parameters)
 	bIsUnrealTraceServerRunning = WaitForUTSProcess(UnrealTraceServerName, Timeout, false);
 	TestFalse("UnrealTraceServer should not be in processes", bIsUnrealTraceServerRunning);
 
-	bConnectionStatus = WaitForConnectionIconStatus(TraceStoreConnection, Timeout, false);
-	TestFalse("Connection status should be false", bConnectionStatus);
+	bConnectionStatus = WaitForConnectionIconStatus(TraceStoreConnection, Timeout, true);
+	TestTrue("Connection status should be true", bConnectionStatus);
 
 	FString UEPath = FPlatformProcess::GenerateApplicationPath("UnrealEditor", EBuildConfiguration::Development);
 	FString UEParameters = TEXT("");
diff --git a/Engine/Source/Developer/TraceInsightsFrontend/Private/InsightsFrontend/Tests/TestRunner.cpp b/Engine/Source/Developer/TraceInsightsFrontend/Private/InsightsFrontend/Tests/TestRunner.cpp
--- a/Engine/Source/Developer/TraceInsightsFrontend/Private/InsightsFrontend/Tests/TestRunner.cpp
+++ b/Engine/Source/Developer/TraceInsightsFrontend/Private/InsightsFrontend/Tests/TestRunner.cpp
@@ -4,8 +4,14 @@
 
 #if !UE_BUILD_SHIPPING && !WITH_EDITOR
 
+#include "ISessionManager.h"
+#include "ISessionServicesModule.h"
+#include "Misc/App.h"
+
 #include "IAutomationControllerModule.h"
 #include "IAutomationControllerManager.h"
+#include "IAutomationWindowModule.h"
+#include "IAutomationWorkerModule.h"
 #include "UObject/UObjectGlobals.h" // for StaticExec()
 
 ////////////////////////////////////////////////////////////////////////////////////////////////////
@@ -17,14 +23,16 @@ DEFINE_LOG_CATEGORY(InsightsTestRunner);
 namespace UE::Insights
 {
 
+TSharedPtr<FTestRunner> FTestRunner::Instance;
+
 ////////////////////////////////////////////////////////////////////////////////////////////////////
 
-void FTestRunner::Run(const FString& InCmd)
+void FTestRunner::Run()
 {
-	FString ActualCmd = InCmd.Replace(TEXT("\""), TEXT(""));
+	FString ActualCmd = ScheduledCommand.Replace(TEXT("\""), TEXT(""));
 	if (!ActualCmd.StartsWith(TEXT("Automation RunTests")))
 	{
-		UE_LOG(InsightsTestRunner, Warning, TEXT("[InsightsTestRunner] Command %s does not start with Automation RunTests. Command will be ignored."), *InCmd);
+		UE_LOG(InsightsTestRunner, Warning, TEXT("[InsightsTestRunner] Command %s does not start with Automation RunTests. Command will be ignored."), *ScheduledCommand);
 		return;
 	}
 
@@ -34,11 +42,62 @@ void FTestRunner::Run(const FString& InCmd)
 	AutomationControllerManager->OnTestsComplete().AddLambda([this]() { bIsRunningTests = false; });
 	bIsRunningTests = true;
 
+	OnTick = FTickerDelegate::CreateSP(this, &FTestRunner::Tick);
+	OnTickHandle = FTSTicker::GetCoreTicker().AddTicker(OnTick, 1.0f);
+
 	::StaticExec(NULL, *ActualCmd);
 }
 
 ////////////////////////////////////////////////////////////////////////////////////////////////////
 
+TSharedPtr<FTestRunner> FTestRunner::CreateInstance()
+{
+	ensure(!FTestRunner::Instance.IsValid());
+	if (FTestRunner::Instance.IsValid())
+	{
+		FTestRunner::Instance.Reset();
+	}
+
+	FTestRunner::Instance = MakeShared<FTestRunner>();
+
+	return FTestRunner::Instance;
+}
+
+////////////////////////////////////////////////////////////////////////////////////////////////////
+
+TSharedPtr<FTestRunner> FTestRunner::Get()
+{
+	return FTestRunner::Instance;
+}
+
+////////////////////////////////////////////////////////////////////////////////////////////////////
+
+bool FTestRunner::Tick(float DeltaTime)
+{
+	IAutomationWorkerModule& AutomationWorkerModule = FModuleManager::LoadModuleChecked<IAutomationWorkerModule>("AutomationWorker");
+	IAutomationControllerModule& AutomationControllerModule = FModuleManager::LoadModuleChecked<IAutomationControllerModule>(TEXT("AutomationController"));
+
+	AutomationControllerModule.Tick();
+	AutomationWorkerModule.Tick();
+
+	return true;
+}
+
+////////////////////////////////////////////////////////////////////////////////////////////////////
+
+void FTestRunner::Initialize()
+{
+	ISessionServicesModule& SessionServicesModule = FModuleManager::LoadModuleChecked<ISessionServicesModule>("SessionServices");
+
+	FApp::SetSessionName(TEXT("UnrealInsights"));
+	auto SessionService = SessionServicesModule.GetSessionService();
+	SessionService->Start();
+
+	auto SessionManager = SessionServicesModule.GetSessionManager();
+
+	FModuleManager::Get().LoadModule("AutomationWorker");
+}
+
 } // namespace UE::Insights
 
 #undef LOCTEXT_NAMESPACE
diff --git a/Engine/Source/Developer/TraceInsightsFrontend/Private/InsightsFrontend/Tests/TestRunner.h b/Engine/Source/Developer/TraceInsightsFrontend/Private/InsightsFrontend/Tests/TestRunner.h
--- a/Engine/Source/Developer/TraceInsightsFrontend/Private/InsightsFrontend/Tests/TestRunner.h
+++ b/Engine/Source/Developer/TraceInsightsFrontend/Private/InsightsFrontend/Tests/TestRunner.h
@@ -6,7 +6,10 @@
 
 #if !UE_BUILD_SHIPPING && !WITH_EDITOR
 
+#include "CoreMinimal.h"
+#include "Containers/Ticker.h"
 #include "Logging/LogMacros.h"
+#include "Misc/CoreMisc.h"
 
 #include "InsightsFrontend/ITraceInsightsFrontendModule.h"
 
@@ -15,17 +18,30 @@ DECLARE_LOG_CATEGORY_EXTERN(InsightsTestRunner, Log, All);
 namespace UE::Insights
 {
 
-class FTestRunner
+class FTestRunner : public TSharedFromThis<FTestRunner>
 {
 public:
 	virtual ~FTestRunner() {}
 
-	void Run(const FString& InCmd);
+	static TSharedPtr<FTestRunner> CreateInstance();
+	static TSharedPtr<FTestRunner> Get();
 
+	void Initialize();
+	void ScheduleCommand(const FString& InCmd) { ScheduledCommand = InCmd; }
+	void Run();
+
+	bool Tick(float DeltaTime);
 private:
+	FString ScheduledCommand;
+	FTickerDelegate OnTick;
+
+	/** Handle to the registered OnTick. */
+	FTSTicker::FDelegateHandle OnTickHandle;
+
 	bool bIsRunningTests = false;
+	static TSharedPtr<FTestRunner> Instance;
 };
 
 } // namespace UE::Insights
 
-#endif //UE_BUILD_SHIPPING && !WITH_EDITOR
+#endif //!UE_BUILD_SHIPPING && !WITH_EDITOR
diff --git a/Engine/Source/Programs/AutomationTool/Gauntlet/Framework/Devices/Gauntlet.DevicePool.cs b/Engine/Source/Programs/AutomationTool/Gauntlet/Framework/Devices/Gauntlet.DevicePool.cs
--- a/Engine/Source/Programs/AutomationTool/Gauntlet/Framework/Devices/Gauntlet.DevicePool.cs
+++ b/Engine/Source/Programs/AutomationTool/Gauntlet/Framework/Devices/Gauntlet.DevicePool.cs
@@ -904,8 +904,7 @@ namespace Gauntlet
 				}
 			}
 
-			Log.Info(KnownLogEvents.Gauntlet_DeviceEvent, "Unable to reserve service devices:");
-			return false;
+			return true;
 		}
 
 		/// <summary>
diff --git a/Engine/Source/Programs/AutomationTool/Gauntlet/Platform/Mac/Gauntlet.TargetDeviceMac.cs b/Engine/Source/Programs/AutomationTool/Gauntlet/Platform/Mac/Gauntlet.TargetDeviceMac.cs
--- a/Engine/Source/Programs/AutomationTool/Gauntlet/Platform/Mac/Gauntlet.TargetDeviceMac.cs
+++ b/Engine/Source/Programs/AutomationTool/Gauntlet/Platform/Mac/Gauntlet.TargetDeviceMac.cs
@@ -50,7 +50,7 @@ namespace Gauntlet
 					// If the build already exists on this drive, we can skip installation
 					string BuildVolume = GetVolumeName(BuildDir);
 					string LocalRoot = GetVolumeName(Environment.CurrentDirectory);
-					if (BuildVolume.Equals(LocalRoot, StringComparison.OrdinalIgnoreCase))
+					if (!string.IsNullOrEmpty(BuildVolume) && BuildVolume.Equals(LocalRoot, StringComparison.OrdinalIgnoreCase))
 					{
 						Log.Info("Build exists on desired volume, skipping installation");
 						return;
@@ -241,7 +241,7 @@ namespace Gauntlet
 			string LocalRoot = GetVolumeName(Environment.CurrentDirectory);
 
 			// Must be on our volume to run
-			if (BuildVolume.Equals(LocalRoot, StringComparison.OrdinalIgnoreCase) == false)
+			if (!string.IsNullOrEmpty(BuildVolume) && BuildVolume.Equals(LocalRoot, StringComparison.OrdinalIgnoreCase) == false)
 			{
 				string SubDir = string.IsNullOrEmpty(AppConfig.Sandbox) ? AppConfig.ProjectName : AppConfig.Sandbox;
 				string InstallDir = Path.Combine(InstallRoot, SubDir, AppConfig.ProcessType.ToString());
@@ -273,7 +273,7 @@ namespace Gauntlet
 
 			string BuildVolume = GetVolumeName(TargetDir);
 			string LocalRoot = GetVolumeName(Environment.CurrentDirectory);
-			if (BuildVolume.Equals(LocalRoot, StringComparison.OrdinalIgnoreCase))
+			if (!string.IsNullOrEmpty(BuildVolume) && BuildVolume.Equals(LocalRoot, StringComparison.OrdinalIgnoreCase))
 			{
 				string SubDir = string.IsNullOrEmpty(Sandbox) ? Project : Sandbox;
 				string InstallDir = Path.Combine(InstallRoot, SubDir, ProcessType);
diff --git a/Engine/Source/Programs/AutomationTool/Gauntlet/Unreal/Base/Gauntlet.UnrealAppConfig.cs b/Engine/Source/Programs/AutomationTool/Gauntlet/Unreal/Base/Gauntlet.UnrealAppConfig.cs
--- a/Engine/Source/Programs/AutomationTool/Gauntlet/Unreal/Base/Gauntlet.UnrealAppConfig.cs
+++ b/Engine/Source/Programs/AutomationTool/Gauntlet/Unreal/Base/Gauntlet.UnrealAppConfig.cs
@@ -79,7 +79,7 @@ namespace Gauntlet
 				{
 					CommandLineParams = new GauntletCommandLine();
 				}
-				return CommandLineParams.GenerateFullCommandLine();
+				return CommandLineParams.GenerateFullCommandLine(CanAlterCommandArgs);
 			}
 			set
 			{
diff --git a/Engine/Source/Programs/AutomationTool/Gauntlet/Unreal/Base/Gauntlet.UnrealTestConfiguration.cs b/Engine/Source/Programs/AutomationTool/Gauntlet/Unreal/Base/Gauntlet.UnrealTestConfiguration.cs
--- a/Engine/Source/Programs/AutomationTool/Gauntlet/Unreal/Base/Gauntlet.UnrealTestConfiguration.cs
+++ b/Engine/Source/Programs/AutomationTool/Gauntlet/Unreal/Base/Gauntlet.UnrealTestConfiguration.cs
@@ -285,19 +285,22 @@ namespace Gauntlet
 		/// in quotes etc.
 		/// This is the function we use, as well, when accessing the Role's commandline itself.
 		/// </summary>
+		/// <param name="CanAlterCommandArgs">Some tests require that the command is not altered by any Gauntlet logic and only the provided args are used.</param>
 		/// <returns></returns>
-		public string GenerateFullCommandLine()
+		public string GenerateFullCommandLine(bool CanAlterCommandArgs = true)
 		{
 			string FinalCommandline = "";
-			if (!string.IsNullOrEmpty(Project))
+			if (CanAlterCommandArgs)
 			{
-				FinalCommandline += string.Format("\"{0}\" ", Project);
-			}
-			if (!string.IsNullOrEmpty(GameMap))
-			{
-				FinalCommandline += string.Format("\"{0}\" ", GameMap);
+				if (!string.IsNullOrEmpty(Project))
+				{
+					FinalCommandline += string.Format("\"{0}\" ", Project);
+				}
+				if (!string.IsNullOrEmpty(GameMap))
+				{
+					FinalCommandline += string.Format("\"{0}\" ", GameMap);
+				}
 			}
-
 			foreach (string Key in Params.Keys)
 			{
 				string CurrentArgument;
@@ -898,6 +901,7 @@ namespace Gauntlet
 				UnrealTargetRole.EditorGame,
 				UnrealTargetRole.Server,
 				UnrealTargetRole.EditorServer,
+				UnrealTargetRole.Host,
 				UnrealTargetRole.Editor,
 				UnrealTargetRole.CookedEditor
 			};
diff --git a/Engine/Source/Programs/AutomationTool/Gauntlet/Unreal/Utils/Gauntlet.UnrealHelpers.cs b/Engine/Source/Programs/AutomationTool/Gauntlet/Unreal/Utils/Gauntlet.UnrealHelpers.cs
--- a/Engine/Source/Programs/AutomationTool/Gauntlet/Unreal/Utils/Gauntlet.UnrealHelpers.cs
+++ b/Engine/Source/Programs/AutomationTool/Gauntlet/Unreal/Utils/Gauntlet.UnrealHelpers.cs
@@ -27,6 +27,7 @@ namespace Gauntlet
 		EditorServer,
 		Client,
 		Server,
+		Host,
 		CookedEditor,
 	};
 
@@ -80,6 +81,11 @@ namespace Gauntlet
 			return Type == UnrealTargetRole.EditorGame || Type == UnrealTargetRole.Client;
 		}
 
+		public static bool IsHost(this UnrealTargetRole Type)
+		{
+			return Type == UnrealTargetRole.Host;
+		}
+
 		public static bool IsEditor(this UnrealTargetRole Type)
 		{
 			return Type == UnrealTargetRole.Editor || Type == UnrealTargetRole.CookedEditor;
diff --git a/Engine/Source/Programs/AutomationTool/Insights/Insights.Automation.csproj b/Engine/Source/Programs/AutomationTool/Insights/Insights.Automation.csproj
new file mode 100644
--- /dev/null
+++ b/Engine/Source/Programs/AutomationTool/Insights/Insights.Automation.csproj
@@ -0,0 +1,63 @@
+<Project Sdk="Microsoft.NET.Sdk">
+  <Import Project="..\..\Shared\UnrealEngine.csproj.props" />
+  
+  <PropertyGroup>
+    <TargetFramework>net8.0</TargetFramework>
+	<Configuration Condition=" '$(Configuration)' == '' ">Development</Configuration>
+    <OutputType>Library</OutputType>
+	<RootNamespace>Insights.Automation</RootNamespace>
+    <AssemblyName>Insights.Automation</AssemblyName>
+    <GenerateAssemblyInfo>false</GenerateAssemblyInfo>
+    <GenerateTargetFrameworkAttribute>false</GenerateTargetFrameworkAttribute>
+    <OutputPath>..\..\..\..\Binaries\DotNET\AutomationTool\Insights</OutputPath>
+    <Configurations>Debug;Release;Development</Configurations>
+    <TreatWarningsAsErrors>true</TreatWarningsAsErrors>
+    <AppendTargetFrameworkToOutputPath>false</AppendTargetFrameworkToOutputPath>
+    <DebugType>pdbonly</DebugType>
+	<SatelliteResourceLanguages>en</SatelliteResourceLanguages>    <!-- remove non english resource languages -->
+  </PropertyGroup>
+  <PropertyGroup Condition=" '$(Configuration)|$(Platform)' == 'Development|AnyCPU' ">
+    <DefineConstants>$(DefineConstants);TRACE</DefineConstants>
+    <Optimize>true</Optimize>
+  </PropertyGroup>
+  <PropertyGroup Condition=" '$(Configuration)|$(Platform)' == 'Release|AnyCPU' ">
+    <DefineConstants>$(DefineConstants);TRACE</DefineConstants>
+    <Optimize>true</Optimize>
+  </PropertyGroup>
+  <PropertyGroup Condition=" '$(Configuration)|$(Platform)' == 'Debug|AnyCPU' ">
+    <DefineConstants>$(DefineConstants);DEBUG;TRACE</DefineConstants>
+    <DebugSymbols>true</DebugSymbols>
+    <DebugType>full</DebugType>
+  </PropertyGroup>
+  <ItemGroup>
+    <ProjectReference Include="..\..\Shared\EpicGames.Core\EpicGames.Core.csproj" PrivateAssets="All">
+      <Private>false</Private>
+    </ProjectReference>
+    <ProjectReference Include="..\..\UnrealBuildTool\UnrealBuildTool.csproj" PrivateAssets="All">
+      <Private>false</Private>
+    </ProjectReference>
+    <ProjectReference Include="..\..\AutomationTool\AutomationUtils\AutomationUtils.Automation.csproj" PrivateAssets="All">
+      <Private>false</Private>
+    </ProjectReference>
+    <ProjectReference Include="..\..\AutomationTool\Scripts\AutomationScripts.Automation.csproj" PrivateAssets="All">
+      <Private>false</Private>
+    </ProjectReference>
+    <ProjectReference Include="..\..\AutomationTool\Gauntlet\Gauntlet.Automation.csproj" PrivateAssets="All">
+      <Private>false</Private>
+    </ProjectReference>
+  </ItemGroup>
+  <ItemGroup>
+    <Compile Include="../../../../Platforms/*/Source/Programs/AutomationTool/Insights/**/*.cs">
+      <Link>Platform/$([System.Text.RegularExpressions.Regex]::Replace(%(Compile.RecursiveDir), [\\/].+$, ``))/%(Compile.FileName).cs</Link>
+    </Compile>
+  </ItemGroup>
+  <ItemGroup>
+    <PackageReference Include="Microsoft.CSharp" Version="4.7.0" />
+    <PackageReference Include="MySql.Data" Version="6.10.9" PrivateAssets="all" />
+    <PackageReference Include="System.Data.DataSetExtensions" Version="4.5.0" />
+    <PackageReference Include="System.Net.Http" Version="4.3.4" />
+  </ItemGroup>
+  <ItemGroup>
+    <Folder Include="Utility\" />
+  </ItemGroup>
+</Project>
diff --git a/Engine/Source/Programs/AutomationTool/Insights/RunInsightsTests.cs b/Engine/Source/Programs/AutomationTool/Insights/RunInsightsTests.cs
new file mode 100644
--- /dev/null
+++ b/Engine/Source/Programs/AutomationTool/Insights/RunInsightsTests.cs
@@ -0,0 +1,757 @@
+// Copyright Epic Games, Inc. All Rights Reserved.
+
+using System;
+using System.Collections.Generic;
+using System.Linq;
+using System.IO;
+using AutomationTool;
+using UnrealBuildTool;
+using UnrealBuildBase;
+using Gauntlet;
+using AutomationTool.DeviceReservation;
+using Microsoft.Extensions.Logging;
+using EpicGames.Core;
+using System.Net.Sockets;
+using System.Net;
+
+namespace InsightsTests
+{
+	using Log = Gauntlet.Log;
+	using LogLevel = Gauntlet.LogLevel;
+
+	public class RunInsightsTests : BuildCommand
+	{
+		public override ExitCode Execute()
+		{
+			Log.Level = LogLevel.Verbose;
+
+			Globals.Params = new Params(Params);
+
+			InsightsTestExecutorOptions ContextOptions = new InsightsTestExecutorOptions();
+			AutoParam.ApplyParamsAndDefaults(ContextOptions, Globals.Params.AllArguments);
+
+			if (ContextOptions.Mode == InsightsMode.GenerateTraces && ContextOptions.Clients.Count == 0)
+			{
+				Log.Error("Error: -clients flag is missing or no client specification was set.");
+				return ExitCode.Error_Arguments;
+			}
+
+			return RunTests(ContextOptions);
+		}
+
+		public ExitCode RunTests(InsightsTestExecutorOptions ContextOptions)
+		{
+			InsightsTestRoleContext RoleContext = new InsightsTestRoleContext();
+			InsightsAppBuildSource BuildSource = new InsightsAppBuildSource(ContextOptions.Configuration, ContextOptions.Clients, ContextOptions.OpenTraceFile, ContextOptions.WriteTraceFileOnly);
+
+			SetupDevices(ContextOptions);
+
+			InsightsTestContext TestContext = new InsightsTestContext(BuildSource, RoleContext, ContextOptions);
+
+			ITestNode NewTest = Gauntlet.Utils.TestConstructor.ConstructTest<ITestNode, InsightsTestContext>("", TestContext, new string[] { "InsightsTests" });
+
+			bool TestPassed = ExecuteTest(ContextOptions, NewTest);
+
+			DevicePool.Instance.Dispose();
+
+			return TestPassed ? ExitCode.Success : ExitCode.Error_TestFailure;
+		}
+
+		private bool ExecuteTest(InsightsTestExecutorOptions Options, ITestNode InsightsTestNode)
+		{
+			var Executor = new TestExecutor(ToString());
+
+			try
+			{
+				bool Result = Executor.ExecuteTests(Options, new List<ITestNode>() { InsightsTestNode });
+				return Result;
+			}
+			catch (Exception ex)
+			{
+				Log.Error($"{ex.Message}.{Environment.NewLine}{ex.StackTrace}");
+
+				return false;
+			}
+			finally
+			{
+				Executor.Dispose();
+
+				if (Options.Devices.Count > 0)
+				{
+					InsightsDeviceReservation.GetInstance().ReleaseDevices();
+				}
+
+				DevicePool.Instance.Dispose();
+
+				if (ParseParam("clean"))
+				{
+					Logger.LogInformation("Deleting temp dir {Arg0}", Options.TempDir);
+					DirectoryInfo TempDirInfo = new DirectoryInfo(Options.TempDir);
+					if (TempDirInfo.Exists)
+					{
+						TempDirInfo.Delete(true);
+					}
+				}
+
+				GC.Collect();
+			}
+		}
+
+		protected void SetupDevices(InsightsTestExecutorOptions Options)
+		{
+			Reservation.ReservationDetails = Options.JobDetails;
+
+			DevicePool.Instance.SetLocalOptions(Options.TempDir, Options.Parallel > 1, Options.DeviceURL);
+			DevicePool.Instance.AddLocalDevices(1);
+			DevicePool.Instance.AddVirtualDevices(2);
+
+			int TotalDevicesExpected = 0;
+			Options.Clients.ForEach(Client => TotalDevicesExpected += Client.Platforms.Count);
+
+			if (Options.Devices.Count > 0 && Options.Devices.Count != TotalDevicesExpected)
+			{
+				throw new AutomationException($"-devices should contain one device for each client and their listed platform, expected {TotalDevicesExpected} but got {Options.Devices.Count}");
+			}
+
+			if (Options.Devices.Count > 0)
+			{
+				int CurrentDeviceIdx = 0;
+				foreach (var Client in Options.Clients)
+				{
+					foreach (var Platform in Client.Platforms)
+					{
+						DevicePool.Instance.AddDevices(Platform, Options.Devices[CurrentDeviceIdx++]);
+					}
+				}
+			}
+		}
+	}
+
+	public class InsightsClientInfo
+	{
+		public string TargetName { get; set; }
+		public List<UnrealTargetPlatform> Platforms { get; set; }
+		public FileReference ProjectDir { get; set; }
+
+		public string ProjectName { get; set; }
+		public string Test { get; set; }
+        public string Traces { get; set; }
+	}
+
+	public enum InsightsMode
+	{
+		GenerateTraces,
+		Hub,
+		Viewer
+	}
+
+	public class InsightsTestExecutorOptions : TestExecutorOptions, IAutoParamNotifiable
+	{
+		public Params Params { get; protected set; }
+
+		public string TempDir;
+
+		[AutoParam("")]
+		public string DeviceURL;
+
+		[AutoParam("")]
+		public string JobDetails;
+
+		[AutoParam(0)]
+		public int Sleep;
+
+		[AutoParam("")]
+		public string LogDir;
+
+		[AutoParam("")]
+		public string HostTests;
+
+		[AutoParam("")]
+		public string OpenTraceFile;
+
+        [AutoParam(false)]
+        public bool WriteTraceFileOnly;
+
+		public int Timeout;
+
+		public InsightsMode Mode = InsightsMode.GenerateTraces;
+
+		public List<InsightsClientInfo> Clients = new List<InsightsClientInfo>();
+
+		public Type BuildSourceType { get; protected set; }
+
+		[AutoParam(UnrealTargetConfiguration.Development)]
+		public UnrealTargetConfiguration Configuration;
+
+		public List<string> Devices;
+
+		public InsightsTestExecutorOptions()
+		{
+			BuildSourceType = typeof(InsightsAppBuildSource);
+		}
+
+		public virtual void ParametersWereApplied(string[] InParams)
+		{
+			Params = new Params(InParams);
+			if (string.IsNullOrEmpty(TempDir))
+			{
+				TempDir = Globals.TempDir;
+			}
+			else
+			{
+				Globals.TempDir = TempDir;
+			}
+
+			if (string.IsNullOrEmpty(LogDir))
+			{
+				LogDir = Globals.LogDir;
+			}
+			else
+			{
+				Globals.LogDir = LogDir;
+			}
+
+			LogDir = Path.GetFullPath(LogDir);
+			TempDir = Path.GetFullPath(TempDir);
+
+			string ModeParam = Params.ParseValue("mode=", string.Empty);
+			if (!string.IsNullOrEmpty(ModeParam))
+			{
+				Mode = (InsightsMode)Enum.Parse(typeof(InsightsMode), ModeParam, true);
+			}
+
+			Timeout = Params.ParseValue("timeout=", 0);
+
+			Devices = Params.ParseValue("devices=", string.Empty).Split(';').Where(D => !string.IsNullOrEmpty(D)).ToList() ?? new List<string>();
+
+			string ClientsArg = Params.ParseValue("clients=", string.Empty);
+
+			ClientsArg.Split(',').ToList().ForEach(Spec =>
+			{
+				if (string.IsNullOrEmpty(Spec))
+					return;
+				string[] ClientSpecParts = Spec.Split(';');
+				if (ClientSpecParts.Length != 6)
+				{
+					throw new AutomationException($"Invalid client spec \"{Spec}\", should be formatted as <TargetName>;<Platform1+Platform2+...>;<ProjectDir>;<ProjectName>;<TestName>;<Traces>.");
+				}
+				Clients.Add(new InsightsClientInfo() {
+					TargetName = ClientSpecParts[0],
+					Platforms = ClientSpecParts[1].Split('+')
+									.ToList()
+									.Select(Item => UnrealTargetPlatform.Parse(Item))
+									.ToList(),
+					ProjectDir = new FileReference(ClientSpecParts[2]),
+					ProjectName = ClientSpecParts[3],
+					Test = ClientSpecParts[4],
+                    Traces = ClientSpecParts[5].Replace('+', ',')
+				});
+			});
+
+			string[] CleanArgs = Params.AllArguments
+				.Where(Arg => !Arg.StartsWith("test=", StringComparison.OrdinalIgnoreCase)
+					&& !Arg.StartsWith("device=", StringComparison.OrdinalIgnoreCase))
+				.ToArray();
+
+			Params = new Params(CleanArgs);
+		}
+	}
+
+	public class InsightsDeviceReservation
+	{
+		private static UnrealDeviceReservation UnrealDeviceReservation = new UnrealDeviceReservation();
+		private static InsightsDeviceReservation Instance;
+
+		private List<ITargetDevice> Devices;
+		private List<ITargetDevice> SelectedDevicesPrivate = new List<ITargetDevice>();
+
+		public List<ITargetDevice> ReservedDevices
+		{
+			get { return Devices ?? new List<ITargetDevice>(); }
+		}
+
+		public List<ITargetDevice> SelectedDevices
+		{
+			get { return SelectedDevicesPrivate; }
+		}
+
+		static InsightsDeviceReservation()
+		{
+			Instance = new InsightsDeviceReservation();
+		}
+
+		public static InsightsDeviceReservation GetInstance()
+		{
+			return Instance;
+		}
+
+		public bool Reserve(List<UnrealTargetPlatform> Platforms)
+		{
+			Dictionary<UnrealDeviceTargetConstraint, int> RequiredDeviceTypes = new Dictionary<UnrealDeviceTargetConstraint, int>();
+			foreach (UnrealTargetPlatform Platform in Platforms.Distinct())
+			{
+				RequiredDeviceTypes.Add(new UnrealDeviceTargetConstraint(Platform), Platforms.Count(P => P == Platform));
+			}
+
+			bool ReservedAll = UnrealDeviceReservation.TryReserveDevices(RequiredDeviceTypes, Platforms.Count);
+			Devices = UnrealDeviceReservation.ReservedDevices;
+			return ReservedAll; 
+		}
+
+		public void SelectDevice(ITargetDevice Device)
+		{
+			SelectedDevicesPrivate.Add(Device);
+		}
+
+		public void MarkProblemDevice(ITargetDevice Device)
+		{
+			UnrealDeviceReservation.MarkProblemDevice(Device);
+		}
+
+		public void ReleaseDevices()
+		{
+			UnrealDeviceReservation.ReleaseDevices();
+		}
+	}
+
+    public class InsightsHostSession : IDisposable
+    {
+        public InsightsAppBuildSource BuildSource { get; protected set; }
+        public int Sleep { get; protected set; }
+        public bool WriteTraceFileOnly { get;protected set; }
+        public InsightsMode Mode { get; protected set; }
+        public string HostTests { get; protected set; }
+        public IAppInstance Instance { get; protected set; } 
+
+		public InsightsHostSession(InsightsAppBuildSource InBuildSource, InsightsMode InMode, string InHostTests, int InSleep, bool InWriteTraceFileOnly)
+		{
+			BuildSource = InBuildSource;
+			Sleep = InSleep;
+			Mode = InMode;
+			HostTests = InHostTests;
+            WriteTraceFileOnly = InWriteTraceFileOnly;
+		}
+
+		public IAppInstance RunInsightsApp()
+		{
+			ITargetDevice Device = InsightsDeviceReservation.GetInstance().ReservedDevices.Where(D => D.IsConnected && D.Platform == HostPlatform.Platform).First();
+			InsightsDeviceReservation.GetInstance().SelectDevice(Device);
+			IAppInstall HostInstall = Device.InstallApplication(BuildSource.GetUnrealAppConfig(UnrealTargetRole.Host, Mode, "UnrealInsights", HostPlatform.Platform, Sleep, null, HostTests, null, null, WriteTraceFileOnly));
+			Instance = Device.Run(HostInstall);
+            return Instance;
+		}
+
+		public void Dispose()
+		{
+            if (Instance != null)
+            {
+                Instance.Kill();
+                InsightsDeviceReservation.GetInstance().ReleaseDevices();
+            }
+        }
+	}
+
+	public class InsightsClientSession : IDisposable
+	{
+		private static int QUERY_STATE_INTERVAL_SECONDS = 1;
+
+        Dictionary<UnrealTargetPlatform, IAppInstance> ClientApps = new Dictionary<UnrealTargetPlatform, IAppInstance>();
+
+        public InsightsAppBuildSource BuildSource { get; protected set; }
+		public InsightsClientInfo ClientInfo { get; protected set; }
+		public int Sleep { get; protected set; }
+        public bool WriteTraceFileOnly { get;protected set; }
+
+		public InsightsClientSession(InsightsAppBuildSource InBuildSource, int InSleep, bool InWriteTraceFileOnly, InsightsClientInfo InClientInfo)
+		{
+			BuildSource = InBuildSource;
+			Sleep = InSleep;
+			ClientInfo = InClientInfo;
+            WriteTraceFileOnly = InWriteTraceFileOnly;
+		}
+
+		#region IDisposable Support
+		public void Dispose()
+		{
+            foreach (var KVP in ClientApps)
+            {
+                KVP.Value.Kill();
+            }
+            InsightsDeviceReservation.GetInstance().ReleaseDevices();
+        }
+		#endregion
+
+		/// <summary>
+		/// Installs one client for each platform.
+		/// </summary>
+		public Dictionary<UnrealTargetPlatform, IAppInstance> InstallAndRunClientApps()
+		{
+			UnrealAppConfig AppConfig;
+
+			List<ITargetDevice> DevicesToInstallOn = InsightsDeviceReservation.GetInstance().ReservedDevices.Where(D => D.IsConnected).ToList();
+
+			foreach (var Platform in ClientInfo.Platforms)
+			{
+				ITargetDevice Device = DevicesToInstallOn.Where(D => D.Platform == Platform && !InsightsDeviceReservation.GetInstance().SelectedDevices.Contains(D)).FirstOrDefault();
+                // retry reservation if needed (e.g. when device connection was lost)
+                if (Device == null)
+                {
+                    InsightsDeviceReservation.GetInstance().Reserve(new List<UnrealTargetPlatform>() { Platform });
+                    Device = InsightsDeviceReservation.GetInstance().ReservedDevices.Where(D => D.IsConnected && D.Platform == Platform && !InsightsDeviceReservation.GetInstance().SelectedDevices.Contains(D)).FirstOrDefault();
+                    if (Device == null)
+                    {
+                        throw new AutomationException($"Couldn't fetch a connected device for platform {Platform}");
+                    }
+                }
+				InsightsDeviceReservation.GetInstance().SelectDevice(Device);
+
+                string FileTracePath = WriteTraceFileOnly ? Path.Combine(Unreal.EngineDirectory.FullName, "Programs", "AutomationTool", "Saved", "Logs", $"{ClientInfo.TargetName}_{Platform}.utrace") : null;
+                AppConfig = BuildSource.GetUnrealAppConfig(UnrealTargetRole.Client, InsightsMode.GenerateTraces, ClientInfo.TargetName, Platform, Sleep, ClientInfo.ProjectDir, ClientInfo.Test, ClientInfo.Traces, FileTracePath);
+
+				if (AppConfig.FullClean)
+				{
+					Device.FullClean();
+				}
+
+				// Install the build onto the device
+				if (AppConfig.SkipInstall)
+				{
+					Log.Info("Skipping install due to SkipInstall");
+				}
+				else
+				{
+					DateTimeStopwatch Stopwatch = DateTimeStopwatch.Start();
+
+					try
+					{
+						Device.InstallBuild(AppConfig);
+						Log.Info("Installation completed in {InstallTime}", GetInstallTime(Stopwatch.ElapsedTime));
+					}
+					catch (Exception Ex)
+					{
+						Log.Info("Failed to install insights tests app onto device {Device}: {Exception}", Device, Ex.ToString());
+
+						InsightsDeviceReservation.GetInstance().MarkProblemDevice(Device);
+						InsightsDeviceReservation.GetInstance().ReleaseDevices();
+					}
+				}
+
+                IAppInstall Install = null;
+                IAppInstance Instance = null;
+
+				// Create the IAppInstall instance
+				try
+				{
+					Install = Device.CreateAppInstall(AppConfig);
+				}
+				catch (Exception Ex)
+				{
+					Log.Error("Failed to create IAppInstall for insights tests on device {Device}: {Exception}", Device, Ex.ToString());
+					InsightsDeviceReservation.GetInstance().MarkProblemDevice(Device);
+					InsightsDeviceReservation.GetInstance().ReleaseDevices();
+				}
+
+				// Clean/Copy files to the device
+				try
+				{
+					Device.CleanArtifacts();
+					Device.CopyAdditionalFiles(AppConfig.FilesToCopy);
+				}
+				catch (Exception Ex)
+				{
+					Log.Info("Failed to CleanArtifacts or CopyAdditionalFiles for insights tests on device {Device}: {Exception}", Device, Ex.ToString());
+					InsightsDeviceReservation.GetInstance().MarkProblemDevice(Device);
+					InsightsDeviceReservation.GetInstance().ReleaseDevices();
+				}
+
+				// Run the application
+				try
+				{
+					if (Device is IRunningStateOptions DeviceWithStateOptions)
+					{
+						// Don't wait to detect running state and query for running state every second
+						DeviceWithStateOptions.WaitForRunningState = false;
+						DeviceWithStateOptions.CachedStateRefresh = QUERY_STATE_INTERVAL_SECONDS;
+					}
+
+					Instance = Device.Run(Install);
+					ClientApps.Add(Platform, Instance);
+				}
+				catch (DeviceException DeviceEx)
+				{
+					Log.Warning("Device {0} threw an exception during launch. \nException={1}", Install.Device, DeviceEx.Message);
+					Log.Warning("Failed to start insights test on {0}. Marking as problem device. Will not retry.", Device);
+
+					if (Instance != null)
+					{
+						Instance.Kill();
+					}
+
+					InsightsDeviceReservation.GetInstance().MarkProblemDevice(Device);
+					InsightsDeviceReservation.GetInstance().ReleaseDevices();
+
+					throw new AutomationException($"Unable to start insights client app for {ClientInfo.TargetName} platform {Platform}, see warnings for details.");
+				}
+			}
+
+			return ClientApps;
+		}
+
+		private string GetInstallTime(TimeSpan Time)
+		{
+			string Hours = Time.Hours > 0 ? string.Format("{0} hrs, ", Time.Hours) : string.Empty;
+			string Minutes = Time.Minutes > 0 ? string.Format("{0} mins, ", Time.Minutes) : string.Empty;
+			string Seconds = string.Format("{0} secs", Time.Seconds);
+
+			return Hours + Minutes + Seconds;
+		}
+	}
+
+	public class InsightsTestRoleContext : ICloneable
+	{
+		public UnrealTargetRole Type { get { return UnrealTargetRole.Client; } }
+		public UnrealTargetPlatform Platform;
+		public UnrealTargetConfiguration Configuration { get { return UnrealTargetConfiguration.Development; } }
+
+		public object Clone()
+		{
+			return this.MemberwiseClone();
+		}
+
+		public override string ToString()
+		{
+			string Description = string.Format("{0} {1} {2}", Platform, Configuration, Type);
+			return Description;
+		}
+	};
+
+	public class InsightsTestContext : ITestContext, ICloneable
+	{
+		public InsightsAppBuildSource BuildInfo { get; private set; }
+
+		public string WorkerJobID;
+
+		public InsightsTestExecutorOptions Options { get; set; }
+
+		public Params TestParams { get; set; }
+
+		public InsightsTestRoleContext RoleContext { get; set; }
+
+		public UnrealDeviceTargetConstraint Constraint;
+
+		public int PerTestTimeout { get; private set; }
+
+		public InsightsTestContext(InsightsAppBuildSource InBuildInfo, InsightsTestRoleContext InRoleContext, InsightsTestExecutorOptions InOptions, int InPerTestTimeout = 0)
+		{
+			BuildInfo = InBuildInfo;
+			Options = InOptions;
+			TestParams = new Params(new string[0]);
+			RoleContext = InRoleContext;
+			PerTestTimeout = InPerTestTimeout;
+		}
+
+		public object Clone()
+		{
+			InsightsTestContext Copy = (InsightsTestContext)MemberwiseClone();
+			Copy.RoleContext = (InsightsTestRoleContext)RoleContext.Clone();
+			return Copy;
+		}
+
+		public override string ToString()
+		{
+			string Description = string.Format("{0}", RoleContext);
+			if (WorkerJobID != null)
+			{
+				Description += " " + WorkerJobID;
+			}
+			return Description;
+		}
+	}
+
+	/// <summary>
+	/// Discovers builds for multiple clients and platforms
+	/// </summary>
+	public class InsightsAppBuildSource : IBuildSource
+	{
+		public string OpenTraceFile { get; protected set; }
+        public bool WriteTraceFileOnly { get; protected set; }
+
+        private IFolderBuildSource InsightsClientsBuildFactory;
+        public UnrealTargetConfiguration Configuration { get; protected set; }
+		public Dictionary<string, List<IBuild>> DiscoveredBuilds { get; protected set; }
+
+		public InsightsAppBuildSource(UnrealTargetConfiguration InConfiguration, List<InsightsClientInfo> InClients, string InOpenTraceFile, bool InWriteTraceFileOnly)
+		{
+			Configuration = InConfiguration;
+			OpenTraceFile = InOpenTraceFile;
+            WriteTraceFileOnly = InWriteTraceFileOnly;
+			InitBuildSources(InConfiguration, InClients);
+		}
+
+		private string GetViewerModeOpenTraceFile()
+		{
+			return Path.Combine(Unreal.EngineDirectory.FullName, "Source", "Programs", "AutomationTool", "Insights", "Resources", "ViewerMode.utrace");
+		}
+
+		protected void InitBuildSources(UnrealTargetConfiguration InConfiguration, List<InsightsClientInfo> InClients)
+		{
+			DiscoveredBuilds = new Dictionary<string, List<IBuild>>();
+
+			foreach (var Client in InClients)
+			{
+				foreach(var Platform in Client.Platforms)
+				{
+					InsightsClientsBuildFactory = Gauntlet.Utils.InterfaceHelpers.FindImplementations<IFolderBuildSource>(true)
+						.Where(B => B.CanSupportPlatform(Platform))
+						.First();
+				
+					if (!DiscoveredBuilds.ContainsKey(Client.TargetName))
+					{
+						DiscoveredBuilds.Add(Client.TargetName, new List<IBuild>());
+					}
+
+                    string PlatformFolder = Platform == UnrealTargetPlatform.Win64 ? "Windows" : Platform.ToString();
+                    IBuild DiscoveredBuild = InsightsClientsBuildFactory.GetBuildsAtPath(Client.TargetName, Path.Combine(Client.ProjectDir.FullName, "Saved", "StagedBuilds", PlatformFolder)).FirstOrDefault();
+					if (DiscoveredBuild == null)
+					{
+						throw new AutomationException("No builds were discovered.");
+					}
+					else
+					{
+						DiscoveredBuilds[Client.TargetName].Add(DiscoveredBuild);
+					}
+				}
+
+            }
+
+			UnrealTargetPlatform ThisHostPlaform = HostPlatform.Platform;
+			string HostName = "UnrealInsights";
+			string Extension;
+			string BuildPath = Path.Combine(Unreal.EngineDirectory.FullName, "Binaries", ThisHostPlaform.ToString());
+			string BuildExecutable;
+			if (ThisHostPlaform != UnrealTargetPlatform.Mac)
+			{
+				Extension = ThisHostPlaform == UnrealTargetPlatform.Linux ? string.Empty : ".exe";
+				BuildExecutable = Path.Combine(BuildPath, InConfiguration == UnrealTargetConfiguration.Development ? $"{HostName}{Extension}" : $"{HostName}-{ThisHostPlaform}-{InConfiguration}{Extension}");
+			}
+			else
+			{
+				Extension = ".app";
+				BuildPath = Path.Combine(Unreal.EngineDirectory.FullName, "Binaries", ThisHostPlaform.ToString(), InConfiguration == UnrealTargetConfiguration.Development ? $"{HostName}{Extension}" : $"{HostName}-{ThisHostPlaform}-{InConfiguration}{Extension}");
+				BuildExecutable = Path.Combine(BuildPath, "Contents", "MacOS", InConfiguration == UnrealTargetConfiguration.Development ? $"{HostName}" : $"{HostName}-{ThisHostPlaform}-{InConfiguration}");
+			}
+			IBuild UnrealInsightsBuild = new StagedBuild(
+				ThisHostPlaform,
+				InConfiguration,
+				UnrealTargetRole.Host,
+				BuildPath,
+				BuildExecutable);
+			DiscoveredBuilds.Add(HostName, new List<IBuild>());
+			DiscoveredBuilds[HostName].Add(UnrealInsightsBuild);
+		}
+
+
+        public UnrealAppConfig GetUnrealAppConfig(UnrealTargetRole InRole, InsightsMode InMode, string InAppName, UnrealTargetPlatform InPlatform, int InSleep, FileReference InProjectDir, string InTestsToRun, string InTraces, string InTraceFile = null, bool InWriteTraceFileOnly = false)
+		{
+			var Config = new UnrealAppConfig();
+			Config.Name = BuildName;
+			if (InProjectDir != null)
+			{
+				Config.ProjectName = InProjectDir.FullName;
+			}
+			Config.ProcessType = InRole;
+			Config.Platform = InPlatform;
+			Config.Configuration = UnrealTargetConfiguration.Development;
+			Config.Build = DiscoveredBuilds[InAppName].Where(Build => Build.Platform == InPlatform).First();
+			Config.Sandbox = $"InsightsTests-{InAppName}";
+			Config.FilesToCopy = new List<UnrealFileToCopy>();
+			Config.CanAlterCommandArgs = !InWriteTraceFileOnly;
+
+			if (InSleep > 0)
+			{
+				Config.CommandLineParams.AddRawCommandline(String.Format("--sleep={0}", InSleep));
+			}
+			if (InRole.IsClient())
+			{
+				// A client can have one or the other command line parameter
+				// If tracefile is used it will produce the tracefile directly without connecting to the server
+				// If it uses tracehost it will connect to that server
+                if (!string.IsNullOrEmpty(InTraceFile))
+                {
+                    Config.CommandLineParams.Add("tracefile", InTraceFile);
+                }
+                else
+                {
+                    Config.CommandLineParams.Add("tracehost", GetLocalHostIP(InPlatform));
+                }
+				Config.CommandLineParams.Add("trace", InTraces);
+
+				Config.CommandLineParams.AddRawCommandline("-NoWatchdog -stdout -FORCELOGFLUSH -CrashForUAT -nullrhi -unattended -nosplash -FullStdOutLogOutput");
+				if (!string.IsNullOrEmpty(InTestsToRun))
+				{
+					Config.CommandLineParams.AddRawCommandline($"-ExecCmds=\"Automation RunTests {InTestsToRun};Quit\"");
+				}
+			}
+			else if (InRole.IsHost() && !InWriteTraceFileOnly)
+			{
+				Config.CommandLineParams.Add("abslog", Path.Combine(Unreal.EngineDirectory.FullName, "Programs", "AutomationTool", "Saved", "Logs", "UnrealInsights.log"));
+
+				if (InPlatform == UnrealTargetPlatform.Linux)
+				{
+					Config.CommandLineParams.Add("RenderOffScreen");
+				}
+
+				if (InMode == InsightsMode.Viewer)
+				{
+					if (string.IsNullOrEmpty(InTestsToRun))
+					{
+						throw new AutomationException("Please provide tests for view mode.");
+					}
+
+					string TraceFile = string.IsNullOrEmpty(OpenTraceFile) ? GetViewerModeOpenTraceFile() : OpenTraceFile;
+
+					Config.CommandLineParams.Add("InsightsTest");
+					Config.CommandLineParams.Add("OpenTraceFile", $"\"{TraceFile}\"");
+					Config.CommandLineParams.AddRawCommandline($"-ExecOnAnalysisCompleteCmd=\"Automation RunTests {InTestsToRun};Quit\"");
+				}
+				else if (InMode == InsightsMode.Hub)
+				{
+					if (string.IsNullOrEmpty(InTestsToRun))
+					{
+						throw new AutomationException("Please provide tests for hub mode.");
+					}
+					Config.CommandLineParams.Add("InsightsTest");
+					Config.CommandLineParams.AddRawCommandline($"-RunAutomationTests=\"Automation RunTests {InTestsToRun};Quit\"");
+				}
+			}
+			return Config;
+		}
+
+		private string GetLocalHostIP(UnrealTargetPlatform Platform)
+		{
+			if (Platform.IsInGroup(UnrealPlatformGroup.Desktop))
+			{
+				return "127.0.0.1";
+			}
+			else
+			{
+				var Host = Dns.GetHostEntry(Dns.GetHostName());
+				foreach (var IPEntry in Host.AddressList)
+				{
+					if (IPEntry.AddressFamily == AddressFamily.InterNetwork)
+					{
+						return IPEntry.ToString();
+					}
+				}
+			}
+			throw new AutomationException("Couldn't retrieve local IPv4.");
+		}
+
+		public bool CanSupportPlatform(UnrealTargetPlatform Platform)
+		{
+			return true;
+		}
+
+		public string BuildName { get { return "Insights Host/Client"; } }
+	}
+}
diff --git a/Engine/Source/Programs/AutomationTool/Insights/Tests/Gauntlet.InsightsTestNode.cs b/Engine/Source/Programs/AutomationTool/Insights/Tests/Gauntlet.InsightsTestNode.cs
new file mode 100644
--- /dev/null
+++ b/Engine/Source/Programs/AutomationTool/Insights/Tests/Gauntlet.InsightsTestNode.cs
@@ -0,0 +1,341 @@
+// Copyright Epic Games, Inc. All Rights Reserved.
+
+using AutomationTool;
+using EpicGames.Core;
+using Gauntlet;
+using Polly;
+using System;
+using System.Collections.Generic;
+using System.IO;
+using System.Linq;
+using System.Threading;
+using UnrealBuildBase;
+using UnrealBuildTool;
+
+using Log = Gauntlet.Log;
+
+namespace InsightsTests
+{
+	public class InsightsTests : BaseTest
+	{
+		private static readonly string InitializationLogEntry = "Insights slate application initialized successfully";
+		private InsightsTestContext Context;
+
+		private List<IAppInstance> ClientInstances = new List<IAppInstance>();
+		private IAppInstance HostInstance = null;
+
+		private DateTime SessionStartTime = DateTime.MinValue;
+
+		private TestResult InsightsTestResult;
+
+		public Dictionary<string, InsightsClientSession> InsightsClientSessions { get; private set; }
+		public InsightsHostSession HostSession { get; private set; }
+
+		public string DefaultCommandLine;
+		private string ArtifactPath;
+		public override string Name { get { return "InsightsTest"; } }
+		public override float MaxDuration { protected set; get; }
+
+		int TotalClientInstances = 0;
+
+		public InsightsTests(InsightsTestContext InContext)
+		{
+			Context = InContext;
+
+			MaxDuration = 60 * 30;
+			InsightsTestResult = TestResult.Invalid;
+
+			InsightsClientSessions = new Dictionary<string, InsightsClientSession>();
+			HostSession = new InsightsHostSession(Context.BuildInfo, Context.Options.Mode, Context.Options.HostTests, Context.Options.Sleep, Context.Options.WriteTraceFileOnly);
+		}
+
+		public override bool IsReadyToStart()
+		{
+			ClearPreviousTraces();
+
+			List<UnrealTargetPlatform> TotalPlatforms = new List<UnrealTargetPlatform>() { HostPlatform.Platform };
+
+			foreach (var Client in Context.Options.Clients)
+			{
+				if (!InsightsClientSessions.ContainsKey(Client.TargetName))
+				{
+					InsightsClientSession Session = new InsightsClientSession(Context.BuildInfo, Context.Options.Sleep, Context.Options.WriteTraceFileOnly, Client);
+					TotalPlatforms.AddRange(Client.Platforms);
+					InsightsClientSessions.Add(Client.TargetName, Session);
+				}	
+			}
+			return InsightsDeviceReservation.GetInstance().Reserve(TotalPlatforms);
+		}
+
+		public override string GetTestSummary()
+		{
+			return "Insights Test";
+		}
+
+		public override TestResult GetTestResult()
+		{
+			return InsightsTestResult;
+		}
+
+		public override void SetTestResult(TestResult testResult)
+		{
+			InsightsTestResult = testResult;
+		}
+
+		public override void AddTestEvent(UnrealTestEvent InEvent)
+		{
+			if (InEvent.Summary.Equals("Insufficient devices found"))
+			{
+				Log.Error(KnownLogEvents.Gauntlet_TestEvent, "Test didn't run due to insufficient devices.");
+			}
+		}
+
+		public override bool StartTest(int Pass, int NumPasses)
+		{
+			if (Context.Options.Mode == InsightsMode.GenerateTraces && InsightsClientSessions.Count == 0)
+			{
+				throw new AutomationException("Node already has empty InsightsTestsApps, was IsReadyToStart called?");
+			}
+
+			if (Context.Options.Mode != InsightsMode.GenerateTraces && Context.Options.WriteTraceFileOnly)
+			{
+				throw new AutomationException("Viewer and Hub mode can't run with -WriteTraceFileOnly");
+			}
+
+			ArtifactPath = Path.Join(Context.Options.LogDir, "");
+			Log.Info("InsightsTestNode.StartTest Creating artifacts path at {0}", ArtifactPath);
+			Directory.CreateDirectory(ArtifactPath);
+
+			if (!Context.Options.WriteTraceFileOnly)
+			{
+				HostInstance = HostSession.RunInsightsApp();
+
+				if (HostInstance == null || HostInstance.HasExited || HostInstance.WasKilled)
+				{
+					throw new AutomationException("Host didn't launch properly, insights test cannot proceed.");
+				}
+
+				if (!AwaitInsightsHostSlateInitialization())
+				{
+					HostInstance.Kill();
+					throw new AutomationException("Waited for Insights host initialization signal but didn't receive any, quitting.");
+				}
+				Log.Info("Started Insights host...");
+			}
+
+			foreach (var InsightsTestsApps in InsightsClientSessions.Values)
+			{
+				ClientInstances.AddRange(InsightsTestsApps.InstallAndRunClientApps().Values);
+				
+				foreach (var ClientInstance in ClientInstances)
+				{
+					IDeviceUsageReporter.RecordComment(ClientInstance.Device.Name, (UnrealTargetPlatform)ClientInstance.Device.Platform, IDeviceUsageReporter.EventType.Device, Context.Options.JobDetails);
+					IDeviceUsageReporter.RecordComment(ClientInstance.Device.Name, (UnrealTargetPlatform)ClientInstance.Device.Platform, IDeviceUsageReporter.EventType.Test, this.GetType().Name);
+				}
+
+				if (SessionStartTime == DateTime.MinValue)
+				{
+					SessionStartTime = DateTime.Now;
+				}
+			}
+
+			TotalClientInstances = ClientInstances.Count;
+			
+			MarkTestStarted();
+
+			return true;
+		}
+
+		private bool AwaitInsightsHostSlateInitialization()
+		{
+			const int MaxRetries = 30;
+			for (int i = 0; i < MaxRetries; i++)
+			{
+				if (HostInstance
+					.GetLogBufferReader()
+					.EnumerateNextLines()
+					.Any(L => L.Contains(InitializationLogEntry)))
+				{
+					return true;
+				}
+				Thread.Sleep(2000);
+			}
+			return false;
+		}
+
+		int TotalClientFailures = 0;
+		int TotalClientTimeouts = 0;
+
+		public override void TickTest()
+		{
+			ParseAutomationEvents();
+
+			ClientInstances.ForEach(Client =>
+			{
+				if (Client.HasExited)
+				{
+					Log.Info($"Client exited with exit code {Client.ExitCode}");
+					if (Client.WasKilled || Client.ExitCode != 0)
+					{
+						TotalClientFailures++;
+					}
+				}
+				else if (CheckForTimeout())
+				{
+					Log.Error("One or more clients has timed out.");
+					Client.Kill();
+					TotalClientTimeouts++;
+				}
+			});
+
+			ClientInstances.RemoveAll(Client => Client.HasExited);
+
+			if (Context.Options.Mode == InsightsMode.GenerateTraces && ClientInstances.Count == 0)
+			{
+				if (HostInstance != null)
+				{
+					Log.Info("All clients exited, closing host...");
+					HostInstance.Kill();
+				}
+				InsightsTestResult = TotalClientFailures > 0 || TotalClientTimeouts > 0 ? TestResult.Failed : TestResult.Passed;
+				MarkTestComplete();
+			}
+			else if (Context.Options.Mode != InsightsMode.GenerateTraces && HostInstance.HasExited)
+			{
+				InsightsTestResult = HostInstance.ExitCode == 0 ? TestResult.Passed : TestResult.Failed;
+				MarkTestComplete();
+			}
+		}
+
+		public void ClearPreviousTraces()
+		{
+			Directory.GetFiles(
+				Path.Combine(Unreal.EngineDirectory.FullName, "Programs", "AutomationTool", "Saved", "Logs"),
+				"*.utrace")
+				.ToList()
+				.ForEach(T => { File.Delete(T); });
+		}
+
+		public void CollectClientTraces()
+		{
+			List<FileInfo> AllTraces = new List<FileInfo>();
+
+			if (!Context.Options.WriteTraceFileOnly)
+			{
+				string StorePathRoot = HostPlatform.Platform == UnrealTargetPlatform.Win64
+					? Path.Combine(Environment.GetFolderPath(Environment.SpecialFolder.LocalApplicationData), "UnrealEngine", "Common", "UnrealTrace", "Store")
+					: Path.Combine(Environment.GetFolderPath(Environment.SpecialFolder.UserProfile), "UnrealEngine", "UnrealTrace", "Store");
+				var AllDirs = new DirectoryInfo(StorePathRoot).GetDirectories();
+				if (AllDirs.Length > 0)
+				{
+					var TodayStringFormatted = DateTime.Today.ToString("yyyyMMdd");
+					var TomorrowStringFormatted = DateTime.Today.AddDays(1).ToString("yyyyMMdd");
+					var StorePath = AllDirs.OrderByDescending(D => D.LastAccessTime).First();
+					AllTraces = StorePath.GetFiles("*.utrace")
+						.Where(F => F.Name.StartsWith(TodayStringFormatted) || F.Name.StartsWith(TomorrowStringFormatted))
+						.OrderByDescending(F => F.LastAccessTime).Take(TotalClientInstances).ToList();
+				}
+			}
+
+			if (Context.Options.Mode == InsightsMode.GenerateTraces && !Context.Options.WriteTraceFileOnly && AllTraces.Count == 0)
+			{
+				Log.Error("Could not find any tracefiles generated during this session.");
+			}
+
+			string TargetPath = Path.Combine(Unreal.EngineDirectory.FullName, "Programs", "AutomationTool", "Saved", "Logs");
+			foreach (FileInfo Trace in AllTraces)
+			{
+				Log.Info("Copying trace file to Logs folder: " + Trace.FullName);
+				File.Copy(Trace.FullName, Path.Combine(TargetPath, Trace.Name));
+			}
+		}
+		
+
+		private int LastStdoutSeekPos;
+		private string[] CurrentProcessedLines;
+		private void ParseAutomationEvents()
+		{
+			if (HostInstance != null)
+			{
+				if (LastStdoutSeekPos < HostInstance.StdOut.Length)
+				{
+					CurrentProcessedLines = HostInstance.StdOut
+						.Substring(LastStdoutSeekPos)
+						.Split("\n")
+						.Where(Line => !string.IsNullOrWhiteSpace(Line) && Line.Contains("LogAutomationController"))
+						.ToArray();
+					LastStdoutSeekPos = HostInstance.StdOut.Length - 1;
+
+					foreach (string OutputLine in CurrentProcessedLines)
+					{
+						Console.WriteLine(OutputLine.TrimEnd(Environment.NewLine.ToCharArray()));
+					}
+				}
+			}
+		}
+
+		private bool CheckForTimeout()
+		{
+			return false;
+		}
+
+		public override void StopTest(StopReason InReason)
+		{
+			try
+			{
+				base.StopTest(InReason);
+
+				if (ClientInstances.Count > 0 && !ClientInstances.Any(Client => Client.HasExited))
+				{
+					ClientInstances.Where(Client => Client.HasExited).ToList().ForEach(Client => Client.Kill());
+				}
+
+				string ExitReason = "";
+
+				if (ClientInstances.Any(Client => Client.WasKilled))
+				{
+					if (InReason == StopReason.MaxDuration || InsightsTestResult == TestResult.TimedOut)
+					{
+						InsightsTestResult = TestResult.TimedOut;
+						ExitReason = "Timed Out";
+					}
+					else
+					{
+						InsightsTestResult = TestResult.Failed;
+						ExitReason = $"Process was killed by Gauntlet with reason {InReason}.";
+					}
+				}
+				else if (ClientInstances.Any(Client => Client.ExitCode != 0))
+				{
+					InsightsTestResult = TestResult.Failed;
+					ExitReason = $"A client exited with non-zero exit code";
+				}
+				Log.Info(ExitReason);
+			}
+			catch
+			{
+				throw;
+			}
+			finally
+			{
+				if (HostInstance != null && !HostInstance.HasExited)
+				{
+					HostInstance.Kill();
+				}
+				CollectClientTraces();
+			}
+		}
+
+		public override void CleanupTest()
+		{
+			foreach (var InsightsTestsApps in InsightsClientSessions.Values)
+			{
+				InsightsTestsApps.Dispose();
+			}
+			if (HostInstance != null)
+			{
+				HostSession.Dispose();
+			}
+		}
+	}
+}
diff --git a/Engine/Source/Programs/AutomationTool/Scripts/BuildCookRun.Automation.cs b/Engine/Source/Programs/AutomationTool/Scripts/BuildCookRun.Automation.cs
--- a/Engine/Source/Programs/AutomationTool/Scripts/BuildCookRun.Automation.cs
+++ b/Engine/Source/Programs/AutomationTool/Scripts/BuildCookRun.Automation.cs
@@ -139,6 +139,12 @@ public class BuildCookRun : BuildCommand, IProjectParamsHelpers
 			Params.AdditionalBuildOptions += " -EnablePlugins=\"" + string.Join(",", PluginsToEnable) + "\"";
 		}
 
+		bool WithUBA = ParseParamBool("UBA");
+		if (WithUBA)
+		{
+			Params.AdditionalBuildOptions += " -UBA";
+		}	
+
 		Params.ValidateAndLog();
 		return Params;
 	}
diff --git a/Engine/Source/Programs/UnrealInsights/Private/UserInterfaceCommand.cpp b/Engine/Source/Programs/UnrealInsights/Private/UserInterfaceCommand.cpp
--- a/Engine/Source/Programs/UnrealInsights/Private/UserInterfaceCommand.cpp
+++ b/Engine/Source/Programs/UnrealInsights/Private/UserInterfaceCommand.cpp
@@ -9,6 +9,7 @@
 #include "Framework/Application/SlateApplication.h"
 #include "Framework/Docking/LayoutService.h"
 #include "Framework/Docking/TabManager.h"
+#include "HAL/ExceptionHandling.h"
 #include "HAL/PlatformApplicationMisc.h"
 #include "HAL/PlatformProcess.h"
 #include "Interfaces/IPluginManager.h"
@@ -220,6 +221,7 @@ void FUserInterfaceCommand::InitializeSlateApplication(bool bFrontendMode, const
 		const bool bInitializeTesting = FParse::Param(FCommandLine::Get(), TEXT("InsightsTest"));
 		if (bInitializeTesting)
 		{
+			GUseCrashReportClient = false;
 			const bool bInitAutomationModules = true;
 			TraceInsightsModule.InitializeTesting(bInitAutomationModules, bAutoQuit);
 		}
@@ -300,6 +302,7 @@ void FUserInterfaceCommand::InitializeSlateApplication(bool bFrontendMode, const
 			TraceInsightsFrontendModule.RunAutomationTests(AutomationTests);
 		}
 	}
+	UE_LOG(LogSlate, Display, TEXT("Insights slate application initialized successfully."));
 }
 
 ////////////////////////////////////////////////////////////////////////////////////////////////////
