Launch proxy, game servers and clients with a delay between each process to enforce determinism and provide a configurable proxy listening port.

diff --git a/Engine/Source/Programs/AutomationTool/Scripts/LaunchMultiServer.Automation.cs b/Engine/Source/Programs/AutomationTool/Scripts/LaunchMultiServer.Automation.cs
--- a/Engine/Source/Programs/AutomationTool/Scripts/LaunchMultiServer.Automation.cs
+++ b/Engine/Source/Programs/AutomationTool/Scripts/LaunchMultiServer.Automation.cs
@@ -208,6 +208,7 @@ namespace AutomationTool
 			ConfigHierarchy ProjectGameConfig = ConfigCache.ReadHierarchy(ConfigHierarchyType.Game, DirectoryReference.FromFile(ProjectPath), UnrealTargetPlatform.Win64);
 
 			const String ServerDefConfigSection = "/Script/MultiServerConfiguration.MultiServerSettings";
+			const String ProxyConfigSection = "/Script/MultiServerConfiguration.Proxy";
 
 			ProjectGameConfig.TryGetValuesGeneric(ServerDefConfigSection, "ServerDefinitions", out MultiServerDefinition[] ConfigServerDefs);
 
@@ -255,12 +256,17 @@ namespace AutomationTool
 								Int32.TryParse(ClientCountStr, out ClientCount);
 							}
 
-							Logger.LogInformation("Starting proxy server connected to MultiServer instances");
-							StartProcessesForServerInfos(UniqueServerInfos, ProcessLaunchType.ProxyServer, "");
+							var ProxyListeningPort = 8000;
+
+							ProjectGameConfig.TryGetValue(ProxyConfigSection, "ListenPort", out ProxyListeningPort);
+
+							Logger.LogInformation("Starting proxy server on port {0} connected to MultiServer instances", ProxyListeningPort);
+							StartProcessesForServerInfos(UniqueServerInfos, ProcessLaunchType.ProxyServer, String.Format(" -port={0} ", ProxyListeningPort));
 
 							for (int ClientIndex = 0; ClientIndex < ClientCount; ClientIndex++)
 							{
-								StartProcessesForServerInfos(UniqueServerInfos, ProcessLaunchType.ProxyClient, " -SaveWinPos="+(ClientIndex + 1));
+								var ProxyClientEndpoint = String.Format("127.0.0.1:{0}", ProxyListeningPort);
+								StartProcessesForServerInfos(UniqueServerInfos, ProcessLaunchType.ProxyClient, ProxyClientEndpoint + " -SaveWinPos="+(ClientIndex + 1));
 							}
 						}
 					}
@@ -279,6 +285,8 @@ namespace AutomationTool
 
 		private void StartProcessesForServerInfos(Dictionary<string, MultiServerInstanceInfo> ServerInfos, ProcessLaunchType LaunchType, string AdditionalParameters)
 		{
+			const int PauseBetweenProcessMS = 500;
+
 			var Params = new ProjectParams
 				(
 					Command: this,
@@ -358,6 +366,9 @@ namespace AutomationTool
 							ProcessManager.RemoveProcess(NewProcess);
 						}
 
+						// Pause between starting processes to enforce startup determinism.
+						System.Threading.Thread.Sleep(PauseBetweenProcessMS);
+
 						if (IsAttachingDebugger(LaunchType))
 						{
 							ProcessDebugger.DebugProcess(NewProcess.ProcessObject.Id);
@@ -368,13 +379,8 @@ namespace AutomationTool
 				}
 				else if (LaunchType == ProcessLaunchType.ProxyServer)
 				{
-					var ProxyListeningPort = 8000;
-
-					Logger.LogInformation("Starting MultiServer proxy listening on port {0}", ProxyListeningPort);
-					
 					var Args = CommonArgs;
 
-					Args += String.Format("-port={0} ", ProxyListeningPort);
 					Args += "-NetDriverOverrides=/Script/MultiServerReplication.ProxyNetDriver ";
 
 					if (ServerInfos.Count() > 0)
@@ -388,6 +394,8 @@ namespace AutomationTool
 						}
 					}
 
+					Args += AdditionalParameters;
+
 					var NewProcess = Run(App, Args, null, ERunOptions.Default | ERunOptions.NoWaitForExit | ERunOptions.NoStdOutRedirect);
 
 					if (NewProcess != null)
@@ -396,15 +404,18 @@ namespace AutomationTool
 						// Essentially forces the -NoKill command-line option behavior for these.
 						ProcessManager.RemoveProcess(NewProcess);
 					}
+
+					// Pause between starting processes to enforce startup determinism.
+					System.Threading.Thread.Sleep(PauseBetweenProcessMS);
 				}
 				else if (LaunchType == ProcessLaunchType.ProxyClient)
 				{
 					var Args = CommonArgs;
 
-					Args += "127.0.0.1:8000 " + (Params.Cook ? "" : "-game ");
+					// The server to connect to should be provided in the AdditionalParameters (e.g. 127.0.0.1:8000).
+					Args += AdditionalParameters + (Params.Cook ? "" : "-game ");
 					Args += Params.ClientCommandline + " ";
 					Args += "-windowed ";
-					Args += AdditionalParameters;
 
 					var NewProcess = Run(App, Args, null, ERunOptions.Default | ERunOptions.NoWaitForExit | ERunOptions.NoStdOutRedirect);
 
@@ -414,6 +425,9 @@ namespace AutomationTool
 						// Essentially forces the -NoKill command-line option behavior for these.
 						ProcessManager.RemoveProcess(NewProcess);
 					}
+
+					// Pause between starting processes to enforce startup determinism.
+					System.Threading.Thread.Sleep(PauseBetweenProcessMS);
 				}
 			}
 			catch
