Fix launch of iOS Blueprint projects

diff --git a/Engine/Source/Programs/AutomationTool/IOS/IOSPlatform.Automation.cs b/Engine/Source/Programs/AutomationTool/IOS/IOSPlatform.Automation.cs
--- a/Engine/Source/Programs/AutomationTool/IOS/IOSPlatform.Automation.cs
+++ b/Engine/Source/Programs/AutomationTool/IOS/IOSPlatform.Automation.cs
@@ -2027,8 +2027,24 @@ Console.WriteLine("target receipt executable path: {0}", InExecutablePath);
 				string TargetStageDir = ClientApp.Substring(0, ClientApp.IndexOf('/', Params.BaseStageDirectory.Length + 1));
 				if (AppleExports.UseModernXcode(Params.RawProjectPath))
 				{
-					// modern mode runs from the staged dir
-					PlistFile = Path.Combine(TargetStageDir, Path.GetFileNameWithoutExtension(ClientApp) + ".app", "Info.plist");
+					// Modern mode runs from the staged dir.
+					// Also, Blueprint only projects have "ClientApp" set to UnrealGame, so use the Project name instead
+					string AppName = Path.GetFileNameWithoutExtension(ClientApp);
+					if (!Params.IsCodeBasedProject)
+					{
+						AppName = Params.ShortProjectName;
+						switch (Params.ClientConfigsToBuild[0])
+						{
+							case UnrealTargetConfiguration.Debug:
+							case UnrealTargetConfiguration.DebugGame:
+							case UnrealTargetConfiguration.Test:
+							case UnrealTargetConfiguration.Shipping:
+								// build the correct full appname (Development has none). ie: MyProject-IOS-Debug
+								AppName += "-" + Params.ClientTargetPlatforms[0] + "-" + Params.ClientConfigsToBuild[0];
+								break;
+						}
+					}
+					PlistFile = Path.Combine(TargetStageDir, AppName + ".app", "Info.plist");
 				}
 				else
 				{
