Deleted old and unused BuildGraph task to upload builds to Jupiter (now Cloud DDC) using a API that doesn't exist anymore.

diff --git a/Engine/Source/Programs/AutomationTool/BuildGraph/BuildGraph.Automation.csproj b/Engine/Source/Programs/AutomationTool/BuildGraph/BuildGraph.Automation.csproj
--- a/Engine/Source/Programs/AutomationTool/BuildGraph/BuildGraph.Automation.csproj
+++ b/Engine/Source/Programs/AutomationTool/BuildGraph/BuildGraph.Automation.csproj
@@ -55,7 +55,6 @@
   <ItemGroup>
     <ProjectReference Include="..\..\Shared\EpicGames.BuildGraph\EpicGames.BuildGraph.csproj" />
     <ProjectReference Include="..\..\Shared\EpicGames.Core\EpicGames.Core.csproj" PrivateAssets="All"><Private>false</Private></ProjectReference>
-    <ProjectReference Include="..\..\Shared\EpicGames.Jupiter\EpicGames.Jupiter.csproj" PrivateAssets="All" />
     <ProjectReference Include="..\..\Shared\EpicGames.ProjectStore\EpicGames.ProjectStore.csproj" PrivateAssets="All" />
     <ProjectReference Include="..\..\UnrealBuildTool\UnrealBuildTool.csproj" PrivateAssets="All"><Private>false</Private></ProjectReference>
     <ProjectReference Include="..\AutomationUtils\AutomationUtils.Automation.csproj" PrivateAssets="All"><Private>false</Private></ProjectReference>
diff --git a/Engine/Source/Programs/AutomationTool/BuildGraph/Tasks/UploadTask.cs b/Engine/Source/Programs/AutomationTool/BuildGraph/Tasks/UploadTask.cs
deleted file mode 100644
diff --git a/Engine/Source/Programs/AutomationTool/Scripts/AutomationScripts.Automation.csproj b/Engine/Source/Programs/AutomationTool/Scripts/AutomationScripts.Automation.csproj
--- a/Engine/Source/Programs/AutomationTool/Scripts/AutomationScripts.Automation.csproj
+++ b/Engine/Source/Programs/AutomationTool/Scripts/AutomationScripts.Automation.csproj
@@ -48,7 +48,6 @@
   </ItemGroup>
   <ItemGroup>
     <ProjectReference Include="..\..\Shared\EpicGames.Core\EpicGames.Core.csproj" PrivateAssets="All"><Private>false</Private></ProjectReference>
-    <ProjectReference Include="..\..\Shared\EpicGames.Jupiter\EpicGames.Jupiter.csproj" PrivateAssets="All" />
     <ProjectReference Include="..\..\Shared\EpicGames.ProjectStore\EpicGames.ProjectStore.csproj" PrivateAssets="All" />
     <ProjectReference Include="..\..\UnrealBuildTool\UnrealBuildTool.csproj" PrivateAssets="All"><Private>false</Private></ProjectReference>
     <ProjectReference Include="..\AutomationUtils\AutomationUtils.Automation.csproj" PrivateAssets="All"><Private>false</Private></ProjectReference>
diff --git a/Engine/Source/Programs/AutomationTool/Scripts/DownloadJupiterBuild.cs b/Engine/Source/Programs/AutomationTool/Scripts/DownloadJupiterBuild.cs
deleted file mode 100644
diff --git a/Engine/Source/Programs/Shared/EpicGames.Jupiter/.editorconfig b/Engine/Source/Programs/Shared/EpicGames.Jupiter/.editorconfig
deleted file mode 100644
diff --git a/Engine/Source/Programs/Shared/EpicGames.Jupiter/EpicGames.Jupiter.csproj b/Engine/Source/Programs/Shared/EpicGames.Jupiter/EpicGames.Jupiter.csproj
deleted file mode 100644
diff --git a/Engine/Source/Programs/Shared/EpicGames.Jupiter/JupiterTree.cs b/Engine/Source/Programs/Shared/EpicGames.Jupiter/JupiterTree.cs
deleted file mode 100644
diff --git a/Engine/Source/Programs/Shared/EpicGames.Jupiter/JupiterUtils.cs b/Engine/Source/Programs/Shared/EpicGames.Jupiter/JupiterUtils.cs
deleted file mode 100644
