BuildGraph: Throw an exception when no files are found matching the wildcard for move and copy tasks. Fixes builds not stopping at failures when not using Horde.

diff --git a/Engine/Source/Programs/AutomationTool/BuildGraph/Tasks/CopyTask.cs b/Engine/Source/Programs/AutomationTool/BuildGraph/Tasks/CopyTask.cs
--- a/Engine/Source/Programs/AutomationTool/BuildGraph/Tasks/CopyTask.cs
+++ b/Engine/Source/Programs/AutomationTool/BuildGraph/Tasks/CopyTask.cs
@@ -132,7 +132,7 @@ namespace AutomationTool.Tasks
 				{
 					if (_parameters.ErrorIfNotFound)
 					{
-						Logger.LogError("No files found matching '{SourcePattern}'", sourcePattern);
+						throw new AutomationException("No files found matching '{0}'", sourcePattern);
 					}
 					else
 					{
@@ -158,11 +158,11 @@ namespace AutomationTool.Tasks
 			{
 				if (_parameters.ErrorIfNotFound)
 				{
-					Logger.LogError(ex, "Error while trying to create file pattern match for '{SourcePattern}', error {ExceptionString}", sourcePattern, ex.Message);
+					throw new AutomationException(ex, "Error while trying to create file pattern match for '{0}': {1}", sourcePattern, ex.Message);
 				}
 				else
 				{
-					Logger.LogInformation(ex, "Error while trying to create file pattern match for '{SourcePattern}', error {ExceptionString}", sourcePattern, ex.Message);
+					Logger.LogInformation(ex, "Error while trying to create file pattern match for '{SourcePattern}': {Message}", sourcePattern, ex.Message);
 				}
 			}
 		}
diff --git a/Engine/Source/Programs/AutomationTool/BuildGraph/Tasks/MoveTask.cs b/Engine/Source/Programs/AutomationTool/BuildGraph/Tasks/MoveTask.cs
--- a/Engine/Source/Programs/AutomationTool/BuildGraph/Tasks/MoveTask.cs
+++ b/Engine/Source/Programs/AutomationTool/BuildGraph/Tasks/MoveTask.cs
@@ -104,7 +104,7 @@ namespace AutomationTool.Tasks
 				{
 					if (_parameters.ErrorIfNotFound)
 					{
-						Logger.LogError("No files found matching '{SourcePattern}'", sourcePattern);
+						throw new AutomationException("No files found matching '{0}'", sourcePattern);
 					}
 					else
 					{
@@ -129,11 +129,11 @@ namespace AutomationTool.Tasks
 			{
 				if (_parameters.ErrorIfNotFound)
 				{
-					Logger.LogError("Error while trying to create file pattern match for '{SourcePattern}', error {ExceptionString}", sourcePattern, ex.ToString());
+					throw new AutomationException(ex, "Error while trying to create file pattern match for '{0}': {1}", sourcePattern, ex.Message);
 				}
 				else
 				{
-					Logger.LogInformation("Error while trying to create file pattern match for '{SourcePattern}', error {ExceptionString}", sourcePattern, ex.ToString());
+					Logger.LogInformation(ex, "Error while trying to create file pattern match for '{SourcePattern}': {Message}", sourcePattern, ex.Message);
 				}
 			}
 
