Export correct build ID to fileshare snapshot descriptor from ZenExportSnapshotTask.

[RN] fix

diff --git a/Engine/Source/Programs/AutomationTool/BuildGraph/Tasks/ZenExportSnapshotTask.cs b/Engine/Source/Programs/AutomationTool/BuildGraph/Tasks/ZenExportSnapshotTask.cs
--- a/Engine/Source/Programs/AutomationTool/BuildGraph/Tasks/ZenExportSnapshotTask.cs
+++ b/Engine/Source/Programs/AutomationTool/BuildGraph/Tasks/ZenExportSnapshotTask.cs
@@ -345,7 +345,7 @@ namespace AutomationTool.Tasks
 			return bucketName;
 		}
 
-		private void WriteExportSource(JsonWriter writer, SnapshotStorageType destinationStorageType, ExportSourceData exportSource, string name)
+		private void WriteExportSource(JsonWriter writer, SnapshotStorageType destinationStorageType, ExportSourceData exportSource, string name, string destinationId)
 		{
 			string targetPlatform = _parameters.SnapshotDescriptorPlatform;
 			if (String.IsNullOrEmpty(targetPlatform))
@@ -371,22 +371,12 @@ namespace AutomationTool.Tasks
 						httpVersion = _parameters.DestinationCloudHttpVersion;
 					}
 
-					string destinationId;
 					string storageTypeName = "cloud";
 					string storageIdentifierName = "key";
 					if (destinationStorageType == SnapshotStorageType.Builds)
 					{
 						storageTypeName = "builds";
 						storageIdentifierName = "builds-id";
-
-						IoHash destinationKeyHash = IoHash.Compute(Encoding.UTF8.GetBytes(name));
-						CbObjectId objectId = new CbObjectId(destinationKeyHash.ToByteArray());
-						destinationId = objectId.ToString().ToLowerInvariant();
-					}
-					else
-					{
-						IoHash destinationKeyHash = IoHash.Compute(Encoding.UTF8.GetBytes(name));
-						destinationId = destinationKeyHash.ToString().ToLowerInvariant();
 					}
 
 					writer.WriteValue("name", name);
@@ -540,6 +530,7 @@ namespace AutomationTool.Tasks
 			}
 			int exportIndex;
 			string[] exportNames = new string[exportSources.Count];
+			string[] exportIds = new string[exportSources.Count];
 			List<ExportSourceData> successfullyExportedSources = new List<ExportSourceData>();
 
 			// Get the Zen executable path
@@ -634,6 +625,7 @@ namespace AutomationTool.Tasks
 							destinationId = destinationKeyHash.ToString().ToLowerInvariant();
 						}
 
+						exportIds[exportIndex] = destinationId;
 						exportSingleSourceCommandline.AppendFormat(" {0} --{1} {2} {3} {4} {5}", hostUrlArg, storageIdentifierName, destinationId, baseKeyArg, exportSource._projectId, exportSource._oplogId);
 						if (_parameters.SkipExport || TryExportOplogCommand(zenExe.FullName, exportSingleSourceCommandline.ToString()))
 						{
@@ -753,7 +745,7 @@ namespace AutomationTool.Tasks
 						{
 							writer.WriteObjectStart();
 							writer.WriteArrayStart("snapshots");
-							WriteExportSource(writer, destinationStorageType, exportSource, exportNames[exportIndex]);
+							WriteExportSource(writer, destinationStorageType, exportSource, exportNames[exportIndex], exportIds[exportIndex]);
 							writer.WriteArrayEnd();
 							writer.WriteObjectEnd();
 						}
@@ -772,7 +764,7 @@ namespace AutomationTool.Tasks
 						exportIndex = 0;
 						foreach (ExportSourceData exportSource in successfullyExportedSources)
 						{
-							WriteExportSource(writer, destinationStorageType, exportSource, exportNames[exportIndex]);
+							WriteExportSource(writer, destinationStorageType, exportSource, exportNames[exportIndex], exportIds[exportIndex]);
 							exportIndex = exportIndex + 1;
 						}
 
