Fix for incorrect name used when publishing snapshots.

diff --git a/Engine/Source/Programs/AutomationTool/BuildGraph/Tasks/ZenExportSnapshotTask.cs b/Engine/Source/Programs/AutomationTool/BuildGraph/Tasks/ZenExportSnapshotTask.cs
--- a/Engine/Source/Programs/AutomationTool/BuildGraph/Tasks/ZenExportSnapshotTask.cs
+++ b/Engine/Source/Programs/AutomationTool/BuildGraph/Tasks/ZenExportSnapshotTask.cs
@@ -607,7 +607,7 @@ namespace AutomationTool.Tasks
 
 						StringBuilder destinationNameBuilder = new StringBuilder();
 						destinationNameBuilder.AppendFormat("{0}.{1}.{2}", bucketName, _parameters.DestinationIdentifier, exportSource._oplogId);
-						exportNames[exportIndex] = destinationStorageType.ToString().ToLowerInvariant();
+						exportNames[exportIndex] = destinationNameBuilder.ToString().ToLowerInvariant();
 
 						string destinationId;
 						if (destinationStorageType == SnapshotStorageType.Builds)
