Added build graph tasks to upload artifacts to Cloud DDC (Jupiter). These are WIP and is only really useful for cases were files are not very large. Will be continuing to iterate on these to support larger files.

Made IoHash calculations on streams use the multithread version - performance gain on a buffer larger then a few MB.

diff --git a/Engine/Source/Programs/AutomationTool/BuildGraph/Tasks/CreateCloudArtifactTask.cs b/Engine/Source/Programs/AutomationTool/BuildGraph/Tasks/CreateCloudArtifactTask.cs
new file mode 100644
--- /dev/null
+++ b/Engine/Source/Programs/AutomationTool/BuildGraph/Tasks/CreateCloudArtifactTask.cs
@@ -0,0 +1,789 @@
+// Copyright Epic Games, Inc. All Rights Reserved.
+
+using System;
+using System.Collections.Concurrent;
+using System.Collections.Generic;
+using System.Diagnostics;
+using System.IO;
+using System.Linq;
+using System.Net.Http;
+using System.Net.Http.Headers;
+using System.Net.Mime;
+using System.Threading.Tasks;
+using System.Xml;
+using AutomationTool.Tasks.CloudDDC;
+using EpicGames.Compression;
+using EpicGames.Core;
+using EpicGames.Horde.Artifacts;
+using EpicGames.Horde.Commits;
+using EpicGames.Horde.Streams;
+using EpicGames.Serialization;
+using Microsoft.Extensions.Logging;
+using IoHash = EpicGames.Core.IoHash;
+
+#nullable enable
+
+namespace AutomationTool.Tasks
+{
+	/// <summary>
+	/// Parameters for a <see cref="CreateCloudArtifactTask"/>.
+	/// </summary>
+	public class CreateCloudArtifactTaskParameters
+	{
+		/// <summary>
+		/// Name of the artifact
+		/// </summary>
+		[TaskParameter]
+		public string Name { get; set; } = null!;
+
+		/// <summary>
+		/// The artifact type. Determines the permissions and expiration policy for the artifact.
+		/// </summary>
+		[TaskParameter]
+		public string Type { get; set; } = null!;
+
+		/// <summary>
+		/// Description for the artifact. Will be shown through the Horde dashboard.
+		/// </summary>
+		[TaskParameter(Optional = true)]
+		public string? Description { get; set; }
+
+		/// <summary>
+		/// The platform the artifact is created for
+		/// </summary>
+		[TaskParameter]
+		public string Platform { get; set; } = null!;
+
+		/// <summary>
+		/// The project this artifact is created for
+		/// </summary>
+		[TaskParameter]
+		public string Project { get; set; } = null!;
+
+		/// <summary>
+		/// The platform the artifact is created for
+		/// </summary>
+		[TaskParameter(Optional = true)]
+		public string? Host { get; set; }
+
+		/// <summary>
+		/// The platform the artifact is created for
+		/// </summary>
+		[TaskParameter(Optional = true)]
+		public string? Namespace { get; set; }
+
+		/// <summary>
+		/// The access token to use
+		/// </summary>
+		[TaskParameter(Optional = true)]
+		public string? AccessToken { get; set; }
+
+		/// <summary>
+		/// Base path for the uploaded files. All the tagged files must be under this directory. Defaults to the workspace root directory.
+		/// </summary>
+		[TaskParameter(Optional = true)]
+		public string? BaseDir { get; set; }
+
+		/// <summary>
+		/// Stream containing the artifact.
+		/// </summary>
+		[TaskParameter(Optional = true)]
+		public string? StreamId { get; set; }
+
+		/// <summary>
+		/// Commit for the uploaded artifact.
+		/// </summary>
+		[TaskParameter(Optional = true)]
+		public string? Commit { get; set; }
+		/// <summary>
+		/// Files to include in the artifact.
+		/// </summary>
+		[TaskParameter(ValidationType = TaskParameterValidationType.FileSpec)]
+		public string Files { get; set; } = null!;
+
+		/// <summary>
+		/// Other metadata for the artifact, separated by semicolons.
+		/// </summary>
+		[TaskParameter(Optional = true)]
+		public string? Metadata { get; set; }
+
+		/// <summary>
+		/// Force the blobs to be uploaded no matter if they exist or not
+		/// </summary>
+		[TaskParameter(Optional = true)]
+		public bool Force { get; set; } = false;
+
+		/// <summary>
+		/// The streamid used for your mainline, used to build content as patches on top of this branch if no other baselines exist.
+		/// </summary>
+		[TaskParameter]
+		public string BaselineBranch { get; set; } = null!;
+
+		/// <summary>
+        /// Set to upload more blobs in parallel, defaults to number of cores. Set to -1 to go as wide as possible.
+        /// </summary>
+        [TaskParameter(Optional = true)]
+        public int MaxBlobUploadParallelism { get; set; } = Environment.ProcessorCount;
+	}
+
+	/// <summary>
+	/// Uploads an artifact to Cloud DDC
+	/// </summary>
+	[TaskElement("CreateCloudArtifact", typeof(CreateCloudArtifactTaskParameters))]
+	public class CreateCloudArtifactTask : BgTaskImpl
+	{
+		readonly CreateCloudArtifactTaskParameters _parameters;
+
+		/// <summary>
+		/// Constructor.
+		/// </summary>
+		/// <param name="parameters">Parameters for this task.</param>
+		public CreateCloudArtifactTask(CreateCloudArtifactTaskParameters parameters)
+			=> _parameters = parameters;
+
+		/// <summary>
+		/// ExecuteAsync the task.
+		/// </summary>
+		/// <param name="job">Information about the current job.</param>
+		/// <param name="buildProducts">Set of build products produced by this node.</param>
+		/// <param name="tagNameToFileSet">Mapping from tag names to the set of files they include.</param>
+		public override async Task ExecuteAsync(JobContext job, HashSet<FileReference> buildProducts, Dictionary<string, HashSet<FileReference>> tagNameToFileSet)
+		{
+			ArtifactName name = new ArtifactName(_parameters.Name);
+			ArtifactType type = new ArtifactType(_parameters.Type);
+			List<string> metadata = (_parameters.Metadata ?? String.Empty).Split(';', StringSplitOptions.RemoveEmptyEntries | StringSplitOptions.TrimEntries).ToList();
+
+			string project = _parameters.Project;
+			string platform = _parameters.Platform;
+
+			metadata.Add($"name:{name}");
+			metadata.Add($"type:{type}");
+			if (!String.IsNullOrEmpty(_parameters.Description))
+			{
+				metadata.Add($"description:{_parameters.Description}");
+			}
+
+			// Add keys for the job that's executing
+			string? jobId = Environment.GetEnvironmentVariable("UE_HORDE_JOBID");
+			if (!String.IsNullOrEmpty(jobId))
+			{
+				metadata.Add($"job:{jobId}");
+
+				string? stepId = Environment.GetEnvironmentVariable("UE_HORDE_STEPID");
+				if (!String.IsNullOrEmpty(stepId))
+				{
+					metadata.Add($"job:{jobId}/step:{stepId}");
+				}
+			}
+
+			// Figure out the current change and stream id
+			StreamId streamId;
+			if (!String.IsNullOrEmpty(_parameters.StreamId))
+			{
+				streamId = new StreamId(_parameters.StreamId);
+			}
+			else
+			{
+				string? streamIdEnvVar = Environment.GetEnvironmentVariable("UE_HORDE_STREAMID");
+				if (!String.IsNullOrEmpty(streamIdEnvVar))
+				{
+					streamId = new StreamId(streamIdEnvVar);
+				}
+				else
+				{
+					throw new AutomationException("Missing UE_HORDE_STREAMID environment variable; unable to determine current stream.");
+				}
+			}
+			metadata.Add($"stream:{streamId}");
+			
+			string baselineBranch = _parameters.BaselineBranch;
+			if (String.IsNullOrEmpty(baselineBranch))
+			{
+				Logger.LogWarning("No baseline branch set, falling back to current stream as baseline. This will generate in-optimal de-dup for branches, make sure to set this to your mainline");
+				baselineBranch = streamId.ToString();
+			}
+
+			CommitId commitId;
+			if (!String.IsNullOrEmpty(_parameters.Commit))
+			{
+				commitId = new CommitId(_parameters.Commit);
+				metadata.Add($"commit:{commitId}");
+			}
+			else
+			{
+				try
+				{
+					int change = CommandUtils.P4Env.Changelist;
+					if (change > 0)
+					{
+						commitId = CommitId.FromPerforceChange(CommandUtils.P4Env.Changelist);
+						metadata.Add($"commit:{commitId}");
+					}
+				}
+				catch (AutomationException)
+				{
+					// not an error to not use p4, commit is just metadata
+				}
+				
+				// no commit specified, not required so ignoring
+			}
+
+			string host;
+			if (!String.IsNullOrEmpty(_parameters.Host))
+			{
+				host = _parameters.Host;
+			}
+			else
+			{
+				string? cloudHostEnvVar = Environment.GetEnvironmentVariable("UE-CloudPublishHost");
+				if (!String.IsNullOrEmpty(cloudHostEnvVar))
+				{
+					host = cloudHostEnvVar;
+				}
+				else
+				{
+					throw new AutomationException("Missing UE-CloudPublishHost environment variable; unable to determine cloud host.");
+				}
+			}
+
+			string ns;
+			if (!String.IsNullOrEmpty(_parameters.Namespace))
+			{
+				ns = _parameters.Namespace;
+			}
+			else
+			{
+				string? cloudDefaultNamespace = Environment.GetEnvironmentVariable("UE-CloudPublishNamespace");
+				if (!String.IsNullOrEmpty(cloudDefaultNamespace))
+				{
+					ns = cloudDefaultNamespace;
+				}
+				else
+				{
+					throw new AutomationException("Missing UE-CloudPublishNamespace environment variable; unable to default namespace please specify it in the task.");
+				}
+			}
+
+			string accessToken;
+			if (!String.IsNullOrEmpty(_parameters.AccessToken))
+			{
+				accessToken = _parameters.AccessToken;
+			}
+			else
+			{
+				string? cloudAccessTokenEnvVar = Environment.GetEnvironmentVariable("UE-CloudDataCacheAccessToken");
+				if (!String.IsNullOrEmpty(cloudAccessTokenEnvVar))
+				{
+					accessToken = cloudAccessTokenEnvVar;
+				}
+				else
+				{
+					throw new AutomationException("Missing UE-CloudDataCacheAccessToken environment variable; unable to find access token to use.");
+				}
+			}
+
+			// Resolve the files to include
+			DirectoryReference baseDir = ResolveDirectory(_parameters.BaseDir);
+			List<FileReference> files = BgTaskImpl.ResolveFilespec(baseDir, _parameters.Files, tagNameToFileSet).ToList();
+
+			bool validFiles = true;
+			foreach (FileReference file in files)
+			{
+				if (!file.IsUnderDirectory(baseDir))
+				{
+					Logger.LogError("Artifact file {File} is not under {BaseDir}", file, baseDir);
+					validFiles = false;
+				}
+			}
+
+			if (!validFiles)
+			{
+				throw new AutomationException($"Unable to create artifact {name} with given file list.");
+			}
+
+			long totalSize = files.Sum(x => x.ToFileInfo().Length);
+			string bucketId = $"{project}.{type}.{streamId}.{platform}";
+			CbObjectId buildId = CbObjectId.NewObjectId();
+			using HttpClient httpClient = BuildHttpClient(host, ns, bucketId, buildId, accessToken);
+			Logger.LogInformation("Creating artifact {ArtifactId} '{ArtifactName}' ({ArtifactType}) from {NumFiles} files ({TotalSize:n1}mb). Namespace: {NamespaceId}, Bucket: {Bucket}", buildId, name, type, files.Count, totalSize / (1024.0 * 1024.0), ns, bucketId);
+
+			// Do the work
+			Stopwatch timer = Stopwatch.StartNew();
+
+			// Compute file hashes
+			Dictionary<FileReference, IoHash> hashCache = await ComputeHashCacheAsync(files);
+
+			double hashTimeSecs = timer.Elapsed.TotalSeconds;
+			Logger.LogInformation("Hash cache built for build '{ArtifactName}' in {Duration:n1}s.", buildId, hashTimeSecs);
+
+			// Create the build
+			uint chunkSizeHint = await UploadBuildAsync(httpClient, project, streamId.ToString(), baselineBranch, platform, name.ToString(), metadata);
+
+			long totalSizeUploaded = 0;
+			CbObjectId buildPartId = CbObjectId.NewObjectId();
+			if (_parameters.Force)
+			{
+				// force all parts to be uploaded, unlikely that blobs will be missing after this.
+				// this is faster for blobs that are very likely to be unique and when your connectivity to CloudDDC is very fast
+				Logger.LogInformation("Force uploading all {CountOfParts} parts of artifact '{ArtifactId}'. ", files.Count, buildId);
+
+				long sizeUploaded = 0;
+				TimeSpan blobUploadStart = timer.Elapsed;
+				await Parallel.ForEachAsync(files, new ParallelOptions {MaxDegreeOfParallelism = _parameters.MaxBlobUploadParallelism}, async (file, _) =>
+				{
+					IoHash hashOfFile = hashCache[file];
+					long uploadedBytes = await UploadBlobAsync(httpClient, buildPartId, file, hashOfFile);
+					sizeUploaded += uploadedBytes;
+				});
+
+				totalSizeUploaded += sizeUploaded;
+				double uploadTimeBlobsSecs = (timer.Elapsed - blobUploadStart).TotalSeconds;
+				Logger.LogInformation("Force uploaded blobs for artifact {ArtifactId} in {Time:n1}s ({Rate:n1}mb/s, {RateMbps:n1}mbps, {TotalSize:n1}mb)", buildId, uploadTimeBlobsSecs, sizeUploaded / (uploadTimeBlobsSecs * 1024.0 * 1024.0), (sizeUploaded * 8.0) / (uploadTimeBlobsSecs * 1024.0 * 1024.0), sizeUploaded / (1024.0 * 1024.0));
+			}
+
+			TimeSpan partCreationStart = timer.Elapsed;
+			// Create the build part
+			(List<IoHash> missingParts, IoHash partHash) = await UploadBuildPartObjectAsync(httpClient, name, buildPartId, baseDir, files, hashCache, chunkSizeHint);
+
+			double partCreationSecs = (timer.Elapsed - partCreationStart).TotalSeconds;
+			Logger.LogInformation("Spent {Time:n1}s finding {CountOfMissingParts} missing part(s) for build '{ArtifactName}'.", partCreationSecs, missingParts.Count, name);
+
+			const int MaxAttempts = 5;
+			for (int i = 0; i < MaxAttempts; i++)
+			{
+				// upload missing parts
+				if (missingParts.Count != 0)
+				{
+					Logger.LogInformation("{NumMissingParts} parts missing from build '{ArtifactName}'. Starting upload.", missingParts.Count, name);
+
+					List<KeyValuePair<IoHash, FileReference>> filesToDownload = new();
+
+					foreach (FileReference f in files)
+					{
+						IoHash hashOfFile = hashCache[f];
+
+						if (missingParts.Contains(hashOfFile))
+						{
+							filesToDownload.Add(new KeyValuePair<IoHash, FileReference>(hashOfFile, f));
+						}
+					}
+
+					long sizeUploaded = 0;
+					TimeSpan blobUploadStart = timer.Elapsed;
+					await Parallel.ForEachAsync(filesToDownload, new ParallelOptions {MaxDegreeOfParallelism = _parameters.MaxBlobUploadParallelism}, async (pair, _) =>
+					{
+						long uploadedBytes = await UploadBlobAsync(httpClient, buildPartId, pair.Value, pair.Key);
+						sizeUploaded += uploadedBytes;
+					});
+
+					totalSizeUploaded += sizeUploaded;
+					double uploadTimeBlobsSecs = (timer.Elapsed - blobUploadStart).TotalSeconds;
+					Logger.LogInformation("Uploaded blobs for artifact {ArtifactId} in {Time:n1}s ({Rate:n1}mb/s, {RateMbps:n1}mbps, {TotalSize:n1}mb)", buildId, uploadTimeBlobsSecs, sizeUploaded / (uploadTimeBlobsSecs * 1024.0 * 1024.0), (sizeUploaded * 8.0) / (uploadTimeBlobsSecs * 1024.0 * 1024.0), sizeUploaded / (1024.0 * 1024.0));
+					// verify that all parts are now found
+					missingParts = await FinalizeBuildPartAsync(httpClient, buildPartId, partHash);
+				}
+
+				if (missingParts.Count == 0)
+				{
+					Logger.LogInformation("Part {PartId} successfully uploaded for build '{ArtifactName}'.", buildPartId, buildId);
+
+					// upload successful
+					break;
+				}
+			}
+
+			// Flag the build as uploaded, making it read only
+			await FinalizeBuildAsync(httpClient);
+
+			double uploadTimeSecs = timer.Elapsed.TotalSeconds;
+			Logger.LogInformation("Uploaded artifact {ArtifactId} in {Time:n1}s ({Rate:n1}mb/s, {RateMbps:n1}mbps)", buildId, timer.Elapsed.TotalSeconds, totalSizeUploaded / (uploadTimeSecs * 1024.0 * 1024.0), (totalSizeUploaded * 8.0) / (uploadTimeSecs * 1024.0 * 1024.0));
+		}
+
+		static HttpClient BuildHttpClient(string host, string ns, string bucketId, CbObjectId buildId, string accessToken)
+		{
+			HttpClient httpClient = new HttpClient();
+			httpClient.BaseAddress = new Uri($"{host}/api/v2/builds/{ns}/{bucketId}/{buildId}/");
+
+			CbObjectId sessionId = CbObjectId.NewObjectId();
+			Logger.LogInformation("Using SessionId {SessionId} please include this in any error reports.", sessionId);
+			string authScheme = "Bearer"; // TODO: for a proper implementation this should be configurable
+			httpClient.DefaultRequestHeaders.Add("Authorization", $"{authScheme} {accessToken}");
+			httpClient.DefaultRequestHeaders.Add("UE-Session", sessionId.ToString());
+			httpClient.DefaultRequestHeaders.Add("Accept", CustomMediaTypeNames.UnrealCompactBinary);
+			httpClient.Timeout = TimeSpan.FromMinutes(5.0); // bump timeout for each request as we sometimes have larger files to upload
+			return httpClient;
+		}
+
+		static async Task<Dictionary<FileReference, IoHash>> ComputeHashCacheAsync(List<FileReference> files)
+		{
+			ConcurrentDictionary<FileReference, IoHash> cache = new();
+
+			await Parallel.ForEachAsync(files, async (file, token) =>
+			{
+				await using Stream s = file.ToFileInfo().OpenRead();
+				IoHash hash = await IoHash.ComputeAsync(s, token);
+				cache.TryAdd(file, hash);
+			});
+
+			return cache.ToDictionary();
+		}
+
+		static async Task<uint> UploadBuildAsync(HttpClient httpClient, string projectName, string branch, string baselineBranch, string platform, string name, List<string> metadata)
+		{
+			CbWriter writer = new CbWriter();
+			writer.BeginObject();
+			writer.WriteString("project", projectName);
+			writer.WriteString("branch", branch);
+			writer.WriteString("baselineBranch", baselineBranch);
+			writer.WriteString("platform", platform);
+			writer.WriteString("name", name);
+			foreach (string line in metadata)
+			{
+				string[] s = line.Split(":");
+				if (s.Length > 1 && !String.IsNullOrEmpty(s[0]))
+				{
+					writer.WriteString(s[0], s[1]);
+					continue;
+				}
+				s = line.Split("=");
+				if (s.Length > 1 && !String.IsNullOrEmpty(s[0]))
+				{
+					writer.WriteString(s[0], s[1]);
+					continue;
+				}
+				Logger.LogWarning("Unable to convert metadata line: \'{Line}\' to key-value. Please use : or = to separate key and value. Ignoring.", line);
+			}
+			writer.EndObject();
+
+			CbObject buildObject = writer.ToObject();
+
+			using HttpContent requestContent = new ByteArrayContent(buildObject.GetView().ToArray());
+			requestContent.Headers.ContentType = new MediaTypeHeaderValue(CustomMediaTypeNames.UnrealCompactBinary);
+			using HttpResponseMessage result = await httpClient.PutAsync(new Uri($"", UriKind.Relative), requestContent);
+			result.EnsureSuccessStatusCode();
+			// TODO: return chunk size once Jupiter is updated to respond with it
+			//PutBuildResponse response = await result.Content.ReadAsCompactBinaryAsync<PutBuildResponse>();
+			//return response.ChunkSize;
+			return 0;
+		}
+
+		static async Task FinalizeBuildAsync(HttpClient httpClient)
+		{
+			using HttpResponseMessage result = await httpClient.PostAsync(new Uri("finalize", UriKind.Relative), null);
+			result.EnsureSuccessStatusCode();
+		}
+
+		static async Task<(List<IoHash>, IoHash)> UploadBuildPartObjectAsync(HttpClient httpClient, ArtifactName artifactName, CbObjectId buildPartId, DirectoryReference baseDir, List<FileReference> files, Dictionary<FileReference, IoHash> hashCache, uint chunkSizeHint)
+		{
+			CbWriter writer = new CbWriter();
+			writer.BeginObject();
+			writer.BeginObject("files");
+
+			writer.BeginUniformArray("paths", CbFieldType.String);
+			foreach (FileReference file in files)
+			{
+				// paths are always forwarded slashes
+				string path = file.MakeRelativeTo(baseDir).Replace('\\','/');
+				writer.WriteStringValue(path);
+			}
+			writer.EndUniformArray();
+
+			// encode chmod attributes if we are running on non-windows
+			if (!OperatingSystem.IsWindows())
+			{
+				writer.BeginUniformArray("chmod", CbFieldType.IntegerPositive);
+				foreach (FileReference file in files)
+				{
+					writer.WriteIntegerValue((uint)File.GetUnixFileMode(file.FullName));
+				}
+				writer.EndUniformArray();
+			}
+			else if (OperatingSystem.IsWindows())
+			{
+				// encode windows file attributes when running on windows
+				writer.BeginUniformArray("attr", CbFieldType.IntegerPositive);
+				foreach (FileReference file in files)
+				{
+					// paths are always forwarded slashes
+					writer.WriteIntegerValue((uint)File.GetAttributes(file.FullName));
+				}
+				writer.EndUniformArray();
+			}
+
+			writer.BeginUniformArray("rawhashes", CbFieldType.BinaryAttachment);
+			foreach (FileReference file in files)
+			{
+				IoHash hash = hashCache[file];
+				writer.WriteBinaryAttachmentValue(hash);
+			}
+			writer.EndUniformArray();
+
+			writer.BeginUniformArray("rawsizes", CbFieldType.IntegerPositive);
+			foreach (FileReference file in files)
+			{
+				// 0 if not chunked, the chunk size if chunked
+				writer.WriteIntegerValue(file.ToFileInfo().Length);
+			}
+			writer.EndUniformArray();
+			writer.WriteInteger("chunksizehint", chunkSizeHint);
+
+			writer.EndObject(); // files
+			writer.EndObject();
+
+			CbObject buildPartObject = writer.ToObject();
+			IoHash partHash = IoHash.Compute(buildPartObject.GetView().Span);
+			string partName = artifactName.ToString();
+
+			using HttpContent requestContent = new ByteArrayContent(buildPartObject.GetView().ToArray());
+			requestContent.Headers.ContentType = new MediaTypeHeaderValue(CustomMediaTypeNames.UnrealCompactBinary);
+			using HttpResponseMessage result = await httpClient.PutAsync(new Uri($"parts/{buildPartId}/{partName}", UriKind.Relative), requestContent);
+			result.EnsureSuccessStatusCode();
+
+			NeedsResponse response = await result.Content.ReadAsCompactBinaryAsync<NeedsResponse>();
+			return (response.Needs, partHash);
+		}
+
+		static async Task<List<IoHash>> FinalizeBuildPartAsync(HttpClient httpClient, CbObjectId partId, IoHash partHash)
+		{
+			using HttpResponseMessage result = await httpClient!.PostAsync(new Uri($"parts/{partId}/finalize/{partHash}", UriKind.Relative), null);
+			result.EnsureSuccessStatusCode();
+
+			NeedsResponse response = await result.Content.ReadAsCompactBinaryAsync<NeedsResponse>();
+			return response.Needs;
+		}
+
+		static async Task<long> UploadBlobAsync(HttpClient httpClient, CbObjectId partId, FileReference fileReference, IoHash hashOfFile)
+		{
+			FileInfo fi = fileReference.ToFileInfo();
+			await using Stream sourceStream = fi.OpenRead();
+			long rawLength = fi.Length;
+			bool compressContent = true;
+			long contentLength;
+			Stream? streamToSend = null;
+			string? tempFileName = null;
+			try
+			{
+				MediaTypeHeaderValue contentType;
+
+				bool useTempFile = rawLength > UseTempFileAtBytes;
+
+				if (compressContent)
+				{
+					if (useTempFile)
+					{
+						tempFileName = Path.GetTempFileName();
+						FileInfo tempFileInfo = new FileInfo(tempFileName!);
+						{
+							await using Stream s = tempFileInfo.OpenWrite();
+							CompressedBuffer.CompressContent(s, OodleCompressorType.Kraken, OodleCompressionLevel.VeryFast, await sourceStream.ReadAllBytesAsync());
+						}
+
+						tempFileInfo.Refresh();
+						streamToSend = tempFileInfo.OpenRead();
+						contentLength = tempFileInfo.Length;
+					}
+					else
+					{
+						MemoryStream ms = new MemoryStream();
+						CompressedBuffer.CompressContent(ms, OodleCompressorType.Kraken, OodleCompressionLevel.VeryFast, await sourceStream.ReadAllBytesAsync());
+						ms.Position = 0;
+						
+						streamToSend = ms;
+						contentLength = ms.Length;
+					}
+
+					contentType = new MediaTypeHeaderValue(CustomMediaTypeNames.UnrealCompressedBuffer);
+
+				}
+				else
+				{
+					// submit the uncompressed content
+					streamToSend = sourceStream;
+					contentType = new MediaTypeHeaderValue(MediaTypeNames.Application.Octet);
+					contentLength = fi.Length;
+				}
+
+				/*if (rawLength > MultipartAtBytes)
+				{
+					await MultipartUploadFileAsync(httpClient, partId, hashOfFile, streamToSend, contentType, idealChunkSize);
+				}
+				else*/
+				{
+					using HttpContent requestContent = new StreamContent(streamToSend);
+					requestContent.Headers.ContentType = contentType;
+
+					using HttpResponseMessage result = await httpClient!.PutAsync(new Uri($"parts/{partId}/blobs/{hashOfFile}", UriKind.Relative), requestContent);
+					result.EnsureSuccessStatusCode();
+				}
+			}
+			finally
+			{
+				if (streamToSend != null)
+				{
+					await streamToSend.DisposeAsync();
+				}
+
+				if (tempFileName != null)
+				{
+					if (File.Exists(tempFileName))
+					{
+						File.Delete(tempFileName);
+					}
+				}
+			}
+
+			return contentLength;
+		}
+
+		// TODO: Multipart uploads not fully implemented in Jupiter yet
+		/*private static async Task MultipartUploadFileAsync(HttpClient httpClient, CbObjectId partId, IoHash hashOfFile, Stream streamToSend, MediaTypeHeaderValue contentType, int idealChunkSize)
+		{
+			if (!streamToSend.CanSeek)
+				throw new ArgumentException("Stream must be seekable");
+			Logger.LogInformation("Multipart uploading blob {Hash}.", hashOfFile);
+
+			IoHash hashOfCompressedContent = await IoHash.ComputeAsync(streamToSend);
+			// reset the stream so we can send it
+			streamToSend.Position = 0;
+
+			string uploadId;
+			string blobName;
+			List<string> partQueryStrings;
+			// start the multipart upload
+			{
+				using HttpContent requestContent = new ByteArrayContent(Array.Empty<byte>());
+				requestContent.Headers.ContentType = new MediaTypeHeaderValue(CustomMediaTypeNames.UnrealCompactBinary);
+
+				using HttpResponseMessage result = await httpClient!.PostAsync(new Uri($"parts/{partId}/blobs/{hashOfFile}/startMultipartUpload", UriKind.Relative), requestContent);
+				result.EnsureSuccessStatusCode();
+				MultipartUploadIdResponse response = await result.Content.ReadAsCompactBinaryAsync<MultipartUploadIdResponse>();
+				uploadId = response.UploadId;
+				blobName = response.BlobName;
+				partQueryStrings = response.PartQueryStrings;
+			}
+
+			// upload parts
+			// split the large file into the smallest chunks allowed
+			long bytesLeft = streamToSend.Length;
+
+			byte[] chunk = Array.Empty<byte>();
+			long bytesToRead = Math.Min(idealChunkSize, bytesLeft);
+			while (bytesToRead > 0)
+			{
+				try
+				{
+					chunk = ArrayPool<byte>.Shared.Rent((int)bytesToRead);
+					using HttpContent requestContent = new ByteArrayContent(chunk);
+					requestContent.Headers.ContentType = new MediaTypeHeaderValue(MediaTypeNames.Application.Octet);
+
+					using HttpResponseMessage result = await httpClient!.PutAsync(new Uri($"parts/{partId}/blobs/{hashOfFile}/uploadMultipart?{queryString}", UriKind.Relative), requestContent);
+					
+					result.EnsureSuccessStatusCode();
+				}
+				finally
+				{
+					bytesLeft -= bytesToRead;
+					bytesToRead = Math.Min(idealChunkSize, bytesLeft);
+					ArrayPool<byte>.Shared.Return(chunk);
+				}
+			}
+			
+			// complete the multipart upload
+			{
+				CompleteMultipartUploadRequest request = new CompleteMultipartUploadRequest() {BlobName = blobName, IsCompressed = isCompressed, UploadId = uploadId};
+				byte[] buf = CbSerializer.SerializeToByteArray(request);
+				using HttpContent requestContent = new ByteArrayContent(buf);
+				requestContent.Headers.ContentType = contentType;
+
+				using HttpResponseMessage result = await httpClient!.PostAsync(new Uri($"parts/{partId}/blobs/{hashOfFile}/completeMultipart", UriKind.Relative), requestContent);
+				result.EnsureSuccessStatusCode();
+			}
+
+			Logger.LogInformation("Multipart uploading of blob {Hash} finished.", hashOfFile);
+		}*/
+
+		private const long UseTempFileAtBytes = Int32.MaxValue; // use a backing file when the uncompressed file is too large to fit into memory
+
+		/// <inheritdoc/>
+		public override void Write(XmlWriter writer)
+			=> Write(writer, _parameters);
+
+		/// <inheritdoc/>
+		public override IEnumerable<string> FindConsumedTagNames()
+			=> FindTagNamesFromFilespec(_parameters.Files);
+
+		/// <inheritdoc/>
+		public override IEnumerable<string> FindProducedTagNames()
+			=> Enumerable.Empty<string>();
+	}
+
+	namespace CloudDDC
+	{
+		class StartMultipartUploadRequest
+		{
+			[CbField("blobLength")]
+			public ulong BlobLength { get; set; }
+		}
+
+		class MultipartUploadIdResponse
+		{
+			[CbField("uploadId")]
+			public string UploadId { get; set; } = null!;
+
+			[CbField("blobName")]
+			public string BlobName { get; set; } = null!;
+		
+			[CbField("blobQueryStrings")]
+#pragma warning disable CA2227
+			public List<string> PartQueryStrings { get; set; } = new List<string>();
+#pragma warning restore CA2227
+		}
+
+		class PutBuildResponse
+		{
+			[CbField("chunkSize")]
+			public uint ChunkSize { get; set; }
+		}
+
+		class NeedsResponse
+		{
+			public NeedsResponse()
+			{
+				Needs = new List<IoHash>();
+			}
+
+			public NeedsResponse(List<IoHash> needs)
+			{
+				Needs = needs;
+			}
+
+			[CbField("needs")]
+			public List<IoHash> Needs { get; set; }
+		}
+
+		class CompleteMultipartUploadRequest
+		{
+			[CbField("blobName")]
+			public string BlobName { get; set; } = null!;
+
+			[CbField("uploadId")]
+			public string UploadId { get; set; } = null!;
+
+			[CbField("isCompressed")]
+			public bool IsCompressed { get; set; }
+		}
+
+		static class CustomMediaTypeNames
+		{
+			/// <summary>
+			/// Media type for compact binary
+			/// </summary>
+			public const string UnrealCompactBinary = "application/x-ue-cb";
+
+			/// <summary>
+			/// Media type for compressed buffers
+			/// </summary>
+			public const string UnrealCompressedBuffer = "application/x-ue-comp";
+		}
+	}
+}
diff --git a/Engine/Source/Programs/AutomationTool/BuildGraph/Tasks/RetrieveCloudArtifactTask.cs b/Engine/Source/Programs/AutomationTool/BuildGraph/Tasks/RetrieveCloudArtifactTask.cs
new file mode 100644
--- /dev/null
+++ b/Engine/Source/Programs/AutomationTool/BuildGraph/Tasks/RetrieveCloudArtifactTask.cs
@@ -0,0 +1,497 @@
+// Copyright Epic Games, Inc. All Rights Reserved.
+
+using System;
+using System.Collections.Generic;
+using System.Diagnostics;
+using System.IO;
+using System.Linq;
+using System.Net.Http;
+using System.Net.Http.Headers;
+using System.Net.Mime;
+using System.Threading.Tasks;
+using System.Xml;
+using AutomationTool.Tasks.CloudDDC;
+using EpicGames.Core;
+using EpicGames.Horde.Artifacts;
+using EpicGames.Horde.Commits;
+using EpicGames.Horde.Streams;
+using EpicGames.Serialization;
+using Microsoft.Extensions.Logging;
+using Exception = System.Exception;
+
+#nullable enable
+
+namespace AutomationTool.Tasks
+{
+	/// <summary>
+	/// Parameters for a <see cref="RetrieveCloudArtifactTask"/>.
+	/// </summary>
+	public class RetrieveCloudArtifactTaskParameters
+	{
+		/// <summary>
+		/// Stream containing the artifact
+		/// </summary>
+		[TaskParameter(Optional = true)]
+		public string? StreamId { get; set; } = null!;
+
+		/// <summary>
+		/// Change number for the artifact
+		/// </summary>
+		[TaskParameter(Optional = true)]
+		public string? Commit { get; set; }
+
+		/// <summary>
+		/// Maximum commit for the artifact
+		/// </summary>
+		[TaskParameter(Optional = true)]
+		public string? MaxCommit { get; set; }
+
+		/// <summary>
+		/// Name of the artifact
+		/// </summary>
+		[TaskParameter(Optional = true)]
+		public string Name { get; set; } = null!;
+
+		/// <summary>
+		/// The artifact type. Determines the permissions and expiration policy for the artifact.
+		/// </summary>
+		[TaskParameter]
+		public string Type { get; set; } = null!;
+
+		/// <summary>
+		/// Keys for the artifact
+		/// </summary>
+		[TaskParameter(Optional = true)]
+		public string Keys { get; set; } = null!;
+
+		/// <summary>
+		/// Output directory for 
+		/// </summary>
+		[TaskParameter]
+		public string? OutputDir { get; set; }
+
+		/// <summary>
+		/// The platform the artifact is created for
+		/// </summary>
+		[TaskParameter]
+		public string Platform { get; set; } = null!;
+
+		/// <summary>
+		/// The project this artifact is created for
+		/// TODO: Can this be defaulted?
+		/// </summary>
+		[TaskParameter]
+		public string Project { get; set; } = null!;
+
+		/// <summary>
+		/// The platform the artifact is created for
+		/// </summary>
+		[TaskParameter(Optional = true)]
+		public string? Host { get; set; }
+
+		/// <summary>
+		/// The platform the artifact is created for
+		/// </summary>
+		[TaskParameter(Optional = true)]
+		public string? Namespace { get; set; }
+
+		/// <summary>
+		/// The access token to use
+		/// </summary>
+		[TaskParameter(Optional = true)]
+		public string? AccessToken { get; set; }
+
+		/// <summary>
+		/// Set this to use the latest match if multiple artifacts are possible matches
+		/// </summary>
+		[TaskParameter(Optional = true)]
+		public bool AllowMultipleMatches { get; set; } = false;
+	}
+
+	/// <summary>
+	/// Retrieves an artifact from Cloud DDC
+	/// </summary>
+	[TaskElement("RetrieveCloudArtifact", typeof(RetrieveCloudArtifactTaskParameters))]
+	public class RetrieveCloudArtifactTask : BgTaskImpl
+	{
+		readonly RetrieveCloudArtifactTaskParameters _parameters;
+
+		/// <summary>
+		/// Constructor.
+		/// </summary>
+		/// <param name="parameters">Parameters for this task.</param>
+		public RetrieveCloudArtifactTask(RetrieveCloudArtifactTaskParameters parameters)
+			=> _parameters = parameters;
+
+		/// <summary>
+		/// ExecuteAsync the task.
+		/// </summary>
+		/// <param name="job">Information about the current job.</param>
+		/// <param name="buildProducts">Set of build products produced by this node.</param>
+		/// <param name="tagNameToFileSet">Mapping from tag names to the set of files they include.</param>
+		public override async Task ExecuteAsync(JobContext job, HashSet<FileReference> buildProducts, Dictionary<string, HashSet<FileReference>> tagNameToFileSet)
+		{
+			if (!String.IsNullOrEmpty(_parameters.Commit) && !String.IsNullOrEmpty(_parameters.MaxCommit))
+			{
+				throw new AutomationException("Cannot specify both Commit and MaxCommit parameters for retrieving an artifact.");
+			}
+
+			string project = _parameters.Project;
+			string platform = _parameters.Platform;
+
+			// Figure out the current change and stream id
+			StreamId streamId;
+			if (!String.IsNullOrEmpty(_parameters.StreamId))
+			{
+				streamId = new StreamId(_parameters.StreamId);
+			}
+			else
+			{
+				string? streamIdEnvVar = Environment.GetEnvironmentVariable("UE_HORDE_STREAMID");
+				if (!String.IsNullOrEmpty(streamIdEnvVar))
+				{
+					streamId = new StreamId(streamIdEnvVar);
+				}
+				else
+				{
+					throw new AutomationException("Missing UE_HORDE_STREAMID environment variable; unable to determine current stream.");
+				}
+			}
+
+			string host;
+			if (!String.IsNullOrEmpty(_parameters.Host))
+			{
+				host = _parameters.Host;
+			}
+			else
+			{
+				string? cloudHostEnvVar = Environment.GetEnvironmentVariable("UE-CloudPublishHost");
+				if (!String.IsNullOrEmpty(cloudHostEnvVar))
+				{
+					host = cloudHostEnvVar;
+				}
+				else
+				{
+					throw new AutomationException("Missing UE-CloudPublishHost environment variable; unable to determine cloud host.");
+				}
+			}
+
+			string ns;
+			if (!String.IsNullOrEmpty(_parameters.Namespace))
+			{
+				ns = _parameters.Namespace;
+			}
+			else
+			{
+				string? cloudDefaultNamespace = Environment.GetEnvironmentVariable("UE-CloudPublishNamespace");
+				if (!String.IsNullOrEmpty(cloudDefaultNamespace))
+				{
+					ns = cloudDefaultNamespace;
+				}
+				else
+				{
+					throw new AutomationException("Missing UE-CloudPublishNamespace environment variable; unable to default namespace please specify it in the task.");
+				}
+			}
+
+			string accessToken;
+			if (!String.IsNullOrEmpty(_parameters.AccessToken))
+			{
+				accessToken = _parameters.AccessToken;
+			}
+			else
+			{
+				string? cloudAccessTokenEnvVar = Environment.GetEnvironmentVariable("UE-CloudDataCacheAccessToken");
+				if (!String.IsNullOrEmpty(cloudAccessTokenEnvVar))
+				{
+					accessToken = cloudAccessTokenEnvVar;
+				}
+				else
+				{
+					throw new AutomationException("Missing UE-CloudDataCacheAccessToken environment variable; unable to find access token to use.");
+				}
+			}
+
+			// Get the current commit id
+			CommitId? minCommitId = null;
+			CommitId? maxCommitId = null;
+			if (!String.IsNullOrEmpty(_parameters.MaxCommit))
+			{
+				maxCommitId = new CommitId(_parameters.MaxCommit);
+			}
+			else if (!String.IsNullOrEmpty(_parameters.Commit))
+			{
+				minCommitId = maxCommitId = new CommitId(_parameters.Commit);
+			}
+			else
+			{
+				try
+				{
+					int change = CommandUtils.P4Env.Changelist;
+					if (change > 0)
+					{
+						minCommitId = maxCommitId = CommitId.FromPerforceChange(CommandUtils.P4Env.Changelist);
+					}
+				}
+				catch (AutomationException)
+				{
+					// not an error to run without p4
+				}
+			}
+
+			ArtifactName? name = null;
+			if (!String.IsNullOrEmpty(_parameters.Name))
+			{
+				name = new ArtifactName(_parameters.Name);
+			}
+
+			ArtifactType type = new ArtifactType(_parameters.Type);
+			List<string> keys = (_parameters.Keys ?? String.Empty).Split(';', StringSplitOptions.RemoveEmptyEntries | StringSplitOptions.TrimEntries).ToList();
+
+			CbWriter queryWriter = new CbWriter();
+			queryWriter.BeginObject(); // root
+			queryWriter.BeginObject("query"); // query
+			
+			queryWriter.BeginObject("stream");
+			queryWriter.WriteString("$eq", streamId.ToString());
+			queryWriter.EndObject();
+
+			if (maxCommitId != null)
+			{
+				if (minCommitId == maxCommitId)
+				{
+					queryWriter.BeginObject("commit");
+					queryWriter.WriteString("$eq", maxCommitId.ToString());
+					queryWriter.EndObject();
+				}
+				else if (minCommitId != null)
+				{
+					queryWriter.BeginObject("commit");
+					queryWriter.WriteString("$lte", maxCommitId.ToString());
+					queryWriter.WriteString("$gte", minCommitId.ToString());
+					queryWriter.EndObject();
+				}
+				else
+				{
+					queryWriter.BeginObject("commit");
+					queryWriter.WriteString("$lte", maxCommitId.ToString());
+					queryWriter.EndObject();
+				}
+			}
+
+			if (name != null)
+			{
+				queryWriter.BeginObject("name");
+				queryWriter.WriteString("$eq", name.ToString());
+				queryWriter.EndObject();
+			}
+
+			queryWriter.BeginObject("type");
+			queryWriter.WriteString("$eq", type.ToString());
+			queryWriter.EndObject();
+
+			if (keys.Count > 0)
+			{
+				queryWriter.BeginObject("keys");
+				queryWriter.BeginUniformArray("$in", CbFieldType.String);
+				foreach (string key in keys)
+				{
+					queryWriter.WriteStringValue(key);
+				}
+				queryWriter.EndUniformArray();
+				queryWriter.EndObject();
+			}
+			
+			queryWriter.EndObject(); // end query
+			queryWriter.BeginObject("options");
+			queryWriter.WriteInteger("limit", 1); // we only care about the first match
+			queryWriter.EndObject(); // options
+			queryWriter.EndObject(); // end root object
+
+			CbObject queryObject = queryWriter.ToObject();
+			DirectoryReference outputDir = ResolveDirectory(_parameters.OutputDir);
+			string bucketId = $"{project}.{type}.{streamId}.{platform}";
+			using HttpClient httpClient = BuildHttpClient(host, ns, bucketId, accessToken);
+
+			List<CbObjectId> artifactIds = await SearchArtifactsAsync(httpClient, queryObject);
+
+			if (artifactIds.Count == 0)
+			{
+				throw new AutomationException("Unable to find any artifact matching given criteria.");
+			}
+
+			if (!_parameters.AllowMultipleMatches && artifactIds.Count != 1)
+			{
+				throw new AutomationException("More then one matching artifact given criteria, set \"AllowMultipleMatches\" option if you want to use the newest matching artifact or refine the search");
+			}
+			
+			CbObjectId artifact = artifactIds.First();
+			Logger.LogInformation("Found artifact {ArtifactId}", artifact);
+
+			// Upload the files
+			Stopwatch timer = Stopwatch.StartNew();
+			await DownloadArtifactAsync(httpClient, artifact, outputDir.ToDirectoryInfo());
+			Logger.LogInformation("Retrieved artifact {ArtifactId} in {Time:n1}s", artifact, timer.Elapsed.TotalSeconds);
+		}
+
+		static HttpClient BuildHttpClient(string host, string ns, string bucketId, string accessToken)
+		{
+			HttpClient httpClient = new HttpClient();
+			httpClient.BaseAddress = new Uri($"{host}/api/v2/builds/{ns}/{bucketId}/");
+
+			CbObjectId sessionId = CbObjectId.NewObjectId();
+			Logger.LogInformation("Using SessionId {SessionId} please include this in any error reports.", sessionId);
+			string authScheme = "Bearer"; // TODO: for a proper implementation this should be configurable
+			httpClient.DefaultRequestHeaders.Add("Authorization", $"{authScheme} {accessToken}");
+			httpClient.DefaultRequestHeaders.Add("UE-Session", sessionId.ToString());
+			httpClient.DefaultRequestHeaders.Add("Accept", CustomMediaTypeNames.UnrealCompactBinary);
+			httpClient.Timeout = TimeSpan.FromMinutes(5.0); // bump timeout for each request as we sometimes have larger files
+			return httpClient;
+		}
+
+		static async Task DownloadArtifactAsync(HttpClient httpClient, CbObjectId buildId, System.IO.DirectoryInfo outputDir)
+		{
+			using HttpResponseMessage result = await httpClient.GetAsync(new Uri($"{buildId}", UriKind.Relative));
+			result.EnsureSuccessStatusCode();
+
+			byte[] response = await result.Content.ReadAsByteArrayAsync();
+			CbObject o = new CbObject(response);
+			if (o["parts"].Equals(CbField.Empty))
+			{
+				throw new Exception("Failed to find parts in build, unexpected result, unable to download");
+			}
+
+			CbObject parts = o["parts"].AsObject();
+			foreach (CbField part in parts)
+			{
+				Utf8String partName = part.Name;
+				CbObjectId partId = part.AsObjectId();
+
+				Logger.LogInformation("Downloading part \'{PartName}\' with id {PartId}", partName, partId);
+				using HttpResponseMessage partResult = await httpClient.GetAsync(new Uri($"{buildId}/parts/{partId}", UriKind.Relative));
+				partResult.EnsureSuccessStatusCode();
+
+				byte[] partResponse = await partResult.Content.ReadAsByteArrayAsync();
+				CbObject partObject = new CbObject(partResponse);
+				CbObject filesObject = partObject["files"].AsObject();
+				List<string> paths = filesObject["paths"].AsArray().Select(field => field.AsString()).ToList();
+				List<long> rawsizes = filesObject["rawsizes"].AsArray().Select(field => field.AsInt64()).ToList();
+				List<IoHash> rawhashes = filesObject["rawhashes"].AsArray().Select(field => field.AsHash()).ToList();
+				//uint chunkSizeHint = filesObject["chunksizehint"].AsUInt32();
+
+				string attributeField = OperatingSystem.IsWindows() ? "attr" : "chmod";
+				List<int> attributes =  filesObject[attributeField].AsArray().Select(field => field.AsInt32()).ToList();
+
+				List<(string, long, IoHash, int)> foundFiles = [];
+				for (int i = 0; i < paths.Count; i++)
+				{
+					string path = paths[i];
+					long rawSize = rawsizes[i];
+					IoHash hash = rawhashes[i];
+					int attr = attributes[i];
+
+					foundFiles.Add((path, rawSize, hash, attr));
+				}
+
+				await Parallel.ForEachAsync(foundFiles, async (field, token) =>
+				{
+					(string fileName, long length, IoHash hash, int attr) = field;
+					string fullPath = Path.Combine(outputDir.FullName, fileName);
+					FileInfo outputFile = new FileInfo(fullPath);
+					await DownloadBlobAsync(httpClient, outputFile, buildId, partId, hash);
+					if (OperatingSystem.IsWindows())
+					{
+						File.SetAttributes(outputFile.FullName, (FileAttributes)attr);
+					}
+					else
+					{
+						File.SetUnixFileMode(outputFile.FullName, (UnixFileMode)attr);
+					}
+				});
+			}
+		}
+
+		private static async Task DownloadBlobAsync(HttpClient httpClient, FileInfo fullPath, CbObjectId buildId, CbObjectId partId, IoHash attachmentId)
+		{
+			using HttpRequestMessage request = new HttpRequestMessage(HttpMethod.Get, new Uri($"{buildId}/parts/{partId}/blobs/{attachmentId}", UriKind.Relative));
+			// accept all forms of blobs, we handle them based on the response mime type
+			request.Headers.Add("Accept", "*/*");
+
+			using HttpResponseMessage result = await httpClient!.SendAsync(request);
+			result.EnsureSuccessStatusCode();
+
+			// Make sure the directories are created if they are missing
+			fullPath.Directory?.Create();
+			string contentType = result.Content.Headers.ContentType!.ToString();
+			if (contentType == CustomMediaTypeNames.UnrealCompressedBuffer)
+			{
+				// compressed
+				long? contentLength = result.Content.Headers.ContentLength;
+				if (contentLength == null)
+				{
+					throw new Exception("Content length is missing");
+				}
+				await using Stream s = fullPath.OpenWrite();
+				await CompressedBuffer.DecompressContentAsync(await result.Content.ReadAsStreamAsync(), (ulong)contentLength.Value, s);
+			}
+			else if (contentType == MediaTypeNames.Application.Octet)
+			{
+				// uncompressed, just write to disk
+				await using Stream s = fullPath.OpenWrite();
+				await result.Content.CopyToAsync(s);
+			}
+			else
+			{
+				throw new Exception($"Unknown content type {contentType} when downloading blob {attachmentId}");
+			}
+		}
+
+		static async Task<List<CbObjectId>> SearchArtifactsAsync(HttpClient httpClient, CbObject queryObject)
+		{
+			using ByteArrayContent content = new ByteArrayContent(queryObject.GetView().ToArray());
+			content.Headers.ContentType = new MediaTypeHeaderValue(CustomMediaTypeNames.UnrealCompactBinary);
+
+			using HttpRequestMessage request = new HttpRequestMessage(HttpMethod.Post, new Uri($"search", UriKind.Relative));
+			request.Content = content;
+			request.Headers.Add("Accept", CustomMediaTypeNames.UnrealCompactBinary);
+			using HttpResponseMessage result = await httpClient!.SendAsync(request);
+			result.EnsureSuccessStatusCode();
+
+			SearchResult searchResult = await result.Content.ReadAsCompactBinaryAsync<SearchResult>();
+			return searchResult.Results.Select(result => result.BuildId).ToList();
+		}
+
+		/// <inheritdoc/>
+		public override void Write(XmlWriter writer)
+			=> Write(writer, _parameters);
+
+		/// <inheritdoc/>
+		public override IEnumerable<string> FindConsumedTagNames()
+			=> Enumerable.Empty<string>();
+
+		/// <inheritdoc/>
+		public override IEnumerable<string> FindProducedTagNames()
+			=> Enumerable.Empty<string>();
+	}
+
+	class BuildSearchResult
+	{
+		[CbField("buildId")]
+		public CbObjectId BuildId { get; set; }
+		[CbField("metadata")]
+		public CbObject Metadata { get; set; } = null!;
+	}
+
+	class SearchResult
+	{
+		[CbField("results")] public List<BuildSearchResult> Results { get; set; } = new List<BuildSearchResult>();
+
+		[CbField("partialResult")] public bool PartialResult { get; set; }
+
+		public SearchResult()
+		{
+			Results = new List<BuildSearchResult>();
+			PartialResult = false;
+		}
+	}
+}
diff --git a/Engine/Source/Programs/Shared/EpicGames.IoHash/IoHash.cs b/Engine/Source/Programs/Shared/EpicGames.IoHash/IoHash.cs
--- a/Engine/Source/Programs/Shared/EpicGames.IoHash/IoHash.cs
+++ b/Engine/Source/Programs/Shared/EpicGames.IoHash/IoHash.cs
@@ -176,7 +176,7 @@ namespace EpicGames.Core
 			{
 				Task Callback(ReadOnlyMemory<byte> data)
 				{
-					hasher.Update(data.Span);
+					hasher.UpdateWithJoin(data.Span);
 					return Task.CompletedTask;
 				}
 
diff --git a/Engine/Source/Programs/Shared/EpicGames.Serialization/CompressedBuffer.cs b/Engine/Source/Programs/Shared/EpicGames.Serialization/CompressedBuffer.cs
new file mode 100644
--- /dev/null
+++ b/Engine/Source/Programs/Shared/EpicGames.Serialization/CompressedBuffer.cs
@@ -0,0 +1,381 @@
+// Copyright Epic Games, Inc. All Rights Reserved.
+
+using System;
+using System.Buffers.Binary;
+using System.Collections.Generic;
+using System.IO;
+using System.IO.Hashing;
+using System.Linq;
+using System.Text;
+using System.Threading;
+using System.Threading.Tasks;
+using Blake3;
+using EpicGames.Compression;
+using EpicGames.Core;
+
+namespace EpicGames.Serialization
+{
+	public class CompressedBufferHeader
+	{
+		[System.Diagnostics.CodeAnalysis.SuppressMessage("Design", "CA1028:Enum Storage should be Int32", Justification = "Interop requires byte")]
+		public enum CompressionMethod : byte
+		{
+			// Header is followed by one uncompressed block. 
+			None = 0,
+			// Header is followed by an array of compressed block sizes then the compressed blocks. 
+			Oodle = 3,
+			LZ4 = 4,
+		}
+
+		public const uint ExpectedMagic = 0xb7756362; // <dot>ucb
+		public const uint HeaderLength = 64;
+
+		// A magic number to identify a compressed buffer. Always 0xb7756362.
+
+		public uint Magic { get; set; }
+		// A CRC-32 used to check integrity of the buffer. Uses the polynomial 0x04c11db7.
+
+		public uint Crc32 { get; set; }
+
+		// The method used to compress the buffer. Affects layout of data following the header. 
+		public CompressionMethod Method { get; set; }
+		public byte CompressionLevel { get; set; }
+		public byte CompressionMethodUsed { get; set; }
+
+		// The power of two size of every uncompressed block except the last. Size is 1 << BlockSizeExponent. 
+		public byte BlockSizeExponent { get; set; }
+
+		// The number of blocks that follow the header. 
+		public uint BlockCount { get; set; }
+
+		// The total size of the uncompressed data. 
+		public ulong TotalRawSize { get; set; }
+
+		// The total size of the compressed data including the header. 
+		public ulong TotalCompressedSize { get; set; }
+
+		/** The hash of the uncompressed data. */
+		public byte[] RawHash { get; set; } = Array.Empty<byte>();
+
+		public void ByteSwap()
+		{
+			Magic = BinaryPrimitives.ReverseEndianness(Magic);
+			Crc32 = BinaryPrimitives.ReverseEndianness(Crc32);
+			BlockCount = BinaryPrimitives.ReverseEndianness(BlockCount);
+			TotalRawSize = BinaryPrimitives.ReverseEndianness(TotalRawSize);
+			TotalCompressedSize = BinaryPrimitives.ReverseEndianness(TotalCompressedSize);
+		}
+	}
+
+	public static class CompressedBuffer
+	{
+		private static (CompressedBufferHeader, uint[]) ExtractHeader(BinaryReader br)
+		{
+			byte[] headerData = br.ReadBytes((int)CompressedBufferHeader.HeaderLength);
+
+			using MemoryStream ms = new MemoryStream(headerData);
+			using BinaryReader reader = new BinaryReader(ms);
+
+			// the header is always stored big endian
+			bool needsByteSwap = BitConverter.IsLittleEndian;
+
+			CompressedBufferHeader header = new CompressedBufferHeader
+			{
+				Magic = reader.ReadUInt32(),
+				Crc32 = reader.ReadUInt32(),
+				Method = (CompressedBufferHeader.CompressionMethod)reader.ReadByte(),
+				CompressionLevel = reader.ReadByte(),
+				CompressionMethodUsed = reader.ReadByte(),
+				BlockSizeExponent = reader.ReadByte(),
+				BlockCount = reader.ReadUInt32(),
+				TotalRawSize = reader.ReadUInt64(),
+				TotalCompressedSize = reader.ReadUInt64()
+			};
+			byte[] hash = reader.ReadBytes(32); // a full blake3 hash
+			header.RawHash = hash;
+
+			if (needsByteSwap)
+			{
+				header.ByteSwap();
+			}
+
+			if (header.Magic != CompressedBufferHeader.ExpectedMagic)
+			{
+				throw new InvalidMagicException(header.Magic, CompressedBufferHeader.ExpectedMagic);
+			}
+
+			// calculate the crc from the start of the method field (skipping magic which is a constant and the crc field itself)
+			const int MethodOffset = sizeof(uint) + sizeof(uint);
+
+			// none compressed objects have no extra blocks
+			uint blocksByteUsed = header.Method != CompressedBufferHeader.CompressionMethod.None ? header.BlockCount * (uint)sizeof(uint) : 0;
+
+			byte[] crcData = new byte[blocksByteUsed + headerData.Length];
+			Array.Copy(headerData, crcData, headerData.Length);
+
+			uint[] blocks = Array.Empty<uint>();
+
+			if (blocksByteUsed != 0)
+			{
+				byte[] blocksData = br.ReadBytes((int)blocksByteUsed);
+
+				Array.Copy(blocksData, 0, crcData, headerData.Length, blocksData.Length);
+
+				blocks = new uint[header.BlockCount];
+
+				for (int i = 0; i < header.BlockCount; i++)
+				{
+					ReadOnlySpan<byte> memory = new ReadOnlySpan<byte>(blocksData, i * sizeof(uint), sizeof(uint));
+					uint compressedBlockSize = BinaryPrimitives.ReadUInt32BigEndian(memory);
+					blocks[i] = compressedBlockSize;
+				}
+			}
+
+			uint calculatedCrc = Crc32.HashToUInt32(crcData.AsSpan(MethodOffset, (int)(CompressedBufferHeader.HeaderLength - MethodOffset + blocksByteUsed)));
+
+			if (header.Crc32 != calculatedCrc)
+			{
+				throw new InvalidHashException(header.Crc32, calculatedCrc);
+			}
+
+			return (header, blocks);
+		}
+
+		private static void WriteHeader(CompressedBufferHeader header, BinaryWriter writer)
+		{
+			// the header is always stored big endian
+			bool needsByteSwap = BitConverter.IsLittleEndian;
+			if (needsByteSwap)
+			{
+				header.ByteSwap();
+			}
+
+			writer.Write(header.Magic);
+			writer.Write(header.Crc32);
+			writer.Write((byte)header.Method);
+			writer.Write((byte)header.CompressionLevel);
+			writer.Write((byte)header.CompressionMethodUsed);
+			writer.Write((byte)header.BlockSizeExponent);
+			writer.Write(header.BlockCount);
+			writer.Write(header.TotalRawSize);
+			writer.Write(header.TotalCompressedSize);
+			writer.Write(header.RawHash, 0, 20); // write the first 20 bytes as iohashes are 20 bytes
+			for (int i = 0; i < 12; i++)
+			{
+				// the last 12 bytes should be 0 as they are reserved
+				writer.Write((byte)0);
+			}
+
+			if (needsByteSwap)
+			{
+				header.ByteSwap();
+			}
+		}
+
+		public static async Task DecompressContentAsync(Stream sourceStream, ulong streamSize, Stream targetStream, CancellationToken cancellationToken = default)
+		{
+			using BinaryReader br = new BinaryReader(sourceStream);
+			(CompressedBufferHeader header, uint[] compressedBlockSizes) = ExtractHeader(br);
+
+			if (streamSize < header.TotalCompressedSize)
+			{
+				throw new Exception($"Expected stream to be {header.TotalCompressedSize} but it was {streamSize}");
+			}
+
+			{
+				ulong decompressedPayloadOffset = 0;
+
+				bool willHaveBlocks = header.Method != CompressedBufferHeader.CompressionMethod.None;
+				if (willHaveBlocks)
+				{
+					ulong blockSize = 1ul << header.BlockSizeExponent;
+
+					foreach (uint compressedBlockSize in compressedBlockSizes)
+					{
+						ulong rawBlockSize = Math.Min(header.TotalRawSize - decompressedPayloadOffset, blockSize);
+						byte[] compressedPayload = br.ReadBytes((int)compressedBlockSize);
+
+						int writtenBytes;
+						// if a block has the same raw and compressed size its uncompressed and we should not attempt to decompress it
+						if (rawBlockSize == compressedBlockSize)
+						{
+							writtenBytes = (int)rawBlockSize;
+							targetStream.Write(compressedPayload);
+						}
+						else
+						{
+							writtenBytes = DecompressPayload(compressedPayload, header, rawBlockSize, targetStream);
+						}
+
+						decompressedPayloadOffset += (uint)writtenBytes;
+					}
+				}
+				else
+				{
+					await sourceStream.CopyToAsync(targetStream, cancellationToken);
+				}
+			}
+		}
+
+		private static int DecompressPayload(ReadOnlySpan<byte> compressedPayload, CompressedBufferHeader header, ulong rawBlockSize, Stream target)
+		{
+			switch (header.Method)
+			{
+				case CompressedBufferHeader.CompressionMethod.None:
+					target.Write(compressedPayload);
+					return compressedPayload.Length;
+				case CompressedBufferHeader.CompressionMethod.Oodle:
+					{
+						byte[] result = new byte[rawBlockSize];
+						long writtenBytes = Oodle.Decompress(compressedPayload, result);
+						if (writtenBytes == 0)
+						{
+							throw new Exception("Failed to run oodle decompress");
+						}
+						target.Write(result);
+						return (int)writtenBytes;
+					}
+				case CompressedBufferHeader.CompressionMethod.LZ4:
+					{
+						// Remove LZ4 support as it's not used and requires external dependency
+						/*byte[] result = new byte[rawBlockSize];
+						int writtenBytes = LZ4Codec.Decode(compressedPayload, result);
+						target.Write(result);
+						return writtenBytes;*/
+						throw new NotImplementedException("LZ4 is no longer supported");
+					}
+				default:
+					throw new NotImplementedException($"Method {header.Method} is not a support value");
+			}
+		}
+
+		public static IoHash CompressContent(Stream s, OodleCompressorType method, OodleCompressionLevel compressionLevel, byte[] rawContents)
+		{
+			const long DefaultBlockSize = 256 * 1024;
+			long blockSize = DefaultBlockSize;
+			long blockCount = (rawContents.LongLength + blockSize - 1) / blockSize;
+			Span<byte> contentsSpan = new Span<byte>(rawContents);
+			List<byte[]> blocks = new List<byte[]>();
+
+			for (int i = 0; i < blockCount; i++)
+			{
+				int rawBlockSize = Math.Min(rawContents.Length - (i * (int)blockSize), (int)blockSize);
+				Span<byte> bufferToCompress = contentsSpan.Slice((int)(i * blockSize), rawBlockSize);
+
+				blocks.Add(bufferToCompress.ToArray());
+			}
+
+			return CompressContent(s, method, compressionLevel, blocks, blockSize);
+		}
+
+		private static IoHash CompressContent(Stream s, OodleCompressorType method, OodleCompressionLevel compressionLevel, List<byte[]> blocks, long blockSize)
+		{
+			long blockCount = blocks.Count;
+
+			byte blockSizeExponent = (byte)Math.Floor(Math.Log2(blockSize));
+
+			List<byte[]> compressedBlocks = new List<byte[]>();
+			using Hasher hasher = Hasher.New();
+
+			ulong uncompressedContentLength = (ulong)blocks.Sum(b => b.LongLength);
+
+			ulong compressedContentLength = 0;
+			for (int i = 0; i < blockCount; i++)
+			{
+				int rawBlockSize = blocks[i].Length;
+				byte[] bufferToCompress = blocks[i];
+				hasher.UpdateWithJoin(new ReadOnlySpan<byte>(bufferToCompress, 0, rawBlockSize));
+				int maxSize = Oodle.MaximumOutputSize(method, rawBlockSize);
+				byte[] compressedBlock = new byte[maxSize];
+				long encodedSize = Oodle.Compress(method, bufferToCompress, compressedBlock, compressionLevel);
+
+				if (encodedSize == 0)
+				{
+					throw new Exception("Failed to compress content");
+				}
+
+				byte[] actualCompressedBlock = new byte[encodedSize];
+				Array.Copy(compressedBlock, actualCompressedBlock, encodedSize);
+				compressedBlocks.Add(actualCompressedBlock);
+				compressedContentLength += (ulong)encodedSize;
+			}
+
+			Hash blake3Hash = hasher.Finalize();
+			byte[] hashData = blake3Hash.AsSpan().Slice(0, 20).ToArray();
+			IoHash hash = new IoHash(hashData);
+
+			CompressedBufferHeader header = new CompressedBufferHeader
+			{
+				Magic = CompressedBufferHeader.ExpectedMagic,
+				Crc32 = 0,
+				Method = CompressedBufferHeader.CompressionMethod.Oodle,
+				CompressionLevel = (byte)compressionLevel,
+				CompressionMethodUsed = (byte)method,
+				BlockSizeExponent = blockSizeExponent,
+				BlockCount = (uint)blockCount,
+				TotalRawSize = (ulong)uncompressedContentLength,
+				TotalCompressedSize = (ulong)compressedContentLength,
+				RawHash = hashData
+			};
+
+			byte[] headerAndBlocks = WriteHeaderToBuffer(header, compressedBlocks.Select(b => (uint)b.Length).ToArray());
+
+			using BinaryWriter writer = new BinaryWriter(s, Encoding.Default, leaveOpen: true);
+
+			writer.Write(headerAndBlocks);
+
+			for (int i = 0; i < blockCount; i++)
+			{
+				writer.Write(compressedBlocks[i]);
+			}
+
+			return hash;
+		}
+
+		private static byte[] WriteHeaderToBuffer(CompressedBufferHeader header, uint[] compressedBlockLengths)
+		{
+			uint blockCount = header.BlockCount;
+			uint blocksByteUsed = blockCount * sizeof(uint);
+
+			byte[] headerBuffer = new byte[CompressedBufferHeader.HeaderLength + blocksByteUsed];
+
+			// write the compressed buffer, but with the wrong crc which we update and rewrite later
+			{
+				using MemoryStream ms = new MemoryStream(headerBuffer);
+				using BinaryWriter writer = new BinaryWriter(ms);
+
+				WriteHeader(header, writer);
+
+				for (int i = 0; i < blockCount; i++)
+				{
+					uint value = compressedBlockLengths[i];
+					if (BitConverter.IsLittleEndian)
+					{
+						value = BinaryPrimitives.ReverseEndianness(value);
+					}
+					writer.Write(value);
+				}
+			}
+
+			// calculate the crc from the start of the method field (skipping magic which is a constant and the crc field itself)
+			const int MethodOffset = sizeof(uint) + sizeof(uint);
+			uint calculatedCrc = Crc32.HashToUInt32(headerBuffer.AsSpan(MethodOffset, (int)(CompressedBufferHeader.HeaderLength - MethodOffset + blocksByteUsed)));
+			header.Crc32 = calculatedCrc;
+
+			// write the header again now that we have the crc
+			{
+				using MemoryStream ms = new MemoryStream(headerBuffer);
+				using BinaryWriter writer = new BinaryWriter(ms);
+				WriteHeader(header, writer);
+			}
+
+			return headerBuffer;
+		}
+	}
+
+	public class InvalidHashException(uint headerCrc32, uint calculatedCrc)
+		: Exception($"Header specified crc \"{headerCrc32}\" but calculated hash was \"{calculatedCrc}\"");
+
+	public class InvalidMagicException(uint headerMagic, uint expectedMagic)
+		: Exception($"Header magic \"{headerMagic}\" was incorrect, expected to be {expectedMagic}");
+}
diff --git a/Engine/Source/Programs/Shared/EpicGames.Serialization/EpicGames.Serialization.csproj b/Engine/Source/Programs/Shared/EpicGames.Serialization/EpicGames.Serialization.csproj
--- a/Engine/Source/Programs/Shared/EpicGames.Serialization/EpicGames.Serialization.csproj
+++ b/Engine/Source/Programs/Shared/EpicGames.Serialization/EpicGames.Serialization.csproj
@@ -36,7 +36,11 @@
       <Link>Properties\MetaData.cs</Link>
     </Compile>
   </ItemGroup>
+  <ItemGroup>
+    <PackageReference Include="System.IO.Hashing" Version="8.0.0" />
+  </ItemGroup>
   <ItemGroup>
     <ProjectReference Include="..\EpicGames.IoHash\EpicGames.IoHash.csproj" />
+    <ProjectReference Include="..\EpicGames.Oodle\EpicGames.Oodle.csproj" />
   </ItemGroup>
 </Project>
\ No newline at end of file
