Avoid failing on invalid UnrealApp instance.

#rnx

diff --git a/Engine/Source/Programs/AutomationTool/Gauntlet/Unreal/Base/Gauntlet.UnrealSession.cs b/Engine/Source/Programs/AutomationTool/Gauntlet/Unreal/Base/Gauntlet.UnrealSession.cs
--- a/Engine/Source/Programs/AutomationTool/Gauntlet/Unreal/Base/Gauntlet.UnrealSession.cs
+++ b/Engine/Source/Programs/AutomationTool/Gauntlet/Unreal/Base/Gauntlet.UnrealSession.cs
@@ -1093,7 +1093,7 @@ namespace Gauntlet
 						}
 						catch (Exception Exception)
 						{
-							Log.Info("Encountered an {Exception} when deleting source artifacts at {SourceDirectory}. Artifacts will remain on the device.", Exception, SourceDirectory);
+							Log.Info("Encountered an exception when deleting source artifacts at {SourceDirectory}. Artifacts will remain on the device.\n\tError message: {Exception}", SourceDirectory, Exception.Message);
 						}
 					}
 				}
diff --git a/Engine/Source/Programs/AutomationTool/Gauntlet/Unreal/Base/Gauntlet.UnrealTestNode.cs b/Engine/Source/Programs/AutomationTool/Gauntlet/Unreal/Base/Gauntlet.UnrealTestNode.cs
--- a/Engine/Source/Programs/AutomationTool/Gauntlet/Unreal/Base/Gauntlet.UnrealTestNode.cs
+++ b/Engine/Source/Programs/AutomationTool/Gauntlet/Unreal/Base/Gauntlet.UnrealTestNode.cs
@@ -1091,7 +1091,7 @@ namespace Gauntlet
 				Log.Warning("Failed to save artifacts. {Exception}", Ex);
 			}
 
-			if (CurrentPass + 1 == NumPasses)
+			if (UnrealApp != null && CurrentPass + 1 == NumPasses)
 			{
 				try
 				{
@@ -1651,7 +1651,12 @@ namespace Gauntlet
 		/// <returns></returns>
 		public virtual IEnumerable<UnrealRoleArtifacts> SaveRoleArtifacts(string OutputPath)
 		{
-			return UnrealApp.SaveRoleArtifacts(Context, TestInstance, ArtifactPath);
+			if (UnrealApp != null)
+			{
+				return UnrealApp.SaveRoleArtifacts(Context, TestInstance, ArtifactPath);
+			}
+
+			return Enumerable.Empty<UnrealRoleArtifacts>(); 
 		}
 
 		/// <summary>
