LaunchMultiServer UAT script: Allow attaching to client when running proxy server
- Also fixed issue with missing space when using -clientcmdline with proxy

diff --git a/Engine/Source/Programs/AutomationTool/Scripts/LaunchMultiServer.Automation.cs b/Engine/Source/Programs/AutomationTool/Scripts/LaunchMultiServer.Automation.cs
--- a/Engine/Source/Programs/AutomationTool/Scripts/LaunchMultiServer.Automation.cs
+++ b/Engine/Source/Programs/AutomationTool/Scripts/LaunchMultiServer.Automation.cs
@@ -413,9 +413,15 @@ namespace AutomationTool
 					var Args = CommonArgs;
 
 					// The server to connect to should be provided in the AdditionalParameters (e.g. 127.0.0.1:8000).
-					Args += AdditionalParameters + (Params.Cook ? "" : "-game ");
+					Args += AdditionalParameters + " ";
+					Args += (Params.Cook ? "" : "-game ");
 					Args += Params.ClientCommandline + " ";
 					Args += "-windowed ";
+					
+					if (IsAttachingDebugger(LaunchType))
+					{
+						Args += "-WaitForDebuggerNoBreak ";
+					}
 
 					var NewProcess = Run(App, Args, null, ERunOptions.Default | ERunOptions.NoWaitForExit | ERunOptions.NoStdOutRedirect);
 
@@ -428,6 +434,11 @@ namespace AutomationTool
 
 					// Pause between starting processes to enforce startup determinism.
 					System.Threading.Thread.Sleep(PauseBetweenProcessMS);
+
+					if (IsAttachingDebugger(LaunchType))
+					{
+						ProcessDebugger.DebugProcess(NewProcess.ProcessObject.Id);
+					}
 				}
 			}
 			catch
@@ -443,7 +454,7 @@ namespace AutomationTool
 		private bool IsAttachingDebugger(ProcessLaunchType LaunchType)
 		{
 			return (ParseParam("attachtoservers") && LaunchType == ProcessLaunchType.Server) ||
-				(ParseParam("attachtoclients") && LaunchType == ProcessLaunchType.DirectClient);
+				(ParseParam("attachtoclients") && (LaunchType == ProcessLaunchType.DirectClient || LaunchType == ProcessLaunchType.ProxyClient));
 		}
 
 		private FileReference ProjectFullPath;
