LaunchMultiServer UAT script: Allow -notimeouts to work with proxy client and server

diff --git a/Engine/Source/Programs/AutomationTool/Scripts/LaunchMultiServer.Automation.cs b/Engine/Source/Programs/AutomationTool/Scripts/LaunchMultiServer.Automation.cs
--- a/Engine/Source/Programs/AutomationTool/Scripts/LaunchMultiServer.Automation.cs
+++ b/Engine/Source/Programs/AutomationTool/Scripts/LaunchMultiServer.Automation.cs
@@ -386,6 +386,12 @@ namespace AutomationTool
 					if (ServerInfos.Count() > 0)
 					{
 						Args += String.Format("-ConsoleTitle=\"{0} Proxy Server\" ", Params.ShortProjectName);
+						
+						if (ParseParam("notimeouts"))
+						{
+							Args += "-notimeouts ";
+						}
+
 						Args += "-ProxyGameServers=";
 
 						foreach (var ServerInfo in ServerInfos.Values)
@@ -418,6 +424,11 @@ namespace AutomationTool
 					Args += Params.ClientCommandline + " ";
 					Args += "-windowed ";
 					
+					if (ParseParam("notimeouts"))
+					{
+						Args += "-notimeouts ";
+					}
+
 					if (IsAttachingDebugger(LaunchType))
 					{
 						Args += "-WaitForDebuggerNoBreak ";
