UAT: Do not warn when copying files and all files already exist at the destination and overwriting was not requested. This is expected behavior and is now just a loginformation

diff --git a/Engine/Source/Programs/AutomationTool/BuildGraph/Tasks/CopyTask.cs b/Engine/Source/Programs/AutomationTool/BuildGraph/Tasks/CopyTask.cs
--- a/Engine/Source/Programs/AutomationTool/BuildGraph/Tasks/CopyTask.cs
+++ b/Engine/Source/Programs/AutomationTool/BuildGraph/Tasks/CopyTask.cs
@@ -190,7 +190,7 @@ namespace AutomationTool.Tasks
 				}
 				if (filteredTargetToSourceFile.Count == 0)
 				{
-					Logger.LogWarning("All files already exist, exiting early.");
+					Logger.LogInformation("All files already exist, exiting early.");
 					return;
 				}
 				targetFileToSourceFile = filteredTargetToSourceFile;
