APT: Removing platform specific folders from paths as it's not needed any longer.

diff --git a/Engine/Plugins/Performance/AutomatedPerfTesting/Build/Scripts/AutomatedPerfTestNode.cs b/Engine/Plugins/Performance/AutomatedPerfTesting/Build/Scripts/AutomatedPerfTestNode.cs
--- a/Engine/Plugins/Performance/AutomatedPerfTesting/Build/Scripts/AutomatedPerfTestNode.cs
+++ b/Engine/Plugins/Performance/AutomatedPerfTesting/Build/Scripts/AutomatedPerfTestNode.cs
@@ -350,13 +350,9 @@ namespace AutomatedPerfTest
 			{
 			 	perfreportExe = "PerfreportTool.exe";
 			}
-			else if (OperatingSystem.IsLinux())
+			else if (OperatingSystem.IsLinux() || OperatingSystem.IsMacOS())
 			{
-				perfreportExe = Path.Combine("Linux", "PerfreportTool.dll");
-			}
-			else if (OperatingSystem.IsMacOS())
-			{
-				perfreportExe = Path.Combine("Mac", "PerfreportTool.dll");
+				perfreportExe = Path.Combine("PerfreportTool.dll");
 			}
 
 			var ToolPath = FileReference.Combine(Unreal.EngineDirectory, "Binaries", "DotNET", "CsvTools", perfreportExe);
diff --git a/Engine/Source/Programs/AutomationTool/Gauntlet/Framework/Utils/Gauntlet.ReportGenUtils.cs b/Engine/Source/Programs/AutomationTool/Gauntlet/Framework/Utils/Gauntlet.ReportGenUtils.cs
--- a/Engine/Source/Programs/AutomationTool/Gauntlet/Framework/Utils/Gauntlet.ReportGenUtils.cs
+++ b/Engine/Source/Programs/AutomationTool/Gauntlet/Framework/Utils/Gauntlet.ReportGenUtils.cs
@@ -79,14 +79,6 @@ namespace Gauntlet
 		public static string GetCsvToolPath(string ToolName)
 		{
 			string Extension = BuildHostPlatform.Current.Platform.IsInGroup(UnrealPlatformGroup.Microsoft) ? ".exe" : ".dll";
-			if (BuildHostPlatform.Current.Platform.IsInGroup(UnrealPlatformGroup.Apple))
-			{
-				return Path.Combine(CommandUtils.CmdEnv.LocalRoot, "Engine", "Binaries", "DotNet", "CsvTools", "Mac",  ToolName) + Extension;
-			}
-			if (BuildHostPlatform.Current.Platform.IsInGroup(UnrealPlatformGroup.Unix))
-			{
-				return Path.Combine(CommandUtils.CmdEnv.LocalRoot, "Engine", "Binaries", "DotNet","CSVTools", "Linux", ToolName) + Extension;
-			}
 			return Path.Combine(CommandUtils.CmdEnv.LocalRoot, "Engine", "Binaries", "DotNet", "CsvTools", ToolName) + Extension;
 		}
 
