Fix failing to run test after acquiring new device

#rnx

diff --git a/Engine/Source/Programs/AutomationTool/Gauntlet/Unreal/Base/Gauntlet.UnrealSession.cs b/Engine/Source/Programs/AutomationTool/Gauntlet/Unreal/Base/Gauntlet.UnrealSession.cs
--- a/Engine/Source/Programs/AutomationTool/Gauntlet/Unreal/Base/Gauntlet.UnrealSession.cs
+++ b/Engine/Source/Programs/AutomationTool/Gauntlet/Unreal/Base/Gauntlet.UnrealSession.cs
@@ -11,6 +11,7 @@ using System.Threading;
 using UnrealBuildTool;
 using Gauntlet.Utils;
 using System.Net;
+using Microsoft.Extensions.Logging;
 
 namespace Gauntlet
 {
@@ -1666,6 +1667,10 @@ namespace Gauntlet
 						IAppInstall Install = Device.CreateAppInstall(AppConfig);
 						RolesToInstalls.Add(Role, Install);
 					}
+					else
+					{
+						AppConfig = RolesToConfigs[Role];
+					}
 
 					Device.CleanArtifacts();
 
@@ -1680,23 +1685,29 @@ namespace Gauntlet
 				{
 					string Message = $"Encountered {Ex.GetType()} when creating installation on device {Device}.\n{Ex.Message}";
 
-					if (IsOutOfSpaceException(Ex) && Device.Platform == BuildHostPlatform.Current.Platform)
+					bool IsHostPlatform = Device.Platform == BuildHostPlatform.Current.Platform;
+					EventId EventType = IsHostPlatform ? KnownLogEvents.Gauntlet : KnownLogEvents.Gauntlet_DeviceEvent;
+
+					if (IsOutOfSpaceException(Ex) && IsHostPlatform)
 					{
 						// If on desktop platform, we are not retrying.
 						// It is unlikely that space is going to be made and InstallBuildParallel has marked the build path as problematic.
-						Log.Error(KnownLogEvents.Gauntlet_DeviceEvent, Message);
+						Log.Error(EventType, Message);
 						throw;
 					}
 					else if(IsOverlayException(Ex))
 					{
 						// Errors with Overlay executables are caused by missing files or improper setup, inform user and exit early
-						Log.Error(KnownLogEvents.Gauntlet_DeviceEvent, Message);
+						Log.Error(EventType, Message);
 						throw;
 					}
 					else
 					{
-						UnrealDeviceReservation.MarkProblemDevice(Device, Message);
-						Log.Warning(KnownLogEvents.Gauntlet_DeviceEvent, Message);
+						if (!IsHostPlatform)
+						{
+							UnrealDeviceReservation.MarkProblemDevice(Device, Message);
+						}
+						Log.Warning(EventType, Message);
 					}
 					throw; // Can consider not throwing here - this would let every device complete setup before releasing the problem devices
 				}
