Fix invalid success when running several UnrealTestNode with the same name.
Avoid using a static mapping with test name key to collect Unreal Session.
Reusing the Unreal Session to keep the same device reservation for multiple itiration is already assure by the reservation shutdown that check for current itiration index.

#rnx

diff --git a/Engine/Source/Programs/AutomationTool/Gauntlet/Unreal/Base/Gauntlet.UnrealTestNode.cs b/Engine/Source/Programs/AutomationTool/Gauntlet/Unreal/Base/Gauntlet.UnrealTestNode.cs
--- a/Engine/Source/Programs/AutomationTool/Gauntlet/Unreal/Base/Gauntlet.UnrealTestNode.cs
+++ b/Engine/Source/Programs/AutomationTool/Gauntlet/Unreal/Base/Gauntlet.UnrealTestNode.cs
@@ -238,39 +238,9 @@ namespace Gauntlet
 		}
 
 		/// <summary>
-		/// Dictionary of Sessions for each test, by test name
+		/// UnrealSession for this instance
 		/// </summary>
-		static protected Dictionary<string, UnrealSession> TestSessions = new Dictionary<string, UnrealSession>();
-
-		/// <summary>
-		/// Helper class that turns our wishes into reality
-		/// Indexed from the dictionary of Sessions by Test Name
-		/// </summary>
-		protected UnrealSession UnrealApp
-		{
-			get
-			{
-				if (TestSessions.ContainsKey(Name))
-				{
-					return TestSessions[Name];
-				}
-				else
-				{
-					return null;
-				}
-			}
-			set
-			{
-				if (TestSessions.ContainsKey(Name))
-				{
-					TestSessions[Name] = value;
-				}
-				else
-				{
-					TestSessions.Add(Name, value);
-				}
-			}
-		}
+		protected UnrealSession UnrealApp { get; set; }
 
 		/// <summary>
 		/// Used to track each running app's new log output
@@ -556,15 +526,7 @@ namespace Gauntlet
 
 			if (UnrealApp != null)
 			{
-				if (NumPasses > 1 && CurrentPass + 1 < NumPasses)
-				{
-					// Already prepared from a previous pass
-					return true;
-				}
-				else
-				{
-					throw new AutomationException("Node already has an UnrealApp, was PrepareUnrealSession called twice?");
-				}
+				throw new AutomationException("Node already has an UnrealApp, was PrepareUnrealSession called twice?");
 			}
 
 			// pass through any arguments such as -TestNameArg or -TestNameArg=Value
@@ -607,7 +569,7 @@ namespace Gauntlet
 
 				foreach (UnrealTestRole TestRole in TypesToRoles.Value)
 				{
-					// If a config has overriden a platform then we can't use the context constraints from the commandline
+					// If a config has overridden a platform then we can't use the context constraints from the commandline
 					bool UseContextConstraint = TestRole.Type == UnrealTargetRole.Client && TestRole.PlatformOverride == null;
 
 					// important, use the type from the ContextRolke because Server may have been mapped to EditorServer etc
