ReserveDevicesFromService all return paths return false.

#rnx

diff --git a/Engine/Source/Programs/AutomationTool/Gauntlet/Framework/Devices/Gauntlet.DevicePool.cs b/Engine/Source/Programs/AutomationTool/Gauntlet/Framework/Devices/Gauntlet.DevicePool.cs
--- a/Engine/Source/Programs/AutomationTool/Gauntlet/Framework/Devices/Gauntlet.DevicePool.cs
+++ b/Engine/Source/Programs/AutomationTool/Gauntlet/Framework/Devices/Gauntlet.DevicePool.cs
@@ -904,8 +904,11 @@ namespace Gauntlet
 				}
 			}
 
-			Log.Info(KnownLogEvents.Gauntlet_DeviceEvent, "Unable to reserve service devices:");
-			return false;
+			if (ReservedDevices.Count == 0)
+			{
+				Log.Info(KnownLogEvents.Gauntlet_DeviceEvent, "Unable to reserve all devices from service.");
+			}
+			return ReservedDevices.Count > 0;
 		}
 
 		/// <summary>
