Change Smartling messages' wording

#rnx

diff --git a/Engine/Source/Programs/AutomationTool/SmartlingLocalization/SmartlingLocalizationProvider.cs b/Engine/Source/Programs/AutomationTool/SmartlingLocalization/SmartlingLocalizationProvider.cs
--- a/Engine/Source/Programs/AutomationTool/SmartlingLocalization/SmartlingLocalizationProvider.cs
+++ b/Engine/Source/Programs/AutomationTool/SmartlingLocalization/SmartlingLocalizationProvider.cs
@@ -274,7 +274,7 @@ namespace EpicGames.SmartlingLocalization
 			var RequestErrors = ResponseEnvelope.Response.Errors;
 			foreach (SmartlingRequestError RequestError in RequestErrors)
 			{
-				Logs.Add($"Smartling Warning:\n {RequestError.ToString()}");
+				Logs.Add($"Smartling Info:\n {RequestError.ToString()}");
 			}
 			
 		}
