Use appropriate log level depending on test exit code

#rnx

diff --git a/Engine/Source/Programs/AutomationTool/LowLevelTests/Tests/Gauntlet.LowLevelTestNode.cs b/Engine/Source/Programs/AutomationTool/LowLevelTests/Tests/Gauntlet.LowLevelTestNode.cs
--- a/Engine/Source/Programs/AutomationTool/LowLevelTests/Tests/Gauntlet.LowLevelTestNode.cs
+++ b/Engine/Source/Programs/AutomationTool/LowLevelTests/Tests/Gauntlet.LowLevelTestNode.cs
@@ -318,7 +318,15 @@ namespace LowLevelTests
 						ExitReason = "Tests passed (no report to parse)";
 					}
 				}
-				Log.Info($"Low level test exited with code {ExitCode} and reason: {ExitReason}");
+				string ExitMessage = $"Low level test exited with code {ExitCode} and reason: {ExitReason}";
+				if (LowLevelTestResult != TestResult.Passed)
+				{
+					Log.Error(ExitMessage);
+				}
+				else
+				{
+					Log.Info(ExitMessage);
+				}
 			}
 			catch
 			{
