SignTask comment fix.

diff --git a/Engine/Source/Programs/AutomationTool/BuildGraph/Tasks/SignTask.cs b/Engine/Source/Programs/AutomationTool/BuildGraph/Tasks/SignTask.cs
--- a/Engine/Source/Programs/AutomationTool/BuildGraph/Tasks/SignTask.cs
+++ b/Engine/Source/Programs/AutomationTool/BuildGraph/Tasks/SignTask.cs
@@ -10,7 +10,7 @@ using UnrealBuildBase;
 namespace AutomationTool.Tasks
 {
 	/// <summary>
-	/// Parameters for a task that strips symbols from a set of files
+	/// Parameters for a task that signs a set of files.
 	/// </summary>
 	public class SignTaskParameters
 	{
