Fix variable shadowing in device pool reservation

#rnx

diff --git a/Engine/Source/Programs/AutomationTool/Gauntlet/Framework/Devices/Gauntlet.DevicePool.cs b/Engine/Source/Programs/AutomationTool/Gauntlet/Framework/Devices/Gauntlet.DevicePool.cs
--- a/Engine/Source/Programs/AutomationTool/Gauntlet/Framework/Devices/Gauntlet.DevicePool.cs
+++ b/Engine/Source/Programs/AutomationTool/Gauntlet/Framework/Devices/Gauntlet.DevicePool.cs
@@ -828,6 +828,7 @@ namespace Gauntlet
 			}
 
 			// Second pass, reserve devices
+			List<ITargetDevice> ScopeReservedDevices = new List<ITargetDevice>();
 			foreach (KeyValuePair<UnrealDeviceTargetConstraint, int> Entry in DeviceTypes)
 			{
 				string[] DeviceRequirements = Enumerable.Repeat(Entry.Key.FormatWithIdentifier(), Entry.Value).ToArray();
@@ -855,7 +856,6 @@ namespace Gauntlet
 				}
 
 				// Add target devices from reservation
-				List<ITargetDevice> ReservedDevices = new List<ITargetDevice>();
 				foreach (Device Device in DeviceReservation.Devices)
 				{
 					DeviceDefinition Def = new DeviceDefinition();
@@ -880,9 +880,9 @@ namespace Gauntlet
 						ReportDeviceError(Device.Name, "CreateDeviceError");
 
 						// If some devices from reservation have been created, release them which will also dispose of reservation
-						if (ReservedDevices.Count > 0)
+						if (ScopeReservedDevices.Count > 0)
 						{
-							ReleaseDevices(ReservedDevices);
+							ReleaseDevices(ScopeReservedDevices);
 						}
 						else
 						{
@@ -895,6 +895,7 @@ namespace Gauntlet
 					}
 					else
 					{
+						ScopeReservedDevices.Add(TargetDevice);
 						ReservedDevices.Add(TargetDevice);
 						UsedPlatforms.Add(TargetDevice.Platform);
 					}
@@ -904,11 +905,11 @@ namespace Gauntlet
 				}
 			}
 
-			if (ReservedDevices.Count == 0)
+			if (ScopeReservedDevices.Count == 0)
 			{
 				Log.Info(KnownLogEvents.Gauntlet_DeviceEvent, "Unable to reserve all devices from service.");
 			}
-			return ReservedDevices.Count > 0;
+			return ScopeReservedDevices.Count > 0;
 		}
 
 		/// <summary>
