[DSTM] Added support to manage via session frontend, all launched servers and clients via the Launch Multi Server UAT script.

This CL also adds a new commandline argument to opt-out of launching new console automatically.

diff --git a/Engine/Source/Programs/AutomationTool/Scripts/LaunchMultiServer.Automation.cs b/Engine/Source/Programs/AutomationTool/Scripts/LaunchMultiServer.Automation.cs
--- a/Engine/Source/Programs/AutomationTool/Scripts/LaunchMultiServer.Automation.cs
+++ b/Engine/Source/Programs/AutomationTool/Scripts/LaunchMultiServer.Automation.cs
@@ -314,9 +314,16 @@ namespace AutomationTool
 			var CommonArgs = (Params.Cook ? "" : DeployContext.ProjectArgForCommandLines + " ");
 			if (LaunchType == ProcessLaunchType.Server || LaunchType == ProcessLaunchType.ProxyServer)
 			{
-				CommonArgs += Params.MapToRun + " -server -log -newconsole ";
+				CommonArgs += Params.MapToRun + " -server -log ";
+
+				if (!ParseParam("nonewconsole"))
+				{
+					CommonArgs += " -newconsole ";
+				}
 			}
 
+			CommonArgs += " -messaging ";
+
 			PushDir(Path.GetDirectoryName(App));
 
 			try
@@ -333,12 +340,13 @@ namespace AutomationTool
 						if (LaunchType == ProcessLaunchType.Server)
 						{
 							Args += String.Format("-MultiServerLocalId={0} -MultiServerListenPort={1} -port={2} ", ServerInfo.Definition.LocalId, ServerInfo.Definition.ListenPort, ServerInfo.GameListenPort);
-							Args += String.Format("-ConsoleTitle=\"{0} Server ID {1}\" ", Params.ShortProjectName, ServerInfo.Definition.LocalId);
+							Args += String.Format("-ConsoleTitle=\"{0} Server ID {1}\" -SessionName=\"{0} Server ID {1}\" ", Params.ShortProjectName, ServerInfo.Definition.LocalId);
 							Args += Params.ServerCommandline + " ";
 						}
 						else if (LaunchType == ProcessLaunchType.DirectClient)
 						{
 							Args += String.Format("127.0.0.1:{0} ", ServerInfo.GameListenPort) + (Params.Cook ? "" : "-game ");
+							Args += String.Format(" -SessionName=\"{0} Direct Client | Server ID {1}\" ", Params.ShortProjectName, ServerInfo.Definition.LocalId);
 							Args += Params.ClientCommandline + " ";
 
 							// Save a separate window position for each client
@@ -385,7 +393,7 @@ namespace AutomationTool
 
 					if (ServerInfos.Count() > 0)
 					{
-						Args += String.Format("-ConsoleTitle=\"{0} Proxy Server\" ", Params.ShortProjectName);
+						Args += String.Format("-ConsoleTitle=\"{0} Proxy Server\" -SessionName=\"{0} Proxy Server\" ", Params.ShortProjectName);
 						
 						if (ParseParam("notimeouts"))
 						{
@@ -424,6 +432,8 @@ namespace AutomationTool
 					Args += Params.ClientCommandline + " ";
 					Args += "-windowed ";
 					
+					Args += String.Format(" -SessionName=\"{0} Proxy Client \" ", Params.ShortProjectName);
+					
 					if (ParseParam("notimeouts"))
 					{
 						Args += "-notimeouts ";
