Improved logging in retrieve artifact task

diff --git a/Engine/Source/Programs/AutomationTool/BuildGraph/Tasks/RetrieveCloudArtifactTask.cs b/Engine/Source/Programs/AutomationTool/BuildGraph/Tasks/RetrieveCloudArtifactTask.cs
--- a/Engine/Source/Programs/AutomationTool/BuildGraph/Tasks/RetrieveCloudArtifactTask.cs
+++ b/Engine/Source/Programs/AutomationTool/BuildGraph/Tasks/RetrieveCloudArtifactTask.cs
@@ -317,7 +317,7 @@ namespace AutomationTool.Tasks
 
 			if (artifactIds.Count == 0)
 			{
-				throw new AutomationException("Unable to find any artifact matching given criteria.");
+				throw new AutomationException($"Unable to find any artifact matching given criteria in namespace {ns} and bucket {bucketId}.");
 			}
 
 			if (!_parameters.AllowMultipleMatches && artifactIds.Count != 1)
