Build Graph - Fixed issue with CreateCloudArtifact task which added the job metadata value multiple times with different values. As searches can easily filter on multiple values at the same time there is no need for this to use the same stucture as in CreateArtifact, instead we just simply add the job and step identifiers as seperate values.

diff --git a/Engine/Source/Programs/AutomationTool/BuildGraph/Tasks/CreateCloudArtifactTask.cs b/Engine/Source/Programs/AutomationTool/BuildGraph/Tasks/CreateCloudArtifactTask.cs
--- a/Engine/Source/Programs/AutomationTool/BuildGraph/Tasks/CreateCloudArtifactTask.cs
+++ b/Engine/Source/Programs/AutomationTool/BuildGraph/Tasks/CreateCloudArtifactTask.cs
@@ -172,7 +172,7 @@ namespace AutomationTool.Tasks
 				string? stepId = Environment.GetEnvironmentVariable("UE_HORDE_STEPID");
 				if (!String.IsNullOrEmpty(stepId))
 				{
-					metadata.Add($"job:{jobId}/step:{stepId}");
+					metadata.Add($"step:{stepId}");
 				}
 			}
 
