APT: make C# code compilation not require a .csproj in the project being tested.

- UAT will now also recurse in Builds directories of plugins (if they happen to have them).

diff --git a/Engine/Plugins/Performance/AutomatedPerfTesting/Build/Scripts/AutomatedPerfTesting.Automation.csproj b/Engine/Plugins/Performance/AutomatedPerfTesting/Build/Scripts/AutomatedPerfTesting.Automation.csproj
new file mode 100644
--- /dev/null
+++ b/Engine/Plugins/Performance/AutomatedPerfTesting/Build/Scripts/AutomatedPerfTesting.Automation.csproj
@@ -0,0 +1,31 @@
+<Project Sdk="Microsoft.NET.Sdk">
+  <Import Project="AsisNovaMaps.AutoPerfTest.csproj.props" Condition="Exists('AsisNovaMaps.AutoPerfTest.csproj.props')" />
+  <Import Project="$(EngineDir)\Source\Programs\Shared\UnrealEngine.csproj.props" />
+  <PropertyGroup>
+    <TargetFramework>net8.0</TargetFramework>
+    <OutputType>Library</OutputType>
+    <AppendTargetFrameworkToOutputPath>false</AppendTargetFrameworkToOutputPath>
+    <GenerateTargetFrameworkAttribute>false</GenerateTargetFrameworkAttribute>
+    <Configurations>Debug;Release;Development</Configurations>
+    <RootNamespace>AsisNovaMaps.AutoPerfTest</RootNamespace>
+    <AssemblyName>AsisNovaMaps.AutoPerfTest</AssemblyName>
+    <OutputPath>..\..\Binaries\DotNET\AutomationTool\AutomationScripts</OutputPath>
+    <TreatWarningsAsErrors>true</TreatWarningsAsErrors>
+    <WarningsNotAsErrors>612,618</WarningsNotAsErrors>
+    <DebugType>pdbonly</DebugType>
+  </PropertyGroup>
+  <ItemGroup>
+    <Compile Include="$(EngineDir)\Plugins\Performance\AutomatedPerfTesting\Build\Scripts\AutomatedPerfTestConfig.cs" />
+    <Compile Include="$(EngineDir)\Plugins\Performance\AutomatedPerfTesting\Build\Scripts\AutomatedPerfTestNode.cs" />
+  </ItemGroup>
+  <ItemGroup>
+    <ProjectReference Include="$(EngineDir)\Source\Programs\AutomationTool\AutomationUtils\AutomationUtils.Automation.csproj" />
+    <ProjectReference Include="$(EngineDir)\Source\Programs\Shared\EpicGames.Core\EpicGames.Core.csproj" PrivateAssets="All">
+      <Private>false</Private>
+    </ProjectReference>
+    <ProjectReference Include="$(EngineDir)\Source\Programs\AutomationTool\Gauntlet\Gauntlet.Automation.csproj" PrivateAssets="All">
+      <Private>false</Private>
+    </ProjectReference>
+    <ProjectReference Include="$(EngineDir)\Source\Programs\UnrealBuildTool\UnrealBuildTool.csproj" />
+  </ItemGroup>
+</Project>
diff --git a/Engine/Source/Programs/AutomationTool/Scripts/AddAutomatedPerfTestToProject.cs b/Engine/Source/Programs/AutomationTool/Scripts/AddAutomatedPerfTestToProject.cs
--- a/Engine/Source/Programs/AutomationTool/Scripts/AddAutomatedPerfTestToProject.cs
+++ b/Engine/Source/Programs/AutomationTool/Scripts/AddAutomatedPerfTestToProject.cs
@@ -55,11 +55,6 @@ class AddAutomatedPerfTestToProject : BuildCommand
 	/// </summary>
 	const string AptLocalTestShellFile = "RunLocalTests.sh";
 	
-	/// <summary>
-	/// Name of the dummy C# project file to ensure APT classes are included in builds
-	/// </summary>
-	const string AptCSProj = "PROJECT.AutoPerfTest.csproj";
-	
 	/// <summary>
 	/// Path to the engine's build/batchfiles directory relative to the engine root
 	/// </summary>
@@ -229,37 +224,6 @@ class AddAutomatedPerfTestToProject : BuildCommand
 		return true;
 	}
 
-	protected bool ModifyCSProj(string ProjectFilePath, bool Overwrite=false)
-	{
-		// work out the target directory
-		string ProjectDir = Path.GetDirectoryName(ProjectFilePath);
-
-		string TemplatesFolder = Path.Combine(AptPathFromRoot, AptTemplatesFromPluginRoot);
-
-		string CSProjDest = Path.Combine(ProjectDir, "Build", "Scripts");
-
-		string CSProjFileName = GetProjectName(ProjectFilePath) + ".AutoPerfTest.csproj";
-		
-		// copy example per-project csproj  file
-		Directory.CreateDirectory(CSProjDest);
-		{
-			string DestFile = Path.Combine(CSProjDest,CSProjFileName);
-			if(File.Exists(DestFile) && !Overwrite)
-			{
-				Logger.LogInformation("Project already has file {0}.",DestFile);
-			}
-			else
-			{
-				string SrcFile = Path.Combine(TemplatesFolder,AptCSProj);
-				string Template = File.ReadAllText(SrcFile);
-				Template = Template.Replace("**REPLACE_PROJECTNAME**", GetProjectName(ProjectFilePath));
-				Template = Template.Replace("**REPLACE_PROJECTPATH**",GetProjectRelativePath(ProjectFilePath));
-				File.WriteAllText(DestFile, Template);
-			}
-		}
-		return true;
-	}
-	
 	protected bool ModifyLocalTestsBatchFile(string ProjectFilePath, bool Overwrite=false)
 	{
 		// work out the target directory
@@ -269,7 +233,7 @@ class AddAutomatedPerfTestToProject : BuildCommand
 
 		string BatchFileDest = Path.Combine(ProjectDir, "Build", "BatchFiles");
 		
-		// copy example per-project csproj  file
+		// copy example batch files
 		Directory.CreateDirectory(BatchFileDest);
 		{
 			List<string> BatchFiles = [AptLocalTestBatchFile, AptLocalTestShellFile];
@@ -313,7 +277,6 @@ class AddAutomatedPerfTestToProject : BuildCommand
 		EnablePlugin(Project);
 		ModifyProjectConfigFiles(Project);
 		ModifyBuildGraphFiles(Project, Overwrite);
-		ModifyCSProj(Project, Overwrite);
 		ModifyLocalTestsBatchFile(Project, Overwrite);
 
 		Logger.LogInformation("Successfully added.");
diff --git a/Engine/Source/Programs/Shared/EpicGames.Build/System/Rules.cs b/Engine/Source/Programs/Shared/EpicGames.Build/System/Rules.cs
--- a/Engine/Source/Programs/Shared/EpicGames.Build/System/Rules.cs
+++ b/Engine/Source/Programs/Shared/EpicGames.Build/System/Rules.cs
@@ -308,6 +308,13 @@ namespace UnrealBuildBase
 				{
 					Queue.Enqueue(() => FindAllRulesFilesRecursively(SourceDir, Cache, Queue));
 				}
+
+				// search Build directory as well as it can contain automation projects
+				DirectoryItem? BuildDir;
+				if (Directory.TryGetDirectory("Build", out BuildDir))
+				{
+					Queue.Enqueue(() => FindAllRulesFilesRecursively(BuildDir, Cache, Queue));
+				}
 			}
 
 			// If we didn't find anything to stop the search, search all the subdirectories too
