Consolidate device failure by throwing DeviceException or using Gauntlet_DeviceEvent event id for logging
Lower log severity for device, instead use Report Device Error end point to funnel related message

#rnx

diff --git a/Engine/Source/Programs/AutomationTool/Gauntlet/Framework/Devices/Gauntlet.DeviceFirmwareValidator.cs b/Engine/Source/Programs/AutomationTool/Gauntlet/Framework/Devices/Gauntlet.DeviceFirmwareValidator.cs
--- a/Engine/Source/Programs/AutomationTool/Gauntlet/Framework/Devices/Gauntlet.DeviceFirmwareValidator.cs
+++ b/Engine/Source/Programs/AutomationTool/Gauntlet/Framework/Devices/Gauntlet.DeviceFirmwareValidator.cs
@@ -1,5 +1,6 @@
 // Copyright Epic Games, Inc. All Rights Reserved.
 
+using EpicGames.Core;
 using System;
 using System.Linq;
 using UnrealBuildTool;
@@ -38,17 +39,17 @@ namespace Gauntlet
 			bEnabled = bEnabled && (!string.IsNullOrEmpty(FirmwareVersion) || !string.IsNullOrEmpty(ProjectName));
 		}
 
-		public bool TryValidateDevice(ITargetDevice Device)
+		public bool TryValidateDevice(ITargetDevice Device, ref string Message)
 		{
 			if(!string.IsNullOrEmpty(FirmwareVersion))
 			{
 				// Requesting an explicit firmware version
-				return TryUpdateFirmware(Device, FirmwareVersion);
+				return TryUpdateFirmware(Device, FirmwareVersion, ref Message);
 			}
 			else if(!string.IsNullOrEmpty(ProjectName))
 			{
 				// Requesting the firmware version a project requires
-				return TryUpdateFirmwareForProject(Device, ProjectName);
+				return TryUpdateFirmwareForProject(Device, ProjectName, ref Message);
 			}
 
 			return true;
@@ -59,19 +60,20 @@ namespace Gauntlet
 		/// </summary>
 		/// <param name="Device">The Device to update</param>
 		/// <param name="DesiredVersion">The firmware version to update to</param>
+		/// <param name="Message">The output message when failing to validate</param>
 		/// <returns></returns>
-		public bool TryUpdateFirmware(ITargetDevice Device, string DesiredVersion)
+		public bool TryUpdateFirmware(ITargetDevice Device, string DesiredVersion, ref string Message)
 		{
 			try
 			{
-				if (!IDeviceValidator.EnsureConnection(Device))
+				if (!IDeviceValidator.EnsureConnection(Device, ref Message))
 				{
 					return false;
 				}
 
 				if (Device.Platform == null)
 				{
-					Log.Warning("Cannot update firmware for device {Device} because the platform value is null. Was the device mocked?");
+					Log.Warning("Cannot update firmware for device because the platform value is null. Was the device mocked?");
 					return false;
 				}
 
@@ -93,7 +95,7 @@ namespace Gauntlet
 				string CurrentVersion;
 				if (!Handler.GetCurrentVersion(Device, out CurrentVersion))
 				{
-					Log.Warning("Failed to retreive current firmware version from {Device}", Device);
+					Log.Warning("Failed to retrieve current firmware version");
 					return false;
 				}
 
@@ -101,12 +103,13 @@ namespace Gauntlet
 				{
 					if (!Handler.UpdateDeviceFirmware(Device, DesiredVersion))
 					{
-						Log.Warning("Failed to update {Device}'s firmware to version {DesiredVersion}", Device, DesiredVersion);
+						Message = $"Failed to update device's firmware to version {DesiredVersion}";
+						Log.Info(KnownLogEvents.Gauntlet_DeviceEvent, Message);
 						return false;
 					}
 					else
 					{
-						Log.Info("Successfully updated {Device} to new firmware version {DesiredVersion}", Device, DesiredVersion);
+						Log.Info("Successfully updated {DeviceName} to new firmware version {DesiredVersion}", Device.Name, DesiredVersion);
 					}
 				}
 				else
@@ -118,7 +121,17 @@ namespace Gauntlet
 			}
 			catch(Exception Ex)
 			{
-				Log.Warning("Encountered an {ExceptionType} while attempting to update firmware. {Exception}", Ex.GetType(), Ex);
+				string ErrorMessage = $"Encountered an {Ex.GetType()} while attempting to update firmware. {Ex}";
+				if (Ex is DeviceException)
+				{
+					Message = ErrorMessage;
+					Log.Info(KnownLogEvents.Gauntlet_DeviceEvent, Message);
+				}
+				else
+				{
+					Log.Warning(ErrorMessage);
+				}
+
 				return false;
 			}
 		}
@@ -128,12 +141,13 @@ namespace Gauntlet
 		/// </summary>
 		/// <param name="Device">The Device to update</param>
 		/// <param name="ProjectName">Name of the project to get firmware version of</param>
+		/// <param name="Message">The output message when failing to validate</param>
 		/// <returns></returns>
-		public bool TryUpdateFirmwareForProject(ITargetDevice Device, string ProjectName)
+		public bool TryUpdateFirmwareForProject(ITargetDevice Device, string ProjectName, ref string Message)
 		{
 			try
 			{
-				if (!IDeviceValidator.EnsureConnection(Device))
+				if (!IDeviceValidator.EnsureConnection(Device, ref Message))
 				{
 					return false;
 				}
@@ -167,11 +181,21 @@ namespace Gauntlet
 					return false;
 				}
 
-				return TryUpdateFirmware(Device, DesiredVersion);
+				return TryUpdateFirmware(Device, DesiredVersion, ref Message);
 			}
 			catch (Exception Ex)
 			{
-				Log.Warning("Encountered an {ExceptionType} while attempting to update firmware. {Exception}", Ex.GetType(), Ex);
+				string ErrorMessage = $"Encountered an {Ex.GetType()} while attempting to update firmware. {Ex}";
+				if (Ex is DeviceException)
+				{
+					Message = ErrorMessage;
+					Log.Info(KnownLogEvents.Gauntlet_DeviceEvent, Message);
+				}
+				else
+				{
+					Log.Warning(ErrorMessage);
+				}
+
 				return false;
 			}
 		}
diff --git a/Engine/Source/Programs/AutomationTool/Gauntlet/Framework/Devices/Gauntlet.DeviceLoginValidator.cs b/Engine/Source/Programs/AutomationTool/Gauntlet/Framework/Devices/Gauntlet.DeviceLoginValidator.cs
--- a/Engine/Source/Programs/AutomationTool/Gauntlet/Framework/Devices/Gauntlet.DeviceLoginValidator.cs
+++ b/Engine/Source/Programs/AutomationTool/Gauntlet/Framework/Devices/Gauntlet.DeviceLoginValidator.cs
@@ -1,10 +1,6 @@
 // Copyright Epic Games, Inc. All Rights Reserved.
 
-using System;
-using System.Collections.Generic;
-using System.Linq;
-using System.Text;
-using System.Threading.Tasks;
+using EpicGames.Core;
 
 namespace Gauntlet
 {
@@ -22,9 +18,9 @@ namespace Gauntlet
 			AutoParam.ApplyParamsAndDefaults(this, Globals.Params.AllArguments);
 		}
 
-		public bool TryValidateDevice(ITargetDevice Device)
+		public bool TryValidateDevice(ITargetDevice Device, ref string Message)
 		{
-			if(!IDeviceValidator.EnsureConnection(Device))
+			if(!IDeviceValidator.EnsureConnection(Device, ref Message))
 			{
 				return false;
 			}
@@ -39,7 +35,8 @@ namespace Gauntlet
 				}
 				else
 				{
-					Log.Warning("Unable to secure login to an online platform account!");
+					Message = "Unable to secure login to an online platform account on device!";
+					Log.Info(KnownLogEvents.Gauntlet_DeviceEvent, Message);
 					return false;
 				}
 			}
diff --git a/Engine/Source/Programs/AutomationTool/Gauntlet/Framework/Devices/Gauntlet.DevicePool.cs b/Engine/Source/Programs/AutomationTool/Gauntlet/Framework/Devices/Gauntlet.DevicePool.cs
--- a/Engine/Source/Programs/AutomationTool/Gauntlet/Framework/Devices/Gauntlet.DevicePool.cs
+++ b/Engine/Source/Programs/AutomationTool/Gauntlet/Framework/Devices/Gauntlet.DevicePool.cs
@@ -360,7 +360,7 @@ namespace Gauntlet
 
 						if (Failures >= 0.33)
 						{
-							FailedProvisions.ForEach(D => Log.Warning(KnownLogEvents.Gauntlet_DeviceEvent, "Device {Name} could not be added", D));
+							FailedProvisions.ForEach(D => Log.Warning("Device {DeviceName} could not be added", D.Name));
 						}
 					}
 
@@ -370,7 +370,7 @@ namespace Gauntlet
 					// warn if anything was left registered
 					ReservedDevices.ForEach(D =>
 					{
-						Log.Warning(KnownLogEvents.Gauntlet_DeviceEvent, "Device {Name} was not unregistered prior to DevicePool.Dispose!. Forcibly disposing.", D.Name);
+						Log.Warning("Device {DeviceName} was not unregistered prior to DevicePool.Dispose!. Forcibly disposing.", D.Name);
 						D.Dispose();
 					});
 
@@ -502,7 +502,7 @@ namespace Gauntlet
 				// Did they specify a file?
 				if (PossibleFileName && File.Exists(InputReference))
 				{
-					Log.Info(KnownLogEvents.Gauntlet_DeviceEvent, "Adding devices from {Reference}", InputReference);
+					Log.Info("Adding devices from {Reference}", InputReference);
 					List<DeviceDefinition> DeviceDefinitions = JsonSerializer.Deserialize<List<DeviceDefinition>>(
 						File.ReadAllText(InputReference),
 						new JsonSerializerOptions { PropertyNameCaseInsensitive = true }
@@ -510,7 +510,7 @@ namespace Gauntlet
 
 					foreach (DeviceDefinition Def in DeviceDefinitions)
 					{
-						Log.Info(KnownLogEvents.Gauntlet_DeviceEvent, "Adding {Name}", Def);
+						Log.Info("Adding {DeviceDetails}", Def);
 
 						// use Legacy field if it exists
 						if (Def.Platform == null)
@@ -524,7 +524,7 @@ namespace Gauntlet
 							// check whether disabled
 							if (String.Compare(Def.Available, "disabled", true) == 0)
 							{
-								Log.Info(KnownLogEvents.Gauntlet_DeviceEvent, "Skipping {Name} due to being disabled", Def.Name);
+								Log.Info("Skipping {DeviceName} due to being disabled", Def.Name);
 								continue;
 							}
 
@@ -554,13 +554,13 @@ namespace Gauntlet
 
 									if (DateTime.Now < From || DateTime.Now > To)
 									{
-										Log.Info(KnownLogEvents.Gauntlet_DeviceEvent, "Skipping {Name} due to availability constraint {Constraint}", Def.Name, Def.Available);
+										Log.Info("Skipping {DeviceName} due to availability constraint {Constraint}", Def.Name, Def.Available);
 										continue;
 									}
 								}
 								else
 								{
-									Log.Warning(KnownLogEvents.Gauntlet_DeviceEvent, "Failed to parse availability {Constraint} for {Name}", Def.Available, Def.Name);
+									Log.Warning("Failed to parse availability {Constraint} for {DeviceName}", Def.Available, Def.Name);
 								}
 							}
 						}
@@ -709,7 +709,7 @@ namespace Gauntlet
 
 					while (NextDevice != null && ContinuePredicate)
 					{
-						Log.Verbose(KnownLogEvents.Gauntlet_DeviceEvent, "Checking {Name} against predicate", NextDevice.Name);
+						Log.Verbose(KnownLogEvents.Gauntlet_DeviceEvent, "Checking {DeviceName} against predicate", NextDevice.Name);
 						MatchingProvisionedDevices.Remove(NextDevice);
 						ContinuePredicate = Predicate(NextDevice);
 
@@ -725,7 +725,7 @@ namespace Gauntlet
 
 						if (NextDeviceDef != null)
 						{
-							Log.Verbose(KnownLogEvents.Gauntlet_DeviceEvent, "Provisioning device {Name} for the pool", NextDeviceDef.Name);
+							Log.Verbose(KnownLogEvents.Gauntlet_DeviceEvent, "Provisioning device {DeviceName} for the pool", NextDeviceDef.Name);
 
 							// try to create a device. This can fail, but if so we'll just end up back here
 							// on the next iteration
@@ -737,11 +737,11 @@ namespace Gauntlet
 							if (NewDevice != null)
 							{
 								MatchingProvisionedDevices.Add(NewDevice);
-								Log.Verbose(KnownLogEvents.Gauntlet_DeviceEvent, "Added device {Name} to pool", NewDevice.Name);
+								Log.Verbose(KnownLogEvents.Gauntlet_DeviceEvent, "Added device {DeviceName} to pool", NewDevice.Name);
 							}
 							else
 							{
-								Log.Info(KnownLogEvents.Gauntlet_DeviceEvent, "Failed to provision {Name}", NextDeviceDef.Name);
+								Log.Info(KnownLogEvents.Gauntlet_DeviceEvent, "Failed to provision {DeviceName}", NextDeviceDef.Name);
 								// track this
 								if (FailedProvisions.Contains(NextDeviceDef) == false)
 								{
@@ -877,8 +877,6 @@ namespace Gauntlet
 					// If a device from service can't be added, fail reservation and cleanup devices
 					if (TargetDevice == null)
 					{
-						ReportDeviceError(Device.Name, "CreateDeviceError");
-
 						// If some devices from reservation have been created, release them which will also dispose of reservation
 						if (ScopeReservedDevices.Count > 0)
 						{
@@ -890,7 +888,7 @@ namespace Gauntlet
 							DeviceReservation.Dispose();
 						}
 
-						Log.Info(KnownLogEvents.Gauntlet_DeviceEvent, "Unable to make device registration: device registration failed for {Platform}:{Name}", Def.Platform, Def.Name);
+						Log.Info(KnownLogEvents.Gauntlet_DeviceEvent, "Unable to make device registration: device registration failed for {Platform}:{DeviceName}", Def.Platform, Def.Name);
 						return false;
 					}
 					else
@@ -909,6 +907,7 @@ namespace Gauntlet
 			{
 				Log.Info(KnownLogEvents.Gauntlet_DeviceEvent, "Unable to reserve all devices from service.");
 			}
+
 			return ScopeReservedDevices.Count > 0;
 		}
 
@@ -1114,25 +1113,30 @@ namespace Gauntlet
 
 				if (NewDevice == null)
 				{
-					Log.Info(KnownLogEvents.Gauntlet_DeviceEvent, "Failed to create device {Name}. Device could not be connected.", Def.ToString());
+					Log.Warning(KnownLogEvents.Gauntlet_DeviceEvent, "Failed to create device {DeviceName}. Device could not be connected.", Def.Name);
 					return null;
 				}
 
 				if (NewDevice.IsAvailable == false)
 				{
-					Log.Info(KnownLogEvents.Gauntlet_DeviceEvent, "Assigned device {Name} reports unavailable. Requesting a forced disconnect", NewDevice);
+					Log.Info(KnownLogEvents.Gauntlet_DeviceEvent, "Assigned device {DeviceName} reports unavailable. Requesting a forced disconnect", NewDevice.Name);
 					NewDevice.Disconnect(true);
 
 					if (NewDevice.IsAvailable == false)
 					{
-						Log.Info(KnownLogEvents.Gauntlet_DeviceEvent, "Assigned device {Name} still unavailable. Requesting a reboot", NewDevice);
+						Log.Info(KnownLogEvents.Gauntlet_DeviceEvent, "Assigned device {DeviceName} still unavailable. Requesting a reboot", NewDevice.Name);
 						NewDevice.Reboot();
 					}
 				}
 
 				// Now validate if the kit meets the necessary requirements
-				if (!TryValidateDeviceRequirements(NewDevice))
+				string Message = string.Empty;
+				if (!TryValidateDeviceRequirements(NewDevice, ref Message))
 				{
+					if (!string.IsNullOrEmpty(Message))
+					{
+						ReportDeviceError(Def.Name, Message);
+					}
 					Log.Info("\nSkipping device.");
 					return null;
 				}
@@ -1147,7 +1151,16 @@ namespace Gauntlet
 			}
 			catch (Exception Ex)
 			{
-				Log.Info(KnownLogEvents.Gauntlet_DeviceEvent, "Failed to create device {Name}. {Message}\n{StackTrace}", Def.ToString(), Ex.Message, Ex.StackTrace);
+				string WarningMessage = $"Failed to create device {Def.Name}. {Ex.Message}\n{Ex.StackTrace}";
+				if (Ex is DeviceException)
+				{
+					ReportDeviceError(Def.Name, WarningMessage);
+					Log.Info(KnownLogEvents.Gauntlet_DeviceEvent, WarningMessage);
+				}
+				else
+				{
+					Log.Warning(KnownLogEvents.Gauntlet_DeviceEvent, WarningMessage);
+				}
 			}
 
 			return NewDevice;
@@ -1157,8 +1170,9 @@ namespace Gauntlet
 		/// Verifies if the provided TargetDevice meets requirements such as firmware, login status, settings, etc.
 		/// </summary>
 		/// <param name="Device">The device to validate</param>
+		/// <param name="Message">The output message when failing to validate</param>
 		/// <returns>True if the device matches the required specifications</returns>
-		protected bool TryValidateDeviceRequirements(ITargetDevice Device)
+		protected bool TryValidateDeviceRequirements(ITargetDevice Device, ref string Message)
 		{
 			IEnumerable<IDeviceValidator> Validators = InterfaceHelpers.FindImplementations<IDeviceValidator>(true).Where(Validator => Validator.bEnabled);
 			if (!Validators.Any())
@@ -1168,17 +1182,23 @@ namespace Gauntlet
 
 			bool bInitiallyConnected = Device.IsConnected;
 
-			Log.Info("\nValidating requirements for {Device}...", Device);
+			Log.Info(KnownLogEvents.Gauntlet_DeviceEvent, "\nValidating requirements for {DeviceName}...", Device.Name);
 			bool bSucceeded = true;
+			List<string> MessageAggregate = new();
 
 			foreach (IDeviceValidator Validator in Validators)
 			{
 				Log.Info("\nStarting validation for {Validator}", Validator);
 
-				if(!Validator.TryValidateDevice(Device))
+				string ValidationMessage = string.Empty;
+				if(!Validator.TryValidateDevice(Device, ref ValidationMessage))
 				{
-					Log.Warning("Failure! See above for details");
+					Log.Info("Failed!");
 					bSucceeded = false;
+					if (!string.IsNullOrEmpty(ValidationMessage))
+					{
+						MessageAggregate.Add(ValidationMessage);
+					}
 				}
 				else
 				{
@@ -1188,11 +1208,12 @@ namespace Gauntlet
 
 			if(!bSucceeded)
 			{
-				Log.Warning("\nFailed to validate requirements.");
+				Message = string.Join("\n", MessageAggregate.Where(M => !string.IsNullOrEmpty(M)));
+				Log.Info(KnownLogEvents.Gauntlet_DeviceEvent, "\nFailed to validate requirements on device {DeviceName}", Device.Name);
 			}
 			else
 			{
-				Log.Info("\nAll validators passed, selecting device {Device}\n", Device);
+				Log.Info(KnownLogEvents.Gauntlet_DeviceEvent, "\nAll validators passed, selecting device {DeviceName}\n", Device.Name);
 			}
 
 			// Most validators require establishing a connection to the device.
@@ -1212,7 +1233,7 @@ namespace Gauntlet
 		{
 			if (MaxDeviceErrorReports == 0)
 			{
-				Log.Info(KnownLogEvents.Gauntlet_DeviceEvent, "Maximum device errors reported to backend, {Name} : {Message} ignored", ServiceDeviceName, ErrorMessage);
+				Log.Info(KnownLogEvents.Gauntlet_DeviceEvent, "Maximum device errors reported to backend, {DeviceName} : {Message} ignored", ServiceDeviceName, ErrorMessage);
 				return;
 			}
 
diff --git a/Engine/Source/Programs/AutomationTool/Gauntlet/Framework/Devices/Gauntlet.DeviceProfileValidator.cs b/Engine/Source/Programs/AutomationTool/Gauntlet/Framework/Devices/Gauntlet.DeviceProfileValidator.cs
--- a/Engine/Source/Programs/AutomationTool/Gauntlet/Framework/Devices/Gauntlet.DeviceProfileValidator.cs
+++ b/Engine/Source/Programs/AutomationTool/Gauntlet/Framework/Devices/Gauntlet.DeviceProfileValidator.cs
@@ -1,5 +1,6 @@
 // Copyright Epic Games, Inc. All Rights Reserved.
 
+using EpicGames.Core;
 using System;
 
 namespace Gauntlet
@@ -19,9 +20,9 @@ namespace Gauntlet
 			bEnabled = !string.IsNullOrEmpty(DeviceProfile);
 		}
 
-		public bool TryValidateDevice(ITargetDevice Device)
+		public bool TryValidateDevice(ITargetDevice Device, ref string Message)
 		{
-			if(!IDeviceValidator.EnsureConnection(Device))
+			if(!IDeviceValidator.EnsureConnection(Device, ref Message))
 			{
 				return false;
 			}
@@ -57,7 +58,8 @@ namespace Gauntlet
 					PlatformConfigurationBase CurrentConfiguration = ConfigurableDevice.GetCurrentConfigurationSnapshot();
 					if (CurrentConfiguration == null)
 					{
-						Log.Warning("Failed to snapshot existing device configuration");
+						Message = "Failed to snapshot existing device configuration on device";
+						Log.Info(KnownLogEvents.Gauntlet_DeviceEvent, Message);
 						return false;
 					}
 
@@ -69,7 +71,17 @@ namespace Gauntlet
 				}
 				catch(Exception Ex)
 				{
-					Log.Warning("Encountered an exception while trying to configure device.\n{ExceptionType}: {Message}\n\t{Stack}", Ex.GetType(), Ex.Message, Ex.StackTrace);
+					string ErrorMessage = $"Encountered an {Ex.GetType()} while while trying to configure device. {Ex}\n\t{Ex.StackTrace}";
+					if (Ex is DeviceException)
+					{
+						Message = ErrorMessage;
+						Log.Info(KnownLogEvents.Gauntlet_DeviceEvent, Message);
+					}
+					else
+					{
+						Log.Warning(ErrorMessage);
+					}
+
 					return false;
 				}
 			}
diff --git a/Engine/Source/Programs/AutomationTool/Gauntlet/Framework/Devices/Gauntlet.DeviceValidator.cs b/Engine/Source/Programs/AutomationTool/Gauntlet/Framework/Devices/Gauntlet.DeviceValidator.cs
--- a/Engine/Source/Programs/AutomationTool/Gauntlet/Framework/Devices/Gauntlet.DeviceValidator.cs
+++ b/Engine/Source/Programs/AutomationTool/Gauntlet/Framework/Devices/Gauntlet.DeviceValidator.cs
@@ -1,10 +1,12 @@
 // Copyright Epic Games, Inc. All Rights Reserved.
 
+using EpicGames.Core;
+
 namespace Gauntlet
 {
 	/// <summary>
 	/// IDeviceValidator provides the ability to verify a TargetDevice matches a set of requirements.
-	/// They are particularily useful for ensuring development kits are in a desired state prior to a test.
+	/// They are particularly useful for ensuring development kits are in a desired state prior to a test.
 	/// This can include, but is not limited to:
 	///		- Proper firmware versions
 	///		- Having an account signed in
@@ -24,16 +26,18 @@ namespace Gauntlet
 		/// Verifies if the device is valid in the context of the running program
 		/// </summary>
 		/// <param name="Device">The device to validate</param>
+		/// <param name="Message">The output message when failing to validate</param>
 		/// <returns>True if the device succeeded validation</returns>
-		bool TryValidateDevice(ITargetDevice Device);
+		bool TryValidateDevice(ITargetDevice Device, ref string Message);
 
 		/// <summary>
 		/// Helper function that can be used to ensure a device is powered on/connected
 		/// before attempting to run any validation checks
 		/// </summary>
 		/// <param name="Device"></param>
+		/// <param name="Message">The output message when failing to validate</param>
 		/// <returns></returns>
-		protected static bool EnsureConnection(ITargetDevice Device)
+		protected static bool EnsureConnection(ITargetDevice Device, ref string Message)
 		{
 			if (Device == null)
 			{
@@ -45,12 +49,13 @@ namespace Gauntlet
 			{
 				if (!Device.PowerOn())
 				{
-					Log.Warning("Failed to power on device");
+					Message = "Failed to power on device";
+					Log.Info(KnownLogEvents.Gauntlet_DeviceEvent, Message);
 					return false;
 				}
 				else
 				{
-					Log.VeryVerbose("Powered on device while ensuring connection");
+					Log.VeryVerbose(KnownLogEvents.Gauntlet_DeviceEvent, "Powered on device {DeviceName} while ensuring connection", Device.Name);
 				}
 			}
 
@@ -58,16 +63,24 @@ namespace Gauntlet
 			{
 				if (!Device.Connect())
 				{
-					Log.Warning("Failed to connect to device");
+					Message = "Failed to connect to device";
+					Log.Info(KnownLogEvents.Gauntlet_DeviceEvent, Message);
 					return false;
 				}
 				else
 				{
-					Log.VeryVerbose("Connected to device while ensuring connection");
+					Log.VeryVerbose(KnownLogEvents.Gauntlet_DeviceEvent, "Connected to device {DeviceName} while ensuring connection", Device.Name);
 				}
 			}
 
-			return Device.IsAvailable;
+			if (!Device.IsAvailable)
+			{
+				Message = "Device not available";
+				Log.Info(KnownLogEvents.Gauntlet_DeviceEvent, Message);
+				return false;
+			}
+
+			return true;
 		}
 	}
 }
diff --git a/Engine/Source/Programs/AutomationTool/Gauntlet/Framework/Devices/Gauntlet.UnrealDeviceReservation.cs b/Engine/Source/Programs/AutomationTool/Gauntlet/Framework/Devices/Gauntlet.UnrealDeviceReservation.cs
--- a/Engine/Source/Programs/AutomationTool/Gauntlet/Framework/Devices/Gauntlet.UnrealDeviceReservation.cs
+++ b/Engine/Source/Programs/AutomationTool/Gauntlet/Framework/Devices/Gauntlet.UnrealDeviceReservation.cs
@@ -81,7 +81,7 @@ namespace Gauntlet
 				// actually acquire them
 				DevicePool.Instance.ReserveDevices(AcquiredDevices);
 
-				Log.Info("Selected devices {0} for client(s). Prepping", string.Join(", ", AcquiredDevices));
+				Log.Info("Selected devices {Devices} for client(s). Prepping", string.Join(", ", AcquiredDevices));
 
 				foreach (ITargetDevice Device in AcquiredDevices)
 				{
@@ -94,6 +94,12 @@ namespace Gauntlet
 					{
 						Log.Info("Rebooting {0}", Device);
 						Device.Reboot();
+						// Force a re-login after a reboot to be sure it is ready.
+						if (Globals.Params.ParseParam("VerifyLogin") && Device is IOnlineServiceLogin DeviceLogin)
+						{
+							Log.Info("Revalidate device login after reboot...");
+							DeviceLogin.VerifyLogin();
+						}
 					}
 
 					if (Device.IsConnected == false)
@@ -112,7 +118,7 @@ namespace Gauntlet
 
 					// mark these as problems. Could be something grabbed them before us, could be that they are
 					// unresponsive in some way
-					LostDevices.ToList().ForEach(D => MarkProblemDevice(D));
+					LostDevices.ToList().ForEach(D => MarkProblemDevice(D, "Agent lost connection to device."));
 					AcquiredDevices.ToList().ForEach(D => D.Disconnect());
 
 					DevicePool.Instance.ReleaseDevices(AcquiredDevices);
@@ -181,7 +187,7 @@ namespace Gauntlet
 			return Devices;
 		}
 
-		public void MarkProblemDevice(ITargetDevice Device, string ErrorMessage = "MarkProblemDevice")
+		public void MarkProblemDevice(ITargetDevice Device, string ErrorMessage)
 		{
 			if (ProblemDevices.Where(D => D.Name == Device.Name && D.Platform == Device.Platform).Count() > 0)
 			{
diff --git a/Engine/Source/Programs/AutomationTool/Gauntlet/Platform/Android/Gauntlet.TargetDeviceAndroid.cs b/Engine/Source/Programs/AutomationTool/Gauntlet/Platform/Android/Gauntlet.TargetDeviceAndroid.cs
--- a/Engine/Source/Programs/AutomationTool/Gauntlet/Platform/Android/Gauntlet.TargetDeviceAndroid.cs
+++ b/Engine/Source/Programs/AutomationTool/Gauntlet/Platform/Android/Gauntlet.TargetDeviceAndroid.cs
@@ -167,7 +167,7 @@ namespace Gauntlet
 
 						if (AdbResult.ExitCode != 0)
 						{
-							throw new AutomationException("adb failed to connect to {0}. {1}", DeviceName, AdbResult.Output);
+							throw new DeviceException("adb failed to connect to {0}. {1}", DeviceName, AdbResult.Output);
 						}
 					}
 
@@ -182,21 +182,21 @@ namespace Gauntlet
 				// sanity check that it was now found
 				if (!ConnectedDevices.Keys.Contains(DeviceName))
 				{
-					throw new AutomationException("Failed to find new device {0} in connection list", DeviceName);
+					throw new DeviceException("Failed to find new device {0} in connection list", DeviceName);
 				}
 			}
 
 			if (ConnectedDevices[DeviceName] != "device")
 			{
 				Dispose();
-				throw new AutomationException("Device {0} is '{1}'.", DeviceName, ConnectedDevices[DeviceName]);
+				throw new DeviceException("Device {0} is '{1}'.", DeviceName, ConnectedDevices[DeviceName]);
 			}
 
 			// Get the external storage path from the device (once the device is validated as connected)
 			IProcessResult StorageQueryResult = RunAdbDeviceCommand(AndroidPlatform.GetStorageQueryCommand());
 			if (StorageQueryResult.ExitCode != 0)
 			{
-				throw new AutomationException("Failed to query external storage setup on {0}", DeviceName);
+				throw new DeviceException("Failed to query external storage setup on {0}", DeviceName);
 			}
 			StoragePath = StorageQueryResult.Output.Trim();
 
@@ -584,7 +584,7 @@ namespace Gauntlet
 
 			if (Result.ExitCode != 0)
 			{
-				throw new AutomationException("adb command {0} failed. {1}", Args, Result.Output);
+				throw new DeviceException("adb command {0} failed. {1}", Args, Result.Output);
 			}
 
 			return Result.Output;
@@ -862,7 +862,7 @@ namespace Gauntlet
 
 					if (AdbResult.ExitCode != 0)
 					{
-						throw new AutomationException("Failed to install {0}. Error {1}", SourcePath, AdbResult.Output);
+						throw new DeviceException("Failed to install {0}. Error {1}", SourcePath, AdbResult.Output);
 					}
 
 					// for APK query the package info and get the update time
@@ -909,17 +909,17 @@ namespace Gauntlet
 								}
 								if (AdbResult.ExitCode != 0)
 								{
-									throw new AutomationException("Failed to push {0} to device. Error {1}", SourcePath, AdbResult.Output);
+									throw new DeviceException("Failed to push {0} to device. Error {1}", SourcePath, AdbResult.Output);
 								}
 							}
 							else
 							{
-								throw new AutomationException("Failed to reconnect {0}", Name);
+								throw new DeviceException("Failed to reconnect {0}", Name);
 							}
 						}
 						else
 						{
-							throw new AutomationException("Failed to push {0} to device. Error {1}", SourcePath, AdbResult.Output);
+							throw new DeviceException("Failed to push {0} to device. Error {1}", SourcePath, AdbResult.Output);
 						}
 					}
 
diff --git a/Engine/Source/Programs/AutomationTool/Gauntlet/Platform/IOS/Gauntlet.TargetDeviceIOS.cs b/Engine/Source/Programs/AutomationTool/Gauntlet/Platform/IOS/Gauntlet.TargetDeviceIOS.cs
--- a/Engine/Source/Programs/AutomationTool/Gauntlet/Platform/IOS/Gauntlet.TargetDeviceIOS.cs
+++ b/Engine/Source/Programs/AutomationTool/Gauntlet/Platform/IOS/Gauntlet.TargetDeviceIOS.cs
@@ -176,7 +176,7 @@ namespace Gauntlet
 					Log.Info("Attempting to connect to {UUID}", UUID);
 					if (!Connect())
 					{
-						throw new AutomationException("Failed to connect to {DeviceName}. Ensure the device is properly connected to the host.", UUID);
+						throw new DeviceException("Failed to connect to {DeviceName}. Ensure the device is properly connected to the host.", UUID);
 					}
 				}
 				else
@@ -196,10 +196,10 @@ namespace Gauntlet
 				}
 				if (!ConnectedDevices[UUID].bConnected)
 				{
-					Log.Info("Attempting to connect to {DeviceName}");
+					Log.Info("Attempting to connect to {DeviceName}", Name);
 					if(!Connect())
 					{
-						throw new AutomationException("Failed to connect to {DeviceName}", UUID);
+						throw new DeviceException("Failed to connect to {DeviceName}", UUID);
 					}
 				}
 			}
@@ -219,7 +219,7 @@ namespace Gauntlet
 				else
 				{
 					string Message = "As of XCode 15, Apple has removed DeveloperDiskImages for iOS 17 and beyond. " +
-						"Because of this, ios-deploy cannot be used to run and montior apps within Gauntlet. " +
+						"Because of this, ios-deploy cannot be used to run and monitor apps within Gauntlet. " +
 						"Instead, DeviceCtl replaces this functionality which is fully available as of XCode 16. " +
 						"Update to XCode 16 to prevent this error.";
 					throw new AutomationException(Message);
@@ -353,7 +353,7 @@ namespace Gauntlet
 				IProcessResult AppListResult = ExecuteDevicectlCommand("device info apps");
 				if (AppListResult.ExitCode != 0)
 				{
-					throw new AutomationException("App list exited with code {0}: {1}", AppListResult.ExitCode, AppListResult.Output);
+					throw new DeviceException("App list exited with code {0}: {1}", AppListResult.ExitCode, AppListResult.Output);
 				}
 
 				AppOutput = AppListResult.Output;
@@ -363,7 +363,7 @@ namespace Gauntlet
 				IProcessResult AppListResult = ExecuteIOSDeployCommand("-B");
 				if (AppListResult.ExitCode != 0)
 				{
-					throw new AutomationException("App list exited with code {0}: {1}", AppListResult.ExitCode, AppListResult.Output);
+					throw new DeviceException("App list exited with code {0}: {1}", AppListResult.ExitCode, AppListResult.Output);
 				}
 
 				AppOutput = AppListResult.Output;
@@ -392,7 +392,7 @@ namespace Gauntlet
 				UninstallResult.WaitForExit();
 				if (UninstallResult.ExitCode != 0)
 				{
-					throw new AutomationException("Uninstalling {0} exited with code {1}: {2}", App, UninstallResult.ExitCode, UninstallResult.Output);
+					throw new DeviceException("Uninstalling {0} exited with code {1}: {2}", App, UninstallResult.ExitCode, UninstallResult.Output);
 				}
 			}
 		}
@@ -420,14 +420,14 @@ namespace Gauntlet
 			Result.WaitForExit();
 			if (Result.ExitCode != 0)
 			{
-				throw new AutomationException("Cleaning Engine artifacts for {0} exited with code {1}:\n{2}", BundleID, Result.ExitCode, Result.Output);
+				throw new DeviceException("Cleaning Engine artifacts for {0} exited with code {1}:\n{2}", BundleID, Result.ExitCode, Result.Output);
 			}
 
 			Result = ExecuteIOSDeployCommand(CleanProject);
 			Result.WaitForExit();
 			if (Result.ExitCode != 0)
 			{
-				throw new AutomationException("Cleaning Project artifacts for {0} exited with code {1}:\n{2}", BundleID, Result.ExitCode, Result.Output);
+				throw new DeviceException("Cleaning Project artifacts for {0} exited with code {1}:\n{2}", BundleID, Result.ExitCode, Result.Output);
 			}
 		}
 
@@ -483,7 +483,7 @@ namespace Gauntlet
 			Result.WaitForExit();
 			if(Result.ExitCode != 0)
 			{
-				throw new AutomationException("Uninstall exited with code {0}:\n{1}", Result.ExitCode, Result.Output);
+				throw new DeviceException("Uninstall exited with code {0}:\n{1}", Result.ExitCode, Result.Output);
 			}
 
 			// Install the app
@@ -504,7 +504,7 @@ namespace Gauntlet
 
 			if (Result.ExitCode != 0)
 			{
-				throw new AutomationException("Install exited with code {0}:\n{1}", Result.ExitCode, Result.Output);
+				throw new DeviceException("Install exited with code {0}:\n{1}", Result.ExitCode, Result.Output);
 			}
 
 			double RemainingTime = MaxInstallTime - (StartTime - DateTime.Now).TotalSeconds;
@@ -579,11 +579,11 @@ namespace Gauntlet
 					{
 						if(RemainingTimePerFile <= DeltaTime)
 						{
-							throw new AutomationException("Install timed out after {0} seconds", RemainingTimePerFile);
+							throw new DeviceException("Install timed out after {0} seconds", RemainingTimePerFile);
 						}
 						else
 						{
-							throw new AutomationException("Install of content file {0} failed with exit code {1}:\n{2}", FileName, ContentCopyResult.ExitCode, ContentCopyResult.Output);
+							throw new DeviceException("Install of content file {0} failed with exit code {1}:\n{2}", FileName, ContentCopyResult.ExitCode, ContentCopyResult.Output);
 						}
 					}
 
@@ -721,7 +721,7 @@ namespace Gauntlet
 		{
 			if (Install is not IOSAppInstall IOSInstall)
 			{
-				throw new DeviceException("AppInstance is of incorrect type!");
+				throw new AutomationException("AppInstance is of incorrect type!");
 			}
 
 			// Copy the commandline into a uecommandline.txt
@@ -745,7 +745,7 @@ namespace Gauntlet
 
 			if(CommandLineCopyResult.ExitCode != 0)
 			{
-				throw new AutomationException("Failed to copy uecommandline file to device!");
+				throw new DeviceException("Failed to copy uecommandline file to device!");
 			}
 
 			Log.Info("Launching {0} on {1}", Install.Name, this);
diff --git a/Engine/Source/Programs/AutomationTool/Gauntlet/Platform/Linux/Gauntlet.TargetDeviceLinux.cs b/Engine/Source/Programs/AutomationTool/Gauntlet/Platform/Linux/Gauntlet.TargetDeviceLinux.cs
--- a/Engine/Source/Programs/AutomationTool/Gauntlet/Platform/Linux/Gauntlet.TargetDeviceLinux.cs
+++ b/Engine/Source/Programs/AutomationTool/Gauntlet/Platform/Linux/Gauntlet.TargetDeviceLinux.cs
@@ -37,12 +37,12 @@ namespace Gauntlet
 
 			if (LinuxApp == null)
 			{
-				throw new DeviceException("AppInstance is of incorrect type!");
+				throw new AutomationException("AppInstance is of incorrect type!");
 			}
 
 			if (File.Exists(LinuxApp.ExecutablePath) == false)
 			{
-				throw new DeviceException("Specified path {0} not found!", LinuxApp.ExecutablePath);
+				throw new AutomationException("Specified path {0} not found!", LinuxApp.ExecutablePath);
 			}
 
 			ILongProcessResult Result = null;
diff --git a/Engine/Source/Programs/AutomationTool/Gauntlet/Platform/Windows/Gauntlet.TargetDeviceWindows.cs b/Engine/Source/Programs/AutomationTool/Gauntlet/Platform/Windows/Gauntlet.TargetDeviceWindows.cs
--- a/Engine/Source/Programs/AutomationTool/Gauntlet/Platform/Windows/Gauntlet.TargetDeviceWindows.cs
+++ b/Engine/Source/Programs/AutomationTool/Gauntlet/Platform/Windows/Gauntlet.TargetDeviceWindows.cs
@@ -71,12 +71,12 @@ namespace Gauntlet
 
 			if (WinApp == null)
 			{
-				throw new DeviceException("AppInstance is of incorrect type!");
+				throw new AutomationException("AppInstance is of incorrect type!");
 			}
 
 			if (File.Exists(WinApp.ExecutablePath) == false)
 			{
-				throw new DeviceException("Specified path {0} not found!", WinApp.ExecutablePath);
+				throw new AutomationException("Specified path {0} not found!", WinApp.ExecutablePath);
 			}
 
 			ILongProcessResult Result = null;
diff --git a/Engine/Source/Programs/AutomationTool/Gauntlet/Unreal/Base/Gauntlet.UnrealSession.cs b/Engine/Source/Programs/AutomationTool/Gauntlet/Unreal/Base/Gauntlet.UnrealSession.cs
--- a/Engine/Source/Programs/AutomationTool/Gauntlet/Unreal/Base/Gauntlet.UnrealSession.cs
+++ b/Engine/Source/Programs/AutomationTool/Gauntlet/Unreal/Base/Gauntlet.UnrealSession.cs
@@ -994,7 +994,7 @@ namespace Gauntlet
 			Log.Info("Released problem devices...");
 			foreach(UnrealSessionRole Role in RolesToClear)
 			{
-				Log.Info("\t {Device}", RolesToDevices[Role].Name);
+				Log.Info("\t {DeviceName}", RolesToDevices[Role].Name);
 				RolesToDevices.Remove(Role);
 				RolesToConfigs.Remove(Role);
 				RolesToInstalls.Remove(Role);
@@ -1376,24 +1376,28 @@ namespace Gauntlet
 						}
 						catch (Exception Ex)
 						{
-							// Warn, ignore the device, and do not continue
+							InstallSuccess = false;
 							string ErrorMessage = string.Format("Encountered error setting up device {0} for role {1}. {2}", Device, Role, Ex);
-							if (ErrorMessage.Contains("not enough space"))
+							if (ErrorMessage.Contains("not enough space") && Device.Platform == BuildHostPlatform.Current.Platform)
 							{
-								Log.Error(KnownLogEvents.Gauntlet_DeviceEvent, ErrorMessage);
-								if (Device.Platform == BuildHostPlatform.Current.Platform)
-								{
-									// If on desktop platform, we are not retrying.
-									// It is unlikely that space is going to be made and InstallBuildParallel has marked the build path as problematic.
-									SessionRetries = 0;
-								}
+								Log.Error(ErrorMessage);
+								// If on desktop platform, we are not retrying.
+								// It is unlikely that space is going to be made and InstallBuildParallel has marked the build path as problematic.
+								SessionRetries = 0;
+								break;
+							}
+
+							// Warn, ignore the device, and do not continue
+							if (Ex is DeviceException)
+							{
+								Log.Info(KnownLogEvents.Gauntlet_DeviceEvent, ErrorMessage);
+								UnrealDeviceReservation.MarkProblemDevice(Device, ErrorMessage);
 							}
 							else
 							{
 								Log.Warning(KnownLogEvents.Gauntlet_DeviceEvent, ErrorMessage);
 							}
-							UnrealDeviceReservation.MarkProblemDevice(Device);
-							InstallSuccess = false;
+
 							IDeviceUsageReporter.RecordEnd(Device.Name, Device.Platform, IDeviceUsageReporter.EventType.Install, IDeviceUsageReporter.EventState.Failure);
 							break;
 						}
@@ -1475,31 +1479,27 @@ namespace Gauntlet
 						}
 						catch (DeviceException Ex)
 						{
-							// shutdown all 
-							Log.Warning(KnownLogEvents.Gauntlet_DeviceEvent, "Device {Name} threw an exception during launch. \nException={Exception}", CurrentInstall.Device, Ex.Message);
+							// mark that device as a problem
+							Log.Info(KnownLogEvents.Gauntlet_DeviceEvent, "Failed to start build on {DeviceName}. Marking as problem device and retrying with new set", CurrentInstall.Device.Name);
+							UnrealDeviceReservation.MarkProblemDevice(CurrentInstall.Device, $"Device threw an exception during launch.\nException={Ex.Message}");
 							Success = false;
 						}
 
 						if (Success == false)
 						{
-							Log.Warning(KnownLogEvents.Gauntlet_DeviceEvent, "Failed to start build on {Name}. Marking as problem device and retrying with new set", CurrentInstall.Device);
-
 							// terminate anything that's running
 							foreach (UnrealSessionInstance.RoleInstance RunningRole in AllRoles.Where(X => X.AppInstance != null))
 							{
-								Log.Info("Shutting down {0}", RunningRole.AppInstance.Device);
+								Log.Info("Shutting down {DeviceName}", RunningRole.AppInstance.Device.Name);
 								RunningRole.AppInstance.Kill();
 								RunningRole.AppInstance.Device.Disconnect();
 							}
 
-							// mark that device as a problem
-							UnrealDeviceReservation.MarkProblemDevice(CurrentInstall.Device);
-
 							ReleaseSessionDevices();
 
 							if (SessionRetries == 0)
 							{
-								throw new AutomationException("Unable to start application for session, see warnings for details.");
+								throw new AutomationException("Unable to start application for session, see log for details.");
 							}
 
 							Log.Info("\nUnable to start application for session (retries left={0})\n", --SessionRetries);
@@ -1648,7 +1648,7 @@ namespace Gauntlet
 						{
 							// Telemetry
 							DateTimeStopwatch Stopwatch = DateTimeStopwatch.Start();
-							Log.Info("Installing {BuildName} of type {BuildType} to {Device}...", BuildSource.BuildName, AppConfig.Build.GetType().Name, Device);
+							Log.Info("Installing {BuildName} of type {BuildType} to {DeviceName}...", BuildSource.BuildName, AppConfig.Build.GetType().Name, Device.Name);
 							IDeviceUsageReporter.RecordStart(Device.Name, Device.Platform, IDeviceUsageReporter.EventType.Install, IDeviceUsageReporter.EventState.Success, BuildSource.BuildName);
 
 							try
@@ -1683,7 +1683,7 @@ namespace Gauntlet
 				}
 				catch(Exception Ex)
 				{
-					string Message = $"Encountered {Ex.GetType()} when creating installation on device {Device}.\n{Ex.Message}";
+					string Message = $"Encountered {Ex.GetType()} when creating installation.\n{Ex.Message}";
 
 					bool IsHostPlatform = Device.Platform == BuildHostPlatform.Current.Platform;
 					EventId EventType = IsHostPlatform ? KnownLogEvents.Gauntlet : KnownLogEvents.Gauntlet_DeviceEvent;
@@ -1695,18 +1695,22 @@ namespace Gauntlet
 						Log.Error(EventType, Message);
 						throw;
 					}
-					else if(IsOverlayException(Ex))
+					else if (IsOverlayException(Ex))
 					{
 						// Errors with Overlay executables are caused by missing files or improper setup, inform user and exit early
-						Log.Error(EventType, Message);
+						Log.Error(EventType, "Device {DeviceName} {Message}", Device.Name, Message);
 						throw;
 					}
-					else
+					else if (Ex is DeviceException)
 					{
 						if (!IsHostPlatform)
 						{
 							UnrealDeviceReservation.MarkProblemDevice(Device, Message);
 						}
+						Log.Info(EventType, "Device {DeviceName} {Message}", Device.Name, Message);
+					}
+					else
+					{
 						Log.Warning(EventType, Message);
 					}
 					throw; // Can consider not throwing here - this would let every device complete setup before releasing the problem devices
@@ -1750,7 +1754,7 @@ namespace Gauntlet
 
 				try
 				{
-					Log.Info("Launching {Install} on {Device}", Install, RolesToDevices[Role]);
+					Log.Info("Launching {Install} on {DeviceName}", Install, RolesToDevices[Role].Name);
 					IAppInstance AppInstance = Install.Run();
 
 					if (AppInstance == null)
@@ -1772,8 +1776,15 @@ namespace Gauntlet
 					}
 
 					string WarningMessage = $"Encountered {Ex.GetType()} when attempting to run install {Ex.Message}";
-					UnrealDeviceReservation.MarkProblemDevice(Install.Device, WarningMessage);
-					Log.Warning(WarningMessage);
+					if (Ex is DeviceException)
+					{
+						Log.Info(KnownLogEvents.Gauntlet_DeviceEvent, "Device {DeviceName} {Message}", Install.Device.Name, WarningMessage);
+						UnrealDeviceReservation.MarkProblemDevice(Install.Device, WarningMessage);
+					}
+					else
+					{
+						Log.Warning(WarningMessage);
+					}
 					throw;
 				}
 			}
diff --git a/Engine/Source/Programs/AutomationTool/Insights/RunInsightsTests.cs b/Engine/Source/Programs/AutomationTool/Insights/RunInsightsTests.cs
--- a/Engine/Source/Programs/AutomationTool/Insights/RunInsightsTests.cs
+++ b/Engine/Source/Programs/AutomationTool/Insights/RunInsightsTests.cs
@@ -303,9 +303,9 @@ namespace InsightsTests
 			SelectedDevicesPrivate.Add(Device);
 		}
 
-		public void MarkProblemDevice(ITargetDevice Device)
+		public void MarkProblemDevice(ITargetDevice Device, string Message)
 		{
-			UnrealDeviceReservation.MarkProblemDevice(Device);
+			UnrealDeviceReservation.MarkProblemDevice(Device, Message);
 		}
 
 		public void ReleaseDevices()
@@ -429,9 +429,14 @@ namespace InsightsTests
 					}
 					catch (Exception Ex)
 					{
-						Log.Info("Failed to install insights tests app onto device {Device}: {Exception}", Device, Ex.ToString());
-
-						InsightsDeviceReservation.GetInstance().MarkProblemDevice(Device);
+						if (Ex is DeviceException)
+						{
+							InsightsDeviceReservation.GetInstance().MarkProblemDevice(Device, $"Failed to install insights tests app onto device: {Ex}");
+						}
+						else
+						{
+							Log.Warning("Failed to install insights tests app onto device {DeviceName}: {Exception}", Device.Name, Ex.ToString());
+						}
 						InsightsDeviceReservation.GetInstance().ReleaseDevices();
 					}
 				}
@@ -446,8 +451,15 @@ namespace InsightsTests
 				}
 				catch (Exception Ex)
 				{
-					Log.Error("Failed to create IAppInstall for insights tests on device {Device}: {Exception}", Device, Ex.ToString());
-					InsightsDeviceReservation.GetInstance().MarkProblemDevice(Device);
+					if (Ex is DeviceException)
+					{
+						InsightsDeviceReservation.GetInstance().MarkProblemDevice(Device, $"Failed to create IAppInstall for insights tests on device: {Ex}");
+					}
+					else
+					{
+						Log.Error("Failed to create IAppInstall for insights tests on device {DeviceName}: {Exception}", Device.Name, Ex.ToString());
+					}
+
 					InsightsDeviceReservation.GetInstance().ReleaseDevices();
 				}
 
@@ -459,8 +471,14 @@ namespace InsightsTests
 				}
 				catch (Exception Ex)
 				{
-					Log.Info("Failed to CleanArtifacts or CopyAdditionalFiles for insights tests on device {Device}: {Exception}", Device, Ex.ToString());
-					InsightsDeviceReservation.GetInstance().MarkProblemDevice(Device);
+					if (Ex is DeviceException)
+					{
+						InsightsDeviceReservation.GetInstance().MarkProblemDevice(Device, $"Failed to CleanArtifacts or CopyAdditionalFiles for insights tests on device: {Ex}");
+					}
+					else
+					{
+						Log.Info("Failed to CleanArtifacts or CopyAdditionalFiles for insights tests on device {DeviceName}: {Exception}", Device.Name, Ex.ToString());
+					}
 					InsightsDeviceReservation.GetInstance().ReleaseDevices();
 				}
 
@@ -479,15 +497,14 @@ namespace InsightsTests
 				}
 				catch (DeviceException DeviceEx)
 				{
-					Log.Warning("Device {0} threw an exception during launch. \nException={1}", Install.Device, DeviceEx.Message);
-					Log.Warning("Failed to start insights test on {0}. Marking as problem device. Will not retry.", Device);
+					Log.Warning("Failed to start insights test on {DeviceName}. Marking as problem device. Will not retry.", Device.Name);
 
 					if (Instance != null)
 					{
 						Instance.Kill();
 					}
 
-					InsightsDeviceReservation.GetInstance().MarkProblemDevice(Device);
+					InsightsDeviceReservation.GetInstance().MarkProblemDevice(Device, $"Device threw an exception during launch. \nException={DeviceEx.Message}");
 					InsightsDeviceReservation.GetInstance().ReleaseDevices();
 
 					throw new AutomationException($"Unable to start insights client app for {ClientInfo.TargetName} platform {Platform}, see warnings for details.");
diff --git a/Engine/Source/Programs/AutomationTool/LowLevelTests/RunLowLevelTests.cs b/Engine/Source/Programs/AutomationTool/LowLevelTests/RunLowLevelTests.cs
--- a/Engine/Source/Programs/AutomationTool/LowLevelTests/RunLowLevelTests.cs
+++ b/Engine/Source/Programs/AutomationTool/LowLevelTests/RunLowLevelTests.cs
@@ -385,7 +385,7 @@ namespace LowLevelTests
 			{
 				// Telemetry
 				DateTimeStopwatch Stopwatch = DateTimeStopwatch.Start();
-				Log.Info("Installing {BuildName} of type {BuildType} to {Device}...", BuildSource.BuildName, AppConfig.Build.GetType().Name, Device);
+				Log.Info("Installing {BuildName} of type {BuildType} to {DeviceName}...", BuildSource.BuildName, AppConfig.Build.GetType().Name, Device.Name);
 				IDeviceUsageReporter.RecordStart(Device.Name, Device.Platform, IDeviceUsageReporter.EventType.Install, IDeviceUsageReporter.EventState.Success, BuildSource.BuildName);
 
 				try
@@ -396,10 +396,18 @@ namespace LowLevelTests
 				}
 				catch (Exception Ex)
 				{
-					Log.Info("Failed to install low level tests app onto device {Device}: {Exception}", Device, Ex.ToString());
-
 					IDeviceUsageReporter.RecordEnd(Device.Name, Device.Platform, IDeviceUsageReporter.EventType.Install, IDeviceUsageReporter.EventState.Failure);
-					UnrealDeviceReservation.MarkProblemDevice(Device);
+
+					if (Ex is DeviceException)
+					{
+						UnrealDeviceReservation.MarkProblemDevice(Device, $"Failed to install low level tests app onto device: {Ex}");
+					}
+					else
+					{
+						Log.Info("Failed to install low level tests app onto device {DeviceName}: {Exception}", Device.Name, Ex.ToString());
+					}
+
+
 					UnrealDeviceReservation.ReleaseDevices();
 				}
 			}
@@ -411,8 +419,15 @@ namespace LowLevelTests
 			}
 			catch (Exception Ex)
 			{
-				Log.Info("Failed to create IAppInstall for low level tests on device {Device}: {Exception}", Device, Ex.ToString());
-				UnrealDeviceReservation.MarkProblemDevice(Device);
+				if (Ex is DeviceException)
+				{
+					UnrealDeviceReservation.MarkProblemDevice(Device, $"Failed to create IAppInstall for low level tests on device: {Ex}");
+				}
+				else
+				{
+					Log.Info("Failed to create IAppInstall for low level tests on device {DeviceName}: {Exception}", Device.Name, Ex.ToString());
+				}
+
 				UnrealDeviceReservation.ReleaseDevices();
 			}
 
@@ -424,8 +439,14 @@ namespace LowLevelTests
 			}
 			catch (Exception Ex)
 			{
-				Log.Info("Failed to CleanArtifacts or CopyAdditionalFiles for low level tests on device {Device}: {Exception}", Device, Ex.ToString());
-				UnrealDeviceReservation.MarkProblemDevice(Device);
+				if (Ex is DeviceException)
+				{
+					UnrealDeviceReservation.MarkProblemDevice(Device, $"Failed to CleanArtifacts or CopyAdditionalFiles for low level tests on device: {Ex}");
+				}
+				else
+				{
+					Log.Info("Failed to CleanArtifacts or CopyAdditionalFiles for low level tests on device {DeviceName}: {Exception}", Device.Name, Ex.ToString());
+				}
 				UnrealDeviceReservation.ReleaseDevices();
 			}
 
@@ -444,15 +465,14 @@ namespace LowLevelTests
 			}
 			catch (DeviceException DeviceEx)
 			{
-				Log.Warning("Device {0} threw an exception during launch. \nException={1}", Install.Device, DeviceEx.Message);
-				Log.Warning("Failed to start low level test on {0}. Marking as problem device. Will not retry.", Device);
+				Log.Warning("Failed to start low level test on {DeviceName}. Marking as problem device. Will not retry.", Device.Name);
 
 				if (Instance != null)
 				{
 					Instance.Kill();
 				}
 
-				UnrealDeviceReservation.MarkProblemDevice(Device);
+				UnrealDeviceReservation.MarkProblemDevice(Device, $"Device threw an exception during launch. \nException={DeviceEx.Message}");
 				UnrealDeviceReservation.ReleaseDevices();
 
 				throw new AutomationException("Unable to start low level tests app, see warnings for details.");
@@ -495,9 +515,14 @@ namespace LowLevelTests
 			{
 				InstallSuccess = false;
 
-				Log.Info("Failed to install low level tests app onto device {0}: {1}", Device, Ex.ToString());
-
-				UnrealDeviceReservation.MarkProblemDevice(Device);
+				if (Ex is DeviceException)
+				{
+					UnrealDeviceReservation.MarkProblemDevice(Device, $"Failed to install low level tests app onto device: {Ex}");
+				}
+				else
+				{
+					Log.Info("Failed to install low level tests app onto device {DeviceName}: {Exception}", Device.Name, Ex.ToString());
+				}
 				IDeviceUsageReporter.RecordEnd(Device.Name, (UnrealTargetPlatform)Device.Platform, IDeviceUsageReporter.EventType.Install, IDeviceUsageReporter.EventState.Failure);
 			}
 
@@ -525,20 +550,19 @@ namespace LowLevelTests
 				}
 				catch (DeviceException DeviceEx)
 				{
-					Log.Warning("Device {0} threw an exception during launch. \nException={1}", Install.Device, DeviceEx.Message);
+					UnrealDeviceReservation.MarkProblemDevice(Device, $"Device threw an exception during launch. \nException={DeviceEx.Message}");
 					RunSuccess = false;
 				}
 
 				if (RunSuccess == false)
 				{
-					Log.Warning("Failed to start low level test on {0}. Marking as problem device. Will not retry.", Device);
+					Log.Warning("Failed to start low level test on {DeviceName}. Marking as problem device. Will not retry.", Device.Name);
 
 					if (Instance != null)
 					{
 						Instance.Kill();
 					}
 
-					UnrealDeviceReservation.MarkProblemDevice(Device);
 					UnrealDeviceReservation.ReleaseDevices();
 
 					throw new AutomationException("Unable to start low level tests app, see warnings for details.");
