Send error message to DeviceError end point

#rnx

diff --git a/Engine/Source/Programs/AutomationTool/AutomationUtils/DeviceReservation.cs b/Engine/Source/Programs/AutomationTool/AutomationUtils/DeviceReservation.cs
--- a/Engine/Source/Programs/AutomationTool/AutomationUtils/DeviceReservation.cs
+++ b/Engine/Source/Programs/AutomationTool/AutomationUtils/DeviceReservation.cs
@@ -282,6 +282,11 @@ namespace AutomationTool.DeviceReservation
 			public string DeviceName { get; set; }
 		}
 
+		public sealed class DeviceErrorRequest
+		{
+			public string Message { get; set; }
+		}
+
 		public static Reservation Create(Uri BaseUri, string[] DeviceTypes, TimeSpan Duration, int RetryMax = 5, string PoolID = "", string DeviceName = "")
 		{
 			bool bFirst = true;
@@ -404,7 +409,7 @@ namespace AutomationTool.DeviceReservation
 			try
 			{
 				Uri BaseUri = new Uri(InBaseUri);
-				Utils.InvokeAPI(BaseUri.AppendPath("api/v1/deviceerror/" + DeviceName), "PUT");
+				Utils.InvokeAPI(BaseUri.AppendPath("api/v1/deviceerror/" + DeviceName), "PUT", new DeviceErrorRequest() { Message = Error });
 				Utils.Log(string.Format("Reported device problem: {0} : {1}", DeviceName, Error));
 			}
 			catch (Exception Ex)
