Fix a case insensitive path inside Gauntlet that caused the CsvInfo.dll not to run

diff --git a/Engine/Source/Programs/AutomationTool/Gauntlet/Framework/Utils/Gauntlet.ReportGenUtils.cs b/Engine/Source/Programs/AutomationTool/Gauntlet/Framework/Utils/Gauntlet.ReportGenUtils.cs
--- a/Engine/Source/Programs/AutomationTool/Gauntlet/Framework/Utils/Gauntlet.ReportGenUtils.cs
+++ b/Engine/Source/Programs/AutomationTool/Gauntlet/Framework/Utils/Gauntlet.ReportGenUtils.cs
@@ -79,7 +79,7 @@ namespace Gauntlet
 		public static string GetCsvToolPath(string ToolName)
 		{
 			string Extension = BuildHostPlatform.Current.Platform.IsInGroup(UnrealPlatformGroup.Microsoft) ? ".exe" : ".dll";
-			return Path.Combine(CommandUtils.CmdEnv.LocalRoot, "Engine", "Binaries", "DotNet", "CsvTools", ToolName) + Extension;
+			return Path.Combine(CommandUtils.CmdEnv.LocalRoot, "Engine", "Binaries", "DotNET", "CsvTools", ToolName) + Extension;
 		}
 
 		public static int RunPythonScript(string ScriptPath, string Args)
