Create Session Role option to disable screenshot gif creation
The method that creates a gif screenshot does a series of operations, including copying of screenshot files to a staging folder, over the whole content of Screenshots folder. These operations can take a while if there are literally thousands of screenshot files, as is the case of cosmetic capture jobs.
#rnx

diff --git a/Engine/Source/Programs/AutomationTool/Gauntlet/Unreal/Base/Gauntlet.UnrealSession.cs b/Engine/Source/Programs/AutomationTool/Gauntlet/Unreal/Base/Gauntlet.UnrealSession.cs
--- a/Engine/Source/Programs/AutomationTool/Gauntlet/Unreal/Base/Gauntlet.UnrealSession.cs
+++ b/Engine/Source/Programs/AutomationTool/Gauntlet/Unreal/Base/Gauntlet.UnrealSession.cs
@@ -21,7 +21,7 @@ namespace Gauntlet
 		Dummy,
 		Null
 	};
-
+	
 	/// <summary>
 	/// Represents a role that will be performed in an Unreal Session
 	/// </summary>
@@ -146,6 +146,11 @@ namespace Gauntlet
 		/// Whether this role will compress screenshots produced as an artifact into a jpeg format
 		/// </summary>
 		public bool CompressScreenshots { get; set; }
+		
+		/// <summary>
+		/// Whether this role will create a gif from screenshots produced as artifacts
+		/// </summary>
+		public bool CreateGifFromScreenshots { get; set; }
 
 		/// <summary>
 		/// Is this role Null?
@@ -173,7 +178,7 @@ namespace Gauntlet
 		/// <param name="InConfiguration"></param>
 		/// <param name="InCommandLine"></param>
 		/// <param name="InOptions"></param>
-		public UnrealSessionRole(UnrealTargetRole InType, UnrealTargetPlatform? InPlatform, UnrealTargetConfiguration InConfiguration, string InCommandLine = null, IConfigOption<UnrealAppConfig> InOptions = null, bool CompressScreenshots = true)
+		public UnrealSessionRole(UnrealTargetRole InType, UnrealTargetPlatform? InPlatform, UnrealTargetConfiguration InConfiguration, string InCommandLine = null, IConfigOption<UnrealAppConfig> InOptions = null, bool CompressScreenshots = true, bool CreateGifFromScreenshots = true)
 		{
 			RoleType = InType;
 
@@ -230,6 +235,7 @@ namespace Gauntlet
 			CommandLineParams = new GauntletCommandLine();
 			RoleModifier = ERoleModifier.None;
 			this.CompressScreenshots = CompressScreenshots;
+			this.CreateGifFromScreenshots = CreateGifFromScreenshots;
 		}
 
         /// <summary>
@@ -1052,9 +1058,9 @@ namespace Gauntlet
 						// Save screenshots and create a gif when not running a server
 						if (!InRunningRole.Role.RoleType.IsServer())
 						{
-							SaveScreenshots(RoleName, SourceDirectory.FullName, DestinationDirectory.FullName, InRunningRole.Role.CompressScreenshots);
+							SaveScreenshots(RoleName, SourceDirectory.FullName, DestinationDirectory.FullName, InRunningRole.Role.CompressScreenshots, InRunningRole.Role.CreateGifFromScreenshots);
 						}
-
+						
 						// Copy remaining artifacts
 						SystemHelpers.CopyDirectory(SourceDirectory.FullName, DestinationDirectory.FullName, SystemHelpers.CopyOptions.Default, TruncateLongPathFilter);
 					}
@@ -1792,7 +1798,7 @@ namespace Gauntlet
 			return new UnrealSessionInstance(RoleInstances.ToArray(), DeferredRolesToInstalls);
 		}
 
-		private void SaveScreenshots(string RoleName, string SourceDirectory, string DestinationDirectory, bool bCompressImages)
+		private void SaveScreenshots(string RoleName, string SourceDirectory, string DestinationDirectory, bool bCompressImages, bool bCreateGif)
 		{
 			try
 			{
@@ -1834,9 +1840,12 @@ namespace Gauntlet
 						}
 
 						string GifPath = GenerateNotTakenFilePath(Path.Combine(DestinationDirectory, RoleName + "Test.gif"));
-						if (Image.SaveImagesAsGif(ImageDirectory.FullName, GifPath))
+						if (bCreateGif)
 						{
-							Log.Info("Saved gif to {0}", GifPath);
+							if (Image.SaveImagesAsGif(ImageDirectory.FullName, GifPath))
+							{
+								Log.Info("Saved gif to {0}", GifPath);
+							}
 						}
 					}
 				}
diff --git a/Engine/Source/Programs/AutomationTool/Gauntlet/Unreal/Base/Gauntlet.UnrealTestConfiguration.cs b/Engine/Source/Programs/AutomationTool/Gauntlet/Unreal/Base/Gauntlet.UnrealTestConfiguration.cs
--- a/Engine/Source/Programs/AutomationTool/Gauntlet/Unreal/Base/Gauntlet.UnrealTestConfiguration.cs
+++ b/Engine/Source/Programs/AutomationTool/Gauntlet/Unreal/Base/Gauntlet.UnrealTestConfiguration.cs
@@ -500,6 +500,11 @@ namespace Gauntlet
 		/// Whether this role will compress screenshots produced as an artifact into a jpeg format
 		/// </summary>
 		public bool CompressScreenshots;
+
+		/// <summary>
+		/// Whether this role will create a gif from screenshots produced as an artifact
+		/// </summary>
+		public bool CreateGifFromScreenshots;
 	}
 
 	/// <summary>
diff --git a/Engine/Source/Programs/AutomationTool/Gauntlet/Unreal/Base/Gauntlet.UnrealTestNode.cs b/Engine/Source/Programs/AutomationTool/Gauntlet/Unreal/Base/Gauntlet.UnrealTestNode.cs
--- a/Engine/Source/Programs/AutomationTool/Gauntlet/Unreal/Base/Gauntlet.UnrealTestNode.cs
+++ b/Engine/Source/Programs/AutomationTool/Gauntlet/Unreal/Base/Gauntlet.UnrealTestNode.cs
@@ -637,7 +637,7 @@ namespace Gauntlet
 					SessionRole.ConfigureDevice = TestRole.ConfigureDevice;
 					SessionRole.MapOverride = TestRole.MapOverride;
 					SessionRole.CompressScreenshots = TestRole.CompressScreenshots;
-
+					SessionRole.CreateGifFromScreenshots = TestRole.CreateGifFromScreenshots;
 					SessionRoles.Add(SessionRole);
 				}
 			}
