[Gauntlet] Fixes duplicate key exception in ExperimentalLaunchFlow when running a second+ iteration with -ReinstallPerPass

diff --git a/Engine/Source/Programs/AutomationTool/Gauntlet/Unreal/Base/Gauntlet.UnrealSession.cs b/Engine/Source/Programs/AutomationTool/Gauntlet/Unreal/Base/Gauntlet.UnrealSession.cs
--- a/Engine/Source/Programs/AutomationTool/Gauntlet/Unreal/Base/Gauntlet.UnrealSession.cs
+++ b/Engine/Source/Programs/AutomationTool/Gauntlet/Unreal/Base/Gauntlet.UnrealSession.cs
@@ -1548,9 +1548,9 @@ namespace Gauntlet
 		/// </summary>
 		private bool RoleNeedsInstall(UnrealSessionRole Role)
 		{
-			return !ReinstallPerPass
-				&& !(RolesToConfigs.ContainsKey(Role) && RolesToConfigs[Role] != null)
-				&& !(RolesToInstalls.ContainsKey(Role) && RolesToInstalls[Role] != null);
+			return ReinstallPerPass
+				|| !(RolesToConfigs.ContainsKey(Role) && RolesToConfigs[Role] != null)
+				|| !(RolesToInstalls.ContainsKey(Role) && RolesToInstalls[Role] != null);
 		}
 
 		/// <summary>
@@ -1632,6 +1632,9 @@ namespace Gauntlet
 				{
 					if (RoleNeedsInstall(Role))
 					{
+						RolesToConfigs.Remove(Role);
+						RolesToInstalls.Remove(Role);
+
 						// Create the app config
 						IEnumerable<UnrealSessionRole> OtherRoles = SessionRoles.Where(Other => Other != Role);
 						AppConfig = BuildSource.CreateConfiguration(Role, OtherRoles);
@@ -1680,9 +1683,9 @@ namespace Gauntlet
 
 					Device.CleanArtifacts();
 
-					if(AppConfig.Build.SupportsAdditionalFileCopy)
+					if (RolesToConfigs[Role].Build.SupportsAdditionalFileCopy)
 					{
-						Device.CopyAdditionalFiles(AppConfig.FilesToCopy);
+						Device.CopyAdditionalFiles(RolesToConfigs[Role].FilesToCopy);
 					}
 
 					Role.ConfigureDevice?.Invoke(Device);
