Avoid timeout test to raise an exception when CreateSimpleReportForHorde_legacy is called

#rnx

diff --git a/Engine/Source/Programs/AutomationTool/Gauntlet/Unreal/Base/Gauntlet.UnrealTestNode.cs b/Engine/Source/Programs/AutomationTool/Gauntlet/Unreal/Base/Gauntlet.UnrealTestNode.cs
--- a/Engine/Source/Programs/AutomationTool/Gauntlet/Unreal/Base/Gauntlet.UnrealTestNode.cs
+++ b/Engine/Source/Programs/AutomationTool/Gauntlet/Unreal/Base/Gauntlet.UnrealTestNode.cs
@@ -1419,6 +1419,8 @@ namespace Gauntlet
 			SimpleReport.HasSucceeded = !(Result == TestResult.Failed || Result == TestResult.TimedOut || SimpleReport.Errors.Count > 0);
 			if (SimpleReport.Errors.Count > 0 && Result == TestResult.Passed)
 			{
+				// Make sure the test is marked complete before setting the test result
+				MarkTestComplete();
 				SetUnrealTestResult(TestResult.Failed);
 			}
 			SimpleReport.Status = GetTestResult().ToString();
