Add DeviceId log property to DeviceReservation utils
+ remove duplicated log when reporting device error

#rnx

diff --git a/Engine/Source/Programs/AutomationTool/AutomationUtils/DeviceReservation.cs b/Engine/Source/Programs/AutomationTool/AutomationUtils/DeviceReservation.cs
--- a/Engine/Source/Programs/AutomationTool/AutomationUtils/DeviceReservation.cs
+++ b/Engine/Source/Programs/AutomationTool/AutomationUtils/DeviceReservation.cs
@@ -6,10 +6,11 @@ using System.Linq;
 using System.Threading;
 using System.Net;
 using System.Text;
-using System.Text.RegularExpressions;
 using System.Text.Json;
 using System.Text.Json.Serialization;
 using UnrealBuildBase;
+using Microsoft.Extensions.Logging;
+using EpicGames.Core;
 
 #pragma warning disable SYSLIB0014
 
@@ -118,7 +119,7 @@ namespace AutomationTool.DeviceReservation
 					// There was an exception, warn if we've exceeded retry
 					if (RetryCurrent == RenewRetryMax)
 					{
-						Utils.Log(string.Format("Warning: Reservation renewal returned bad status: {0}", Ex.Message));
+						Utils.Logger.LogWarning("Reservation renewal returned bad status: {Exception}", Ex.Message);
 					}
 
 					// try again
@@ -134,7 +135,7 @@ namespace AutomationTool.DeviceReservation
 			}
 			catch (Exception Ex)
 			{
-				Utils.Log(string.Format("Warning: Reservation delete returned bad status: {0}", Ex.Message));
+				Utils.Logger.LogWarning("Reservation delete returned bad status: {Exception}", Ex.Message);
 			}
 
 		}
@@ -164,29 +165,7 @@ namespace AutomationTool.DeviceReservation
 	public static class Utils
 	{
 
-		static public string SanitizeErrorMessage(string Message)
-		{
-			string[] TriggersSrc = { "Warning:", "Error:", "Exception:" };
-			string[] TriggersDst = { "Warn1ng:", "Err0r:", "Except10n:" };
-
-			for (int Index = 0; Index < TriggersSrc.Length; ++Index)
-			{
-				if (Message.IndexOf(TriggersSrc[Index], StringComparison.OrdinalIgnoreCase) != -1)
-				{
-					Message = Regex.Replace(Message, TriggersSrc[Index], TriggersDst[Index], RegexOptions.IgnoreCase);
-				}
-			}
-
-			return Message;
-		}
-
-		static public void Log(string Message)
-		{
-			Console.WriteLine("<-- Suspend Log Parsing -->");
-			Console.WriteLine(SanitizeErrorMessage(Message));
-			Console.WriteLine("<-- Resume Log Parsing -->");
-
-		}
+		static public ILogger Logger = EpicGames.Core.Log.Logger;
 
 		public static Uri AppendPath(this Uri BaseUri, string NewPath)
 		{
@@ -270,6 +249,8 @@ namespace AutomationTool.DeviceReservation
 		public static string ReservationDetails = "";
 		public bool? InstallRequired { get; set; } = null;
 
+		static readonly Utf8String DeviceIdAnnotation = new Utf8String("DeviceId");
+
 		private sealed class CreateReservationData
 		{
 			public string[] DeviceTypes { get; set; }
@@ -302,7 +283,7 @@ namespace AutomationTool.DeviceReservation
 
 				bFirst = false;
 
-				Console.WriteLine("Requesting device reservation...");
+				Utils.Logger.LogInformation("Requesting device reservation...");
 
 				Exception UnknownException;
 
@@ -325,7 +306,7 @@ namespace AutomationTool.DeviceReservation
 
 					if (RetryCount == RetryMax)
 					{
-						Console.WriteLine("Device reservation unsuccessful");
+						Utils.Logger.LogWarning("Device reservation unsuccessful");
 						string PoolName = string.IsNullOrEmpty(PoolID) ? "default" : PoolID;
 						string FinalMessage = $"Device reservation for {string.Join(", ", DeviceTypes)} from {PoolName} pool unsuccessful";
 						string MaxRetryWebExMessage = GetReservationErrorMessage(WebEx);
@@ -350,14 +331,14 @@ namespace AutomationTool.DeviceReservation
 						Message = WebExMessage;
 					}
 
-					Console.WriteLine(Message);
+					Utils.Logger.LogWarning(Message);
 					RetryCount++;
 					UnknownException = WebEx;
 				}
 				catch (Exception Ex)
 				{
 					UnknownException = Ex;
-					Utils.Log(string.Format("Device reservation unsuccessful: {0}", UnknownException.Message));
+					Utils.Logger.LogWarning("Device reservation unsuccessful: {Exception}", UnknownException.Message);
 				}
 			}
 		}
@@ -396,7 +377,7 @@ namespace AutomationTool.DeviceReservation
 		public void Delete(Uri BaseUri)
 		{
 			Utils.InvokeAPI(BaseUri.AppendPath("api/v1/reservations/" + Guid.ToString()), "DELETE");
-			Console.WriteLine("Successfully deleted device reservation \"{0}\".", Guid);
+			Utils.Logger.LogInformation("Successfully deleted device reservation \"{GUID}\".", Guid);
 		}
 
 		static public void ReportDeviceError(string InBaseUri, string DeviceName, string Error)
@@ -406,15 +387,17 @@ namespace AutomationTool.DeviceReservation
 				return;
 			}
 
+			LogValue DeviceId = new LogValue(DeviceIdAnnotation, DeviceName);
+
 			try
 			{
 				Uri BaseUri = new Uri(InBaseUri);
 				Utils.InvokeAPI(BaseUri.AppendPath("api/v1/deviceerror/" + DeviceName), "PUT", new DeviceErrorRequest() { Message = Error });
-				Utils.Log(string.Format("Reported device problem: {0} : {1}", DeviceName, Error));
+				Utils.Logger.LogInformation("Reported device problem {DeviceName}: {Error}", DeviceId, Error);
 			}
 			catch (Exception Ex)
 			{
-				Utils.Log(string.Format("Failed to report device: {0} : {1}", DeviceName, Ex.Message));
+				Utils.Logger.LogWarning("Failed to report device {DeviceName}: {Exception}", DeviceId, Ex.Message);
 			}
 		}
 
diff --git a/Engine/Source/Programs/AutomationTool/Gauntlet/Framework/Devices/Gauntlet.DevicePool.cs b/Engine/Source/Programs/AutomationTool/Gauntlet/Framework/Devices/Gauntlet.DevicePool.cs
--- a/Engine/Source/Programs/AutomationTool/Gauntlet/Framework/Devices/Gauntlet.DevicePool.cs
+++ b/Engine/Source/Programs/AutomationTool/Gauntlet/Framework/Devices/Gauntlet.DevicePool.cs
@@ -1155,7 +1155,6 @@ namespace Gauntlet
 				if (Ex is DeviceException)
 				{
 					ReportDeviceError(Def.Name, WarningMessage);
-					Log.Info(KnownLogEvents.Gauntlet_DeviceEvent, WarningMessage);
 				}
 				else
 				{
diff --git a/Engine/Source/Programs/AutomationTool/Gauntlet/Framework/Devices/Gauntlet.UnrealDeviceReservation.cs b/Engine/Source/Programs/AutomationTool/Gauntlet/Framework/Devices/Gauntlet.UnrealDeviceReservation.cs
--- a/Engine/Source/Programs/AutomationTool/Gauntlet/Framework/Devices/Gauntlet.UnrealDeviceReservation.cs
+++ b/Engine/Source/Programs/AutomationTool/Gauntlet/Framework/Devices/Gauntlet.UnrealDeviceReservation.cs
@@ -114,8 +114,6 @@ namespace Gauntlet
 
 				if (LostDevices.Count() > 0)
 				{
-					Log.Info("Lost connection to devices {0} for client(s). ", string.Join(", ", LostDevices));
-
 					// mark these as problems. Could be something grabbed them before us, could be that they are
 					// unresponsive in some way
 					LostDevices.ToList().ForEach(D => MarkProblemDevice(D, "Agent lost connection to device."));
diff --git a/Engine/Source/Programs/AutomationTool/Gauntlet/Unreal/Base/Gauntlet.UnrealSession.cs b/Engine/Source/Programs/AutomationTool/Gauntlet/Unreal/Base/Gauntlet.UnrealSession.cs
--- a/Engine/Source/Programs/AutomationTool/Gauntlet/Unreal/Base/Gauntlet.UnrealSession.cs
+++ b/Engine/Source/Programs/AutomationTool/Gauntlet/Unreal/Base/Gauntlet.UnrealSession.cs
@@ -1396,7 +1396,6 @@ namespace Gauntlet
 							// Warn, ignore the device, and do not continue
 							if (Ex is DeviceException)
 							{
-								Log.Info(KnownLogEvents.Gauntlet_DeviceEvent, ErrorMessage);
 								UnrealDeviceReservation.MarkProblemDevice(Device, ErrorMessage);
 							}
 							else
@@ -1710,13 +1709,9 @@ namespace Gauntlet
 						Log.Error(EventType, "Device {DeviceName} {Message}", Device.Name, Message);
 						throw;
 					}
-					else if (Ex is DeviceException)
+					else if (Ex is DeviceException && !IsHostPlatform)
 					{
-						if (!IsHostPlatform)
-						{
-							UnrealDeviceReservation.MarkProblemDevice(Device, Message);
-						}
-						Log.Info(EventType, "Device {DeviceName} {Message}", Device.Name, Message);
+						UnrealDeviceReservation.MarkProblemDevice(Device, Message);
 					}
 					else
 					{
@@ -1787,7 +1782,6 @@ namespace Gauntlet
 					string WarningMessage = $"Encountered {Ex.GetType()} when attempting to run install {Ex.Message}";
 					if (Ex is DeviceException)
 					{
-						Log.Info(KnownLogEvents.Gauntlet_DeviceEvent, "Device {DeviceName} {Message}", Install.Device.Name, WarningMessage);
 						UnrealDeviceReservation.MarkProblemDevice(Install.Device, WarningMessage);
 					}
 					else
