Fixed exception when failing to provide optional RPC arguments to CallRpc

diff --git a/Engine/Source/Programs/AutomationTool/Gauntlet/Framework/RpcFramework/Gauntlet.RpcExecutor.cs b/Engine/Source/Programs/AutomationTool/Gauntlet/Framework/RpcFramework/Gauntlet.RpcExecutor.cs
--- a/Engine/Source/Programs/AutomationTool/Gauntlet/Framework/RpcFramework/Gauntlet.RpcExecutor.cs
+++ b/Engine/Source/Programs/AutomationTool/Gauntlet/Framework/RpcFramework/Gauntlet.RpcExecutor.cs
@@ -747,13 +747,18 @@ namespace Gauntlet
 			{
 				foreach (RpcArgumentDesc ArgDesc in RpcEntry.Args)
 				{
-					if (!InArgs.ContainsKey(ArgDesc.Name) && !ArgDesc.Optional)
+					object ArgValue;
+
+					if (!InArgs.TryGetValue(ArgDesc.Name, out ArgValue))
 					{
-						Log.Warning(string.Format("RPC {0} lists mandatory argument {1} of type {2} which was not provided. Please either repair argument list or update RPC Definition.",
-							RpcName, ArgDesc.Name, ArgDesc.Type));
+						if ( !ArgDesc.Optional )
+						{
+							Log.Warning(string.Format("RPC {0} lists mandatory argument {1} of type {2} which was not provided. Please either repair argument list or update RPC Definition.",
+								RpcName, ArgDesc.Name, ArgDesc.Type));
+						}
 						continue;
 					}
-					object ArgValue = InArgs[ArgDesc.Name];
+
 					if (ArgValue == null)
 					{
 						Log.Warning(string.Format("RPC {0} argument {1} of expected type {2} is null. If null is a possible value please make value optional in RPC instead.",
