Initial RpcExecutor.CallRpcAsync implementation

diff --git a/Engine/Source/Programs/AutomationTool/Gauntlet/Framework/RpcFramework/Gauntlet.RpcExecutor.cs b/Engine/Source/Programs/AutomationTool/Gauntlet/Framework/RpcFramework/Gauntlet.RpcExecutor.cs
--- a/Engine/Source/Programs/AutomationTool/Gauntlet/Framework/RpcFramework/Gauntlet.RpcExecutor.cs
+++ b/Engine/Source/Programs/AutomationTool/Gauntlet/Framework/RpcFramework/Gauntlet.RpcExecutor.cs
@@ -3,6 +3,7 @@
 using Newtonsoft.Json;
 using System;
 using System.Collections.Generic;
+using System.Diagnostics;
 using System.IO;
 using System.Linq;
 using System.Net;
@@ -11,12 +12,14 @@ using System.Net.Sockets;
 using System.Text;
 using System.Text.RegularExpressions;
 using System.Threading;
+using ThreadState = System.Threading.ThreadState;
 
 namespace Gauntlet
 {
 	/// <summary>
 	/// Contains all relevant information for a registered RPC route on a target
 	/// </summary>
+	[DebuggerDisplay("{Name}")]
 	public class RpcDefinition
 	{
 		public string Name;
@@ -38,6 +41,7 @@ namespace Gauntlet
 	/// <summary>
 	/// Definition of an individual argument expected by an RPC. 
 	/// </summary>
+	[DebuggerDisplay("{Name}")]
 	public class RpcArgumentDesc
 	{
 		public string Name;
@@ -97,7 +101,7 @@ namespace Gauntlet
 	/// <summary>
 	/// Singleton responsible for all of the communication from Gauntlet to the targets.
 	/// </summary>
-	public class RpcExecutor
+	public partial class RpcExecutor
 	{
 		private static RpcExecutor ExecutorInstance = null;
 		private static readonly object ThreadLock = new object();
diff --git a/Engine/Source/Programs/AutomationTool/Gauntlet/Framework/RpcFramework/Gauntlet.RpcExecutorAsync.cs b/Engine/Source/Programs/AutomationTool/Gauntlet/Framework/RpcFramework/Gauntlet.RpcExecutorAsync.cs
new file mode 100644
--- /dev/null
+++ b/Engine/Source/Programs/AutomationTool/Gauntlet/Framework/RpcFramework/Gauntlet.RpcExecutorAsync.cs
@@ -0,0 +1,139 @@
+// Copyright Epic Games, Inc. All Rights Reserved.
+
+#nullable enable
+
+using System;
+using System.Collections.Generic;
+using System.Diagnostics.CodeAnalysis;
+using System.Net.Http;
+using System.Threading;
+using System.Threading.Tasks;
+using Newtonsoft.Json;
+
+namespace Gauntlet;
+
+public partial class RpcExecutor
+{
+	/// <summary>
+	/// Asynchronously send an RPC to the target.
+	/// </summary>
+	/// <param name="Target">
+	/// The <see cref="RpcTarget"/> to send the RPC to.
+	/// </param>
+	/// <param name="RpcName">
+	/// The name of the RPC to send.
+	/// </param>
+	/// <param name="Args">
+	/// A map of parameters
+	/// </param>
+	/// <param name="CancellationToken">
+	/// The <see cref="CancellationToken"/> to monitor for a cancellation request.
+	/// </param>
+	/// <returns>
+	/// A <see cref="Task{TResult}"/> representing the target's response.
+	/// </returns>
+	/// <exception cref="ArgumentException">
+	/// The <paramref name="Target"/>, <paramref name="RpcName"/>, or <paramref name="Args"/> parameters are invalid.
+	/// </exception>
+	public static Task<HttpResponseMessage> CallRpcAsync(
+		RpcTarget Target,
+		string RpcName,
+		IReadOnlyDictionary<string,object> Args,
+		CancellationToken CancellationToken = default)
+	{
+		ArgumentNullException.ThrowIfNull(Target);
+		ArgumentException.ThrowIfNullOrEmpty(RpcName);
+		ArgumentNullException.ThrowIfNull(Args);
+
+		if (!TryGetRpcDefinition(Target, RpcName, out var RpcEntry))
+		{
+			throw new ArgumentException(@"Failed to find RPC definition.", nameof(RpcName));
+		}
+
+		if (!TryGetRpcUrl(Target, RpcName, out var RpcUrl))
+		{
+			// There is no reason this would ever hit as the prior exception would have hit
+			throw new ArgumentException(@"Failed to create RPC url.", nameof(RpcName));
+		}
+
+		ValidateRpcArguments(RpcEntry, Args);
+
+		var MethodToUse = new HttpMethod(RpcEntry.Verb);
+
+		var RpcRequest = new HttpRequestMessage(MethodToUse, RpcUrl);
+
+		RpcRequest.Headers.Add("rpcname", RpcName);
+
+		if (Args.Count > 0)
+		{
+			var JsonBody = JsonConvert.SerializeObject(Args);
+
+			RpcRequest.Content = new StringContent(JsonBody);
+		}
+
+		Log.VeryVerbose($"calling RPC {RpcName} on target {Target.TargetName} at {RpcUrl}");
+
+		return ExecutorClient.SendAsync(RpcRequest, CancellationToken);
+	}
+
+	private static void ValidateRpcArguments(RpcDefinition RpcDefinition, IReadOnlyDictionary<string, object> Args)
+	{
+		// Should this be a custom exception type. RpcArgumentException or something?
+
+		foreach (var RpcArgument in RpcDefinition.Args)
+		{
+			if (!Args.TryGetValue(RpcArgument.Name, out var Value))
+			{
+				if (!RpcArgument.Optional)
+				{
+					throw new ArgumentException($@"Required RPC argument '{RpcArgument.Name}' is missing.");
+				}
+
+				continue;
+			}
+
+			// Most of the existing codebase doesn't have nullable enabled.
+			// Guard against potential null reference exceptions.
+			if (Value is null)
+			{
+				throw new ArgumentException($@"RPC argument '{RpcArgument.Name}' is null. It should be marked as optional if not required.");
+			}
+
+			if (RpcArgumentTypes.TryGetValue(RpcArgument.Type, out var Type))
+			{
+				var ValueType = Value.GetType();
+
+				if (Value.GetType() != Type)
+				{
+					throw new ArgumentException($"Invalid type for RPC argument '{RpcArgument.Name}'. Expected {Type.Name}. Was {ValueType.Name}");
+				}
+			}
+			else
+			{
+				// CallRpc just allows this so we will as well
+			}
+		}
+	}
+
+	private static readonly IReadOnlyDictionary<string, Type> RpcArgumentTypes = new Dictionary<string, Type>
+	{
+		{ "bool", typeof(bool) },
+		{ "int", typeof(int) },
+		{ "float", typeof(float) },
+		{ "double", typeof(double) },
+	};
+
+	private static bool TryGetRpcDefinition(RpcTarget Target, string Name, [NotNullWhen(true)] out RpcDefinition? Definition)
+	{
+		Definition = Target.GetRpc(Name);
+
+		return Definition is not null;
+	}
+
+	private static bool TryGetRpcUrl(RpcTarget Target, string Name, [NotNullWhen(true)] out string? Url)
+	{
+		Url = Target.CreateRpcUrl(Name);
+
+		return !string.IsNullOrEmpty(Url);
+	}
+}
\ No newline at end of file
