[Trivial] remove -StageBaseReleasePaks option - it does not appear to have been wired up since it was added in CL 3297366

#rnx

diff --git a/Engine/Source/Developer/LauncherServices/Private/Launcher/LauncherWorker.cpp b/Engine/Source/Developer/LauncherServices/Private/Launcher/LauncherWorker.cpp
--- a/Engine/Source/Developer/LauncherServices/Private/Launcher/LauncherWorker.cpp
+++ b/Engine/Source/Developer/LauncherServices/Private/Launcher/LauncherWorker.cpp
@@ -774,11 +774,6 @@ FString FLauncherWorker::CreateUATCommand( const ILauncherProfileRef& InProfile,
 				{
 					UATCommand += TEXT(" -basedonreleaseversion=");
 					UATCommand += InProfile->GetBasedOnReleaseVersionName();
-
-					if ( InProfile->ShouldStageBaseReleasePaks() )
-					{
-						UATCommand += TEXT(" -stagebasereleasepaks");
-					}
 				}
 
 				if (InProfile->GetOriginalReleaseVersionName().IsEmpty() == false)
diff --git a/Engine/Source/Developer/LauncherServices/Private/Profiles/LauncherProfile.h b/Engine/Source/Developer/LauncherServices/Private/Profiles/LauncherProfile.h
--- a/Engine/Source/Developer/LauncherServices/Private/Profiles/LauncherProfile.h
+++ b/Engine/Source/Developer/LauncherServices/Private/Profiles/LauncherProfile.h
@@ -3175,7 +3175,7 @@ private:
 	// This build generates a new tier patch file for modified content
 	bool AddPatchLevel;
 
-	// This build stages pak files from the release version it is based on
+	// This build stages pak files from the release version it is based on (deprecated)
 	bool StageBaseReleasePaks;
 
 	// This build will cook content for dlc See also DLCName
diff --git a/Engine/Source/Developer/LauncherServices/Public/ILauncherProfile.h b/Engine/Source/Developer/LauncherServices/Public/ILauncherProfile.h
--- a/Engine/Source/Developer/LauncherServices/Public/ILauncherProfile.h
+++ b/Engine/Source/Developer/LauncherServices/Public/ILauncherProfile.h
@@ -760,6 +760,7 @@ public:
 	/**
 	 * Should we stage the pak files from the base release version this patch is built on
 	 */
+	UE_DEPRECATED(5.6, "ShouldStageBaseReleasePaks is no longer used")
 	virtual bool ShouldStageBaseReleasePaks() const = 0;
 
 	/**
@@ -1166,6 +1167,7 @@ public:
 	 */
 	virtual void SetGeneratePatch( bool InShouldGeneratePatch ) = 0;
 	virtual void SetAddPatchLevel( bool InAddPatchLevel) = 0;
+	UE_DEPRECATED(5.6, "ShouldStageBaseReleasePaks is no longer used")
 	virtual void SetStageBaseReleasePaks(bool InStageBaseReleasePaks) = 0;
 
 	virtual bool IsCreatingDLC() const = 0;
diff --git a/Engine/Source/Developer/ProjectLauncher/Private/Widgets/Cook/SProjectLauncherCookByTheBookSettings.cpp b/Engine/Source/Developer/ProjectLauncher/Private/Widgets/Cook/SProjectLauncherCookByTheBookSettings.cpp
--- a/Engine/Source/Developer/ProjectLauncher/Private/Widgets/Cook/SProjectLauncherCookByTheBookSettings.cpp
+++ b/Engine/Source/Developer/ProjectLauncher/Private/Widgets/Cook/SProjectLauncherCookByTheBookSettings.cpp
@@ -583,23 +583,6 @@ TSharedRef<SWidget> SProjectLauncherCookByTheBookSettings::MakeComplexWidget()
 								]
 						]
 
-					+ SVerticalBox::Slot()
-						.AutoHeight()
-						.Padding(0.0f, 4.0f, 0.0f, 0.0f)
-						[
-							// stage base release pak files checkbox
-							SNew(SCheckBox)
-								.IsChecked(this, &SProjectLauncherCookByTheBookSettings::HandleStageBaseReleasePaksCheckBoxIsChecked)
-								.OnCheckStateChanged(this, &SProjectLauncherCookByTheBookSettings::HandleStageBaseReleasePaksCheckBoxCheckStateChanged)
-								.Padding(FMargin(4.0f, 0.0f))
-								.ToolTipText(LOCTEXT("StageBaseReleasePaksCheckBoxTooltip", "If checked, unchanged pak files present in the base release version will be staged."))
-								.Content()
-								[
-									SNew(STextBlock)
-										.Text(LOCTEXT("StageBaseReleasePaksCheckBoxText", "Stage base release pak files"))
-								]
-						]
-
 					+ SVerticalBox::Slot()
 						.AutoHeight()
 						.Padding(0.0f, 4.0f, 0.0f, 0.0f)
@@ -1644,17 +1627,6 @@ void SProjectLauncherCookByTheBookSettings::HandleAddPatchLevelCheckBoxCheckStat
 }
 
 
-void SProjectLauncherCookByTheBookSettings::HandleStageBaseReleasePaksCheckBoxCheckStateChanged( ECheckBoxState NewState )
-{
-	ILauncherProfilePtr SelectedProfile = Model->GetSelectedProfile();
-
-	if (SelectedProfile.IsValid())
-	{
-		SelectedProfile->SetStageBaseReleasePaks(NewState == ECheckBoxState::Checked);
-	}
-}
-
-
 ECheckBoxState SProjectLauncherCookByTheBookSettings::HandleGeneratePatchCheckBoxIsChecked() const
 {
 	ILauncherProfilePtr SelectedProfile = Model->GetSelectedProfile();
@@ -1687,22 +1659,6 @@ ECheckBoxState SProjectLauncherCookByTheBookSettings::HandleAddPatchLevelCheckBo
 }
 
 
-ECheckBoxState SProjectLauncherCookByTheBookSettings::HandleStageBaseReleasePaksCheckBoxIsChecked( ) const
-{
-	ILauncherProfilePtr SelectedProfile = Model->GetSelectedProfile();
-
-	if (SelectedProfile.IsValid())
-	{
-		if (SelectedProfile->ShouldStageBaseReleasePaks())
-		{
-			return ECheckBoxState::Checked;
-		}
-	}
-
-	return ECheckBoxState::Unchecked;
-}
-
-
 void SProjectLauncherCookByTheBookSettings::HandleCreateReleaseVersionCheckBoxCheckStateChanged(ECheckBoxState NewState)
 {
 	ILauncherProfilePtr SelectedProfile = Model->GetSelectedProfile();
diff --git a/Engine/Source/Developer/ProjectLauncher/Private/Widgets/Cook/SProjectLauncherCookByTheBookSettings.h b/Engine/Source/Developer/ProjectLauncher/Private/Widgets/Cook/SProjectLauncherCookByTheBookSettings.h
--- a/Engine/Source/Developer/ProjectLauncher/Private/Widgets/Cook/SProjectLauncherCookByTheBookSettings.h
+++ b/Engine/Source/Developer/ProjectLauncher/Private/Widgets/Cook/SProjectLauncherCookByTheBookSettings.h
@@ -186,11 +186,6 @@ private:
 	// Callback for determining the checked state of the 'AddPatchLevel' check box.
 	ECheckBoxState HandleAddPatchLevelCheckBoxIsChecked( ) const;
 
-	// Callback for check state changes of the 'StageBaseReleasePaks' check box.
-	void HandleStageBaseReleasePaksCheckBoxCheckStateChanged( ECheckBoxState NewState );
-	// Callback for determining the checked state of the 'StageBaseReleasePaks' check box.
-	ECheckBoxState HandleStageBaseReleasePaksCheckBoxIsChecked( ) const;
-
 	// Callback for changing patchSourceContent path (should be the path of a pak file)
 	FText HandlePatchSourceContentPathTextBlockText() const;
 	// Callback for getting the PatchSourceContentPath
diff --git a/Engine/Source/Programs/AutomationTool/AutomationUtils/ProjectParams.cs b/Engine/Source/Programs/AutomationTool/AutomationUtils/ProjectParams.cs
--- a/Engine/Source/Programs/AutomationTool/AutomationUtils/ProjectParams.cs
+++ b/Engine/Source/Programs/AutomationTool/AutomationUtils/ProjectParams.cs
@@ -307,7 +307,6 @@ namespace AutomationTool
             this.CreateReleaseVersion = InParams.CreateReleaseVersion;
             this.GeneratePatch = InParams.GeneratePatch;
 			this.AddPatchLevel = InParams.AddPatchLevel;
-			this.StageBaseReleasePaks = InParams.StageBaseReleasePaks;
 			this.DLCFile = InParams.DLCFile;
 			this.DiscVersion = InParams.DiscVersion;
             this.DLCIncludeEngineContent = InParams.DLCIncludeEngineContent;
@@ -528,7 +527,6 @@ namespace AutomationTool
 			string ReferenceContainerCryptoKeys = null,
 			bool? GeneratePatch = null,
 			bool? AddPatchLevel = null,
-			bool? StageBaseReleasePaks = null,
             string DiscVersion = null,
             string DLCName = null,
 			string DLCOverrideCookedSubDir = null,
@@ -722,7 +720,6 @@ namespace AutomationTool
             this.BasedOnReleaseVersion = ParseParamValueIfNotSpecified(Command, BasedOnReleaseVersion, "basedonreleaseversion", String.Empty);
             this.GeneratePatch = GetParamValueIfNotSpecified(Command, GeneratePatch, this.GeneratePatch, "GeneratePatch");
             this.AddPatchLevel = GetParamValueIfNotSpecified(Command, AddPatchLevel, this.AddPatchLevel, "AddPatchLevel");
-			this.StageBaseReleasePaks = GetParamValueIfNotSpecified(Command, StageBaseReleasePaks, this.StageBaseReleasePaks, "StageBaseReleasePaks");
 			this.DiscVersion = ParseParamValueIfNotSpecified(Command, DiscVersion, "DiscVersion", String.Empty);
 			this.AdditionalCookerOptions = ParseParamValueIfNotSpecified(Command, AdditionalCookerOptions, "AdditionalCookerOptions", String.Empty);
 		
@@ -1893,11 +1890,6 @@ namespace AutomationTool
         /// </summary>
         public bool AddPatchLevel;
 
-        /// <summary>
-        /// Are we staging the unmodified pak files from the base release
-		/// </summary>
-        public bool StageBaseReleasePaks;
-
 		/// <summary>
 		/// Name of dlc to cook and package (if this paramter is supplied cooks the dlc and packages it into the dlc directory)
 		/// </summary>
@@ -3021,14 +3013,6 @@ namespace AutomationTool
             get { return AddPatchLevel; }
         }
 
-        /// <summary>
-        /// True if we should stage pak files from the base release
-        /// </summary>
-        public bool ShouldStageBaseReleasePaks
-        {
-            get { return StageBaseReleasePaks; }
-        }
-
 		public List<Platform> ClientTargetPlatformInstances
 		{
 			get
@@ -3216,11 +3200,6 @@ namespace AutomationTool
                 throw new AutomationException("Creating a new patch tier requires patch generation");
             }
 
-			if (ShouldStageBaseReleasePaks && !HasBasedOnReleaseVersion)
-			{
-				throw new AutomationException("Staging pak files from the base release requires a base release version");
-			}
-
 			if (HasCreateReleaseVersion && HasDLCName)
             {
                 throw new AutomationException("Can't create a release version at the same time as creating dlc.");
@@ -3333,7 +3312,6 @@ namespace AutomationTool
 				Logger.LogDebug("SkipCookingEditorContent={SkipCookingEditorContent}", SkipCookingEditorContent);
                 Logger.LogDebug("GeneratePatch={GeneratePatch}", GeneratePatch);
 				Logger.LogDebug("AddPatchLevel={AddPatchLevel}", AddPatchLevel);
-				Logger.LogDebug("StageBaseReleasePaks={StageBaseReleasePaks}", StageBaseReleasePaks);
 				Logger.LogDebug("DiscVersion={DiscVersion}", DiscVersion);
 				Logger.LogDebug("CreateReleaseVersion={CreateReleaseVersion}", CreateReleaseVersion);
                 Logger.LogDebug("BasedOnReleaseVersion={BasedOnReleaseVersion}", BasedOnReleaseVersion);
