Fix invalid path used to publish packages on PS4 while using modular build

#rnx

diff --git a/Engine/Source/Programs/AutomationTool/Scripts/CopyBuildToStagingDirectory.Automation.cs b/Engine/Source/Programs/AutomationTool/Scripts/CopyBuildToStagingDirectory.Automation.cs
--- a/Engine/Source/Programs/AutomationTool/Scripts/CopyBuildToStagingDirectory.Automation.cs
+++ b/Engine/Source/Programs/AutomationTool/Scripts/CopyBuildToStagingDirectory.Automation.cs
@@ -5833,7 +5833,19 @@ namespace AutomationScripts
 					string Exe = Target;
 					if (Config != UnrealTargetConfiguration.Development)
 					{
-						Exe = Target + "-" + PlatformName + "-" + Config.ToString();
+						string Separator = "-";
+						foreach (var ProjectTarget in Params.ProjectTargets)
+						{
+							if (Target == ProjectTarget.TargetName)
+							{
+								TargetRules ValidatedRules = ProjectTarget.Rules;
+								TargetRules.ValidateTarget(ValidatedRules, StagePlatform.Type);
+								Separator = ValidatedRules.DecoratedSeparator;
+								break;
+							}
+						}
+
+						Exe = Target + Separator + PlatformName + Separator + Config.ToString();
 					}
 
 					// append the architecture part if needed
diff --git a/Engine/Source/Programs/UnrealBuildTool/Configuration/TargetRules.cs b/Engine/Source/Programs/UnrealBuildTool/Configuration/TargetRules.cs
--- a/Engine/Source/Programs/UnrealBuildTool/Configuration/TargetRules.cs
+++ b/Engine/Source/Programs/UnrealBuildTool/Configuration/TargetRules.cs
@@ -3137,6 +3137,17 @@ namespace UnrealBuildTool
 			return rules;
 		}
 
+		/// <summary>
+		/// Validate the target rules for its current platform
+		/// <param name="target">Rules to validate</param>
+		/// <param name="platform">Explicit platform to use</param>
+		/// </summary>
+		public static void ValidateTarget(TargetRules target, UnrealTargetPlatform platform)
+		{
+			UEBuildPlatform BuildPlatform = UEBuildPlatform.GetBuildPlatform(platform);
+			BuildPlatform.ValidateTarget(target);
+		}
+
 		/// <summary>
 		/// Constructor.
 		/// </summary>
