- Added some retries ExportAppleBuild, to help with strange Auth errors
- Increased some timeouts

diff --git a/Engine/Source/Programs/AutomationTool/Apple/ExportAppleBuild.Automation.cs b/Engine/Source/Programs/AutomationTool/Apple/ExportAppleBuild.Automation.cs
--- a/Engine/Source/Programs/AutomationTool/Apple/ExportAppleBuild.Automation.cs
+++ b/Engine/Source/Programs/AutomationTool/Apple/ExportAppleBuild.Automation.cs
@@ -145,7 +145,7 @@ class ExportAppleBuild : BuildCommand
 			string BundleId, AppVersion, BundleVersion, CommandLine;
 			GetInfoFromXcArchive(XcArchive, out BundleId, out AppVersion, out BundleVersion, out CommandLine);
 
-			Logger.LogInformation("Setting What's New for {BundleId}, app version{Version}, build version {BundleVersion}, commandline {CommandLine}, final whatsnew = {WhatsNew}", BundleId, AppVersion, BundleVersion, CommandLine, WhatsNew);
+			Logger.LogInformation("Setting What's New for {BundleId}, app version {Version}, build version {BundleVersion}, commandline {CommandLine}, final whatsnew = {WhatsNew}", BundleId, AppVersion, BundleVersion, CommandLine, WhatsNew);
 
 			// replace some variables
 			WhatsNew = WhatsNew.Replace("@appversion", AppVersion);
@@ -280,125 +280,141 @@ class ExportAppleBuild : BuildCommand
 		GetInfoFromXcArchive(XcArchive, out BundleId, out AppVersion, out BundleVersion, out _);
 
 
-Console.WriteLine("Checking for version >= {0}", BundleVersion);
+		Logger.LogInformation("Checking for version >= {Version}", BundleVersion);
 		Version BundleVersionNumber = new Version(BundleVersion);
 		using (HttpClient = CreateClient(JWTToken))
 		{
-			string AppId = GetAppId(BundleId);
-
-			// get prerelease version id
-			string SubURL = $"preReleaseVersions?filter[app]={AppId}&filter[version]={AppVersion}&fields[preReleaseVersions]=version";
-//Console.WriteLine("request: {0}", SubURL);
-			VersionsResponse VersionResponse = HttpClient.GetAsync<VersionsResponse>(SubURL, CancellationToken).Result;
-			// if the app version doesn't even exist, then the bundle verrsion can't exist in it!
-			if (VersionResponse.data.Count == 0)
+			int Attempts = 0;
+			while (Attempts < 5)
 			{
-				return true;
-			}
-
-			Version MaxVersion = new Version(0, 0);
-
-			// now walk over all builds until we find one >= the bundle version
-			string VersionId = VersionResponse.data[0].id;
-			SubURL = $"preReleaseVersions/{VersionId}/builds?limit=5&fields[builds]=version";
-//Console.WriteLine("request: {0}", SubURL);
-			bool bHasProcessedAll = false;
-			while (!bHasProcessedAll)
-			{
-				PagedBuildsResponse Response = HttpClient.GetAsync<PagedBuildsResponse>(SubURL, CancellationToken).Result;
-
-				foreach (BuildInfo Build in Response.data)
-				{
-					string VersionString = Build.attributes.version;
-					// append .0.0 if needed
-					if (!VersionString.Contains("."))
-					{
-						VersionString += ".0.0";
-					}
-
-					Version Version = new Version(VersionString);
-					if (Version > MaxVersion)
-					{
-						MaxVersion = Version;
-					}
-				}
-
-				if (Response.links.next == null)
+				if (Attempts > 0)
 				{
-					bHasProcessedAll = true;
+					Thread.Sleep(30 * 1000);
 				}
-				else
+				try
 				{
-					// get suburl for next page
-					SubURL = Response.links.next.Replace(HttpClient.BaseAddress.ToString(), "");
-				}
-			}
+					string AppId = GetAppId(BundleId);
 
-			if (MaxVersion >= BundleVersionNumber)
-			{
-				if (SigningIdentity != null && Increment != null)
-				{
-					Version IncrementVersion = new Version(Increment);
-					if (IncrementVersion == new Version(0,0,0))
+					// get prerelease version id
+					string SubURL = $"preReleaseVersions?filter[app]={AppId}&filter[version]={AppVersion}&fields[preReleaseVersions]=version";
+					VersionsResponse VersionResponse = HttpClient.GetAsync<VersionsResponse>(SubURL, CancellationToken).Result;
+					// if the app version doesn't even exist, then the bundle verrsion can't exist in it!
+					if (VersionResponse.data.Count == 0)
 					{
-						throw new Exception($"Invalid version increment {Increment}");
+						return true;
 					}
 
-					string TempDir = Path.GetTempFileName();
-					File.Delete(TempDir);
-					Directory.CreateDirectory(TempDir);
-
-					FileReference Script = FileReference.Combine(Unreal.EngineDirectory, "Build/BatchFiles/Mac/ResignApp.sh");
-					string FixedXcArchive = TempDir + "/FixedVersion.xcarchive";
-					Version FixedVersion = new Version(MaxVersion.Major + IncrementVersion.Major, MaxVersion.Minor + IncrementVersion.Minor, MaxVersion.Build + IncrementVersion.Build);
-					string Options = $"{Script} -s \"{XcArchive}\" -d \"{FixedXcArchive}\" -i \"{SigningIdentity}\" -bv {FixedVersion}";
+					Version MaxVersion = new Version(0, 0);
 
-					if (MaxVersion == BundleVersionNumber)
-					{
-						Logger.LogInformation("A build exists with same version that is in the .xcarchive ({BundleVersion}). Making a new\n" +
-							".xcarchive ({FixedXcArchive}) with a version using the supplied increment. The new version is {FixedVersion}.",
-							BundleVersion, FixedXcArchive, FixedVersion);
-					}
-					else
+					// now walk over all builds until we find one >= the bundle version
+					string VersionId = VersionResponse.data[0].id;
+					SubURL = $"preReleaseVersions/{VersionId}/builds?limit=5&fields[builds]=version";
+					bool bHasProcessedAll = false;
+					while (!bHasProcessedAll)
 					{
-						Logger.LogInformation("A build exists with LARGER version ({MaxVersion}) that is in the .xcarchive ({BundleVersion})\n" + 
-							"which would cause the uploaded build to be auto-incremented. Making a new .xcarchive ({FixedXcArchive})\n" + 
-							"with a incremented FROM THE LARGER VERSION (not the version in your .xcarchive!). The new version is {FixedVersion}.",
-							MaxVersion, BundleVersion, FixedXcArchive, FixedVersion);
+						PagedBuildsResponse Response = HttpClient.GetAsync<PagedBuildsResponse>(SubURL, CancellationToken).Result;
+
+						foreach (BuildInfo Build in Response.data)
+						{
+							string VersionString = Build.attributes.version;
+							// append .0.0 if needed
+							if (!VersionString.Contains("."))
+							{
+								VersionString += ".0.0";
+							}
+
+							Version Version = new Version(VersionString);
+							if (Version > MaxVersion)
+							{
+								MaxVersion = Version;
+							}
+						}
+
+						if (Response.links.next == null)
+						{
+							bHasProcessedAll = true;
+						}
+						else
+						{
+							// get suburl for next page
+							SubURL = Response.links.next.Replace(HttpClient.BaseAddress.ToString(), "");
+						}
 					}
 
-					Logger.LogInformation("Running: {Script} {Options}...", "/usr/bin/env", Options);
-					int ExitCode;
-					string Output = Utils.RunLocalProcessAndReturnStdOut("/usr/bin/env", Options, out ExitCode);
-					if (ExitCode != 0)
+					if (MaxVersion >= BundleVersionNumber)
 					{
-						Logger.LogError("{Output}", Output);
-						return false;
-					}
-
-					Logger.LogInformation("{Output}", Output);
+						if (SigningIdentity != null && Increment != null)
+						{
+							Version IncrementVersion = new Version(Increment);
+							if (IncrementVersion == new Version(0, 0, 0))
+							{
+								throw new Exception($"Invalid version increment {Increment}");
+							}
+
+							string TempDir = Path.GetTempFileName();
+							File.Delete(TempDir);
+							Directory.CreateDirectory(TempDir);
+
+							FileReference Script = FileReference.Combine(Unreal.EngineDirectory, "Build/BatchFiles/Mac/ResignApp.sh");
+							string FixedXcArchive = TempDir + "/FixedVersion.xcarchive";
+							Version FixedVersion = new Version(MaxVersion.Major + IncrementVersion.Major, MaxVersion.Minor + IncrementVersion.Minor, MaxVersion.Build + IncrementVersion.Build);
+							string Options = $"{Script} -s \"{XcArchive}\" -d \"{FixedXcArchive}\" -i \"{SigningIdentity}\" -bv {FixedVersion}";
+
+							if (MaxVersion == BundleVersionNumber)
+							{
+								Logger.LogInformation("A build exists with same version that is in the .xcarchive ({BundleVersion}). Making a new\n" +
+									".xcarchive ({FixedXcArchive}) with a version using the supplied increment. The new version is {FixedVersion}.",
+									BundleVersion, FixedXcArchive, FixedVersion);
+							}
+							else
+							{
+								Logger.LogInformation("A build exists with LARGER version ({MaxVersion}) that is in the .xcarchive ({BundleVersion})\n" +
+									"which would cause the uploaded build to be auto-incremented. Making a new .xcarchive ({FixedXcArchive})\n" +
+									"with a incremented FROM THE LARGER VERSION (not the version in your .xcarchive!). The new version is {FixedVersion}.",
+									MaxVersion, BundleVersion, FixedXcArchive, FixedVersion);
+							}
+
+							Logger.LogInformation("Running: {Script} {Options}...", "/usr/bin/env", Options);
+							int ExitCode;
+							string Output = Utils.RunLocalProcessAndReturnStdOut("/usr/bin/env", Options, out ExitCode);
+							if (ExitCode != 0)
+							{
+								Logger.LogError("{Output}", Output);
+								return false;
+							}
+
+							Logger.LogInformation("{Output}", Output);
+
+							// update the XcArchive location to the resigned one
+							XcArchive = new DirectoryReference(FixedXcArchive);
+
+							// we know that this xcarchive has a version that is not in the appstore, so return true 
+						}
+						else
+						{
+							Version PossibleVersion = new Version(MaxVersion.Major, MaxVersion.Minor + 1, MaxVersion.Build);
+
+							Logger.LogError("A build exists with same or larger version ({MaxVersion}) than is in the .xcarchive ({BundleVersion}). Either make a new build,\n" +
+								"or modify the build version in the .xcarchive with a command like:\n" +
+								"   Engine/Build/BatchFiles/Mac/ResignApp.sh -s {XcArchive} -d UpdatedVersion.xcarchive -i <signing identity> -bv {NewVersion}\n" +
+								"and then upload that .xcarchive",
+								MaxVersion, BundleVersion, XcArchive, PossibleVersion);
+
+							return false;
+						}
 
-					// update the XcArchive location to the resigned one
-					XcArchive = new DirectoryReference(FixedXcArchive);
+					}
 
-					// we know that this xcarchive has a version that is not in the appstore, so return true 
+					// if we get here, there is either no newer version, or we have successfully re-versioned the .app
+					return true;
 				}
-				else
+				catch (HttpRequestException Ex)
 				{
-					Version PossibleVersion = new Version(MaxVersion.Major, MaxVersion.Minor + 1, MaxVersion.Build);
-
-					Logger.LogError("A build exists with same or larger version ({MaxVersion}) than is in the .xcarchive ({BundleVersion}). Either make a new build,\n" + 
-						"or modify the build version in the .xcarchive with a command like:\n" + 
-						"   Engine/Build/BatchFiles/Mac/ResignApp.sh -s {XcArchive} -d UpdatedVersion.xcarchive -i <signing identity> -bv {NewVersion}\n" + 
-						"and then upload that .xcarchive",
-						MaxVersion, BundleVersion, XcArchive, PossibleVersion);
-
-					return false;
+					Logger.LogWarning("Got an exception while checking for existing version numbers: '{Msg}'. Will try again...", Ex.Message);
 				}
-
 			}
 
-			return true;
+			return false;
 		}
 	}
 
@@ -450,17 +466,25 @@ Console.WriteLine("Checking for version >= {0}", BundleVersion);
 	{
 		using (HttpClient = CreateClient(Token))
 		{
-			// wait 15 second between queries for the version
+			// wait 30 second between queries for the version
 
 			string AppId = GetAppId(BundleId);
 
-			int Delay = 15;
-			int Tries = 10 * 60 / Delay;
+			int Delay = 30;
+			int Tries = 15 * 60 / Delay;
 
 			string BuildId = null;
 			while (BuildId == null && Tries-- > 0)
 			{
-				BuildId = GetBuildId(AppId, PreReleaseVersion, BundleVersion);
+				try
+				{
+					BuildId = GetBuildId(AppId, PreReleaseVersion, BundleVersion);
+				}
+				catch (HttpRequestException Ex)
+				{
+					Logger.LogWarning("Got an exception while checking for looking for uploaded version: '{Msg}'. Will try again...", Ex.Message);
+				}
+
 				if (BuildId == null)
 				{
 					Logger.LogInformation("Build not found yet, waiting and will try again...");
