Fix confusing report of canceled session run

#rnx

diff --git a/Engine/Source/Programs/AutomationTool/Gauntlet/Framework/Gauntlet.TestExecutor.cs b/Engine/Source/Programs/AutomationTool/Gauntlet/Framework/Gauntlet.TestExecutor.cs
--- a/Engine/Source/Programs/AutomationTool/Gauntlet/Framework/Gauntlet.TestExecutor.cs
+++ b/Engine/Source/Programs/AutomationTool/Gauntlet/Framework/Gauntlet.TestExecutor.cs
@@ -790,13 +790,13 @@ namespace Gauntlet
 				try
 				{
 					// Note - we log in this order to try and make it easy to grep the log and find the
-					// artifcat links
+					// artifact links
 					Log.Info("*");
 					Log.Info("****************************************************************");
 					Log.Info("Finished Test: {Name} in {Time:mm\\:ss}", TestInfo, DateTime.Now - TestInfo.TimeSetupEnded);
 
 					// Tell the test it's done. If it still thinks its running it was cancelled
-					TestInfo.TestNode.StopTest(TestIsRunning ? StopReason.MaxDuration : StopReason.Completed);
+					TestInfo.TestNode.StopTest(TestIsRunning && TestInfo.FinalResult != TestResult.Cancelled ? StopReason.MaxDuration : StopReason.Completed);
 					TestInfo.TimeTestEnded = DateTime.Now;
 
 					TestResult NodeResult = TestInfo.TestNode.GetTestResult();
@@ -852,7 +852,7 @@ namespace Gauntlet
 				{
 					if (TestIsRunning)
 					{
-						Log.Warning(KnownLogEvents.Gauntlet_TestEvent, "Cancelled Test {Name} threw an exception while stopping. Ex: {Eception}\n{Callstack}", 
+						Log.Warning(KnownLogEvents.Gauntlet_TestEvent, "Cancelled Test {Name} threw an exception while stopping. Ex: {Exception}\n{Callstack}", 
 							TestInfo.TestNode.Name, ex.Message, ex.StackTrace);
 					}
 					else
