Restore ability to run on pre iOS 17 devices. xcrun devicectl path is now for iOS 17+, fall back to idevicedebug if targetted device is iOS 16 or less.

diff --git a/Engine/Source/Programs/AutomationTool/IOS/IOSPlatform.Automation.cs b/Engine/Source/Programs/AutomationTool/IOS/IOSPlatform.Automation.cs
--- a/Engine/Source/Programs/AutomationTool/IOS/IOSPlatform.Automation.cs
+++ b/Engine/Source/Programs/AutomationTool/IOS/IOSPlatform.Automation.cs
@@ -2072,7 +2072,8 @@ Console.WriteLine("target receipt executable path: {0}", InExecutablePath);
 
 			string Program;
 			string Arguments;
-			if (OperatingSystem.IsMacOS())
+			if (OperatingSystem.IsMacOS() &&
+				GetDevices().Where(Device => Device.Id == Params.DeviceNames[0] && int.Parse(Device.SoftwareVersion.Split('.').First()) < 17).Count() == 0)
 			{
 				// As of iOS17, Apple changed the remote debugger protocol and libimobiledevice has stopped working for launching apps.
 				// Instead, as of Xcode14(ish), Apple added a new "xcrun devicectl" cli to do pretty much everything libimobiledevice/ios-deploy did.
@@ -2083,9 +2084,6 @@ Console.WriteLine("target receipt executable path: {0}", InExecutablePath);
 				Arguments += " --console";	// attach to the console so we can get log output
 				Arguments += " --device " + Params.DeviceNames[0];
 				Arguments += " \"" + BundleIdentifier + "\"";
-				
-				// ClientCmdLine is only relevant when running on a Mac
-				Arguments += " " + ClientCmdLine;
 			}
 			else
 			{
@@ -2096,6 +2094,12 @@ Console.WriteLine("target receipt executable path: {0}", InExecutablePath);
 				Arguments += " run '" + BundleIdentifier + "'";
 			}
 
+			if (OperatingSystem.IsMacOS())
+			{
+				// ClientCmdLine is only relevant when running on a Mac
+				Arguments += " " + ClientCmdLine;
+			}
+			
 			IProcessResult ClientProcess = Run(Program, Arguments, null, ClientRunFlags);
 			
 			// This failure mode is now valid only for use of libimobiledevice on Windows.
