Mac native option during UAT iOS packaging (-macnative)

[RN] Mac native option during UAT iOS packaging (-macnative)

diff --git a/Engine/Build/BatchFiles/Mac/MakeMacNative.sh b/Engine/Build/BatchFiles/Mac/MakeMacNative.sh
new file mode 100755
--- /dev/null
+++ b/Engine/Build/BatchFiles/Mac/MakeMacNative.sh
@@ -0,0 +1,10 @@
+#!/bin/sh
+# Copyright Epic Games, Inc. All Rights Reserved.
+
+IOSBinaryPath=$1
+AppBundleName=$2
+
+rm -rf "$IOSBinaryPath/MacNative"
+mkdir -p "$IOSBinaryPath/MacNative/$AppBundleName/Wrapper"
+ln -s "Wrapper/$AppBundleName" "$IOSBinaryPath/MacNative/$AppBundleName/WrappedBundle"
+rsync -au "$IOSBinaryPath/$AppBundleName" "$IOSBinaryPath/MacNative/$AppBundleName/Wrapper"
diff --git a/Engine/Source/Programs/AutomationTool/Apple/ApplePlatform.Automation.cs b/Engine/Source/Programs/AutomationTool/Apple/ApplePlatform.Automation.cs
--- a/Engine/Source/Programs/AutomationTool/Apple/ApplePlatform.Automation.cs
+++ b/Engine/Source/Programs/AutomationTool/Apple/ApplePlatform.Automation.cs
@@ -289,8 +289,22 @@ public abstract class ApplePlatform : Platform
 					{
 						Logger.LogInformation("Finalized {App} for running fully self-contained", GetFinalAppPath(Target, SC));
 					}
+
+					if (Params.bMakeMacNative)
+					{
+						String MakeMacNativeSHPath = FileReference.Combine(Unreal.EngineDirectory, "Build/BatchFiles/Mac/MakeMacNative.sh").FullName;
+						String Arg1 = FileReference.Combine(Params.RawProjectPath.Directory, "Binaries/IOS").FullName;  // Path to ios binary folder
+						String Arg2 = GetFinalAppPath(Target, SC).GetDirectoryName();                                   // .app name
+						Logger.LogInformation("Generate {1}/MacNative/{2} that can be run directly on an Apple Silicon Mac", Arg1, Arg2);
+						String Result = Utils.RunLocalProcessAndReturnStdOut("/bin/sh", $"\"{MakeMacNativeSHPath}\" \"{Arg1}\" \"{Arg2}\"");
+						if (Result.Length > 0)
+						{
+							Logger.LogInformation(Result);
+						}
+					}
 					Logger.LogInformation("==================================================================================================================================================");
 				}
+
 			}
 		}
 
diff --git a/Engine/Source/Programs/AutomationTool/AutomationUtils/ProjectParams.cs b/Engine/Source/Programs/AutomationTool/AutomationUtils/ProjectParams.cs
--- a/Engine/Source/Programs/AutomationTool/AutomationUtils/ProjectParams.cs
+++ b/Engine/Source/Programs/AutomationTool/AutomationUtils/ProjectParams.cs
@@ -449,6 +449,7 @@ namespace AutomationTool
 			this.ProjectDescriptor = InParams.ProjectDescriptor;
 			this.Upload = InParams.Upload;
 			this.XcodeBuildOptions = InParams.XcodeBuildOptions;
+			this.bMakeMacNative = InParams.bMakeMacNative;
 		}
 
 		/// <summary>
@@ -622,7 +623,8 @@ namespace AutomationTool
 			ParamList<string> InMapsToRebuildHLOD = null,
 			ParamList<string> TitleID = null,
 			string Upload = null,
-			string XcodeBuildOptions = null
+			string XcodeBuildOptions = null,
+			bool bMakeMacNative = false
 			)
 		{
 			//
@@ -1070,6 +1072,7 @@ namespace AutomationTool
 			}
 
 			this.XcodeBuildOptions = ParseParamValueIfNotSpecified(Command, XcodeBuildOptions, "xcodebuildoptions", null);
+			this.bMakeMacNative = GetParamValueIfNotSpecified(Command, null, this.bMakeMacNative, "macnative");
 
 			if (ClientConfigsToBuild == null)
 			{
@@ -2409,6 +2412,9 @@ namespace AutomationTool
 		[Help("XcodeBuildOptions", "Extra options to pass to xcodebuild")]
 		public string XcodeBuildOptions { get; set; }
 
+		[Help("macnative", "Generate a \"Designed for iPad\" .app that can be run natively on a silicon Mac")]
+		public bool bMakeMacNative { get; set; }
+
 		private List<SingleTargetProperties> DetectedTargets;
 		private Dictionary<UnrealTargetPlatform, ConfigHierarchy> LoadedEngineConfigs;
 		private Dictionary<UnrealTargetPlatform, ConfigHierarchy> LoadedGameConfigs;
diff --git a/Engine/Source/Programs/AutomationTool/IOS/IOSPlatform.Automation.cs b/Engine/Source/Programs/AutomationTool/IOS/IOSPlatform.Automation.cs
--- a/Engine/Source/Programs/AutomationTool/IOS/IOSPlatform.Automation.cs
+++ b/Engine/Source/Programs/AutomationTool/IOS/IOSPlatform.Automation.cs
@@ -1045,14 +1045,9 @@ Console.WriteLine("target receipt executable path: {0}", InExecutablePath);
 
 				// now generate the ipa
 				PackageIPA(Params, ProjectGameExeFilename, SC);
-
-				// generate a .app that can be run directly on an Apple Silicon Mac (the Mac needs to be in the mobileprovision, just like an IOS device)
-				// @todo IOSOnMac
-				// CreateIOSOnMacApp(ProjectGameExeFilename, Params, SC);
 			}
 		}
 
-
 		PrintRunTime();
 	}
 
@@ -1314,26 +1309,6 @@ Console.WriteLine("target receipt executable path: {0}", InExecutablePath);
 		}
 	}
 
-	// Generate a .app that can be run directly on an Apple Silicon Mac (the Mac needs to be in the mobileprovision, just like an IOS device)
-	private void CreateIOSOnMacApp(string ProjectGameExeFilename, ProjectParams Params, DeploymentContext SC)
-	{
-		// make a version that can run on the Mac directly
-		string ExeName = Params.IsCodeBasedProject ? SC.StageExecutables[0] : Path.GetFileNameWithoutExtension(ProjectGameExeFilename);
-		string GameName = ExeName.Split("-".ToCharArray())[0];
-		string SourceApp = Path.Combine(Path.GetDirectoryName(ProjectGameExeFilename), "Payload", $"{GameName}.app");
-		string DestOuterApp = Path.Combine(Path.GetDirectoryName(ProjectGameExeFilename), "IOSOnMac", $"{GameName}.app");
-		string DestSymlink = Path.Combine(DestOuterApp, "WrappedBundle");
-		string DestInnerApp = Path.Combine(DestOuterApp, "Wrapper", $"{GameName}.app");
-
-		Logger.LogInformation("Creating IOSOnMac app {DestOuterApp}", DestOuterApp);
-
-		DeleteDirectory_NoExceptions(DestOuterApp);
-		CopyDirectory_NoExceptions(SourceApp, DestInnerApp);
-		// the link uses path relative to the symlink itself
-		File.CreateSymbolicLink(DestSymlink, $"Wrapper/{GameName}.app");
-	}
-
-
 	public override void GetFilesToDeployOrStage(ProjectParams Params, DeploymentContext SC)
 	{
 		//		if (UnrealBuildTool.BuildHostPlatform.Current.Platform != UnrealTargetPlatform.Mac)
