Allow reading destination cloud namespace from .ini files in ZenExportSnapshotTask.

[RN] none

diff --git a/Engine/Source/Programs/AutomationTool/BuildGraph/Tasks/ZenExportSnapshotTask.cs b/Engine/Source/Programs/AutomationTool/BuildGraph/Tasks/ZenExportSnapshotTask.cs
--- a/Engine/Source/Programs/AutomationTool/BuildGraph/Tasks/ZenExportSnapshotTask.cs
+++ b/Engine/Source/Programs/AutomationTool/BuildGraph/Tasks/ZenExportSnapshotTask.cs
@@ -15,6 +15,7 @@ using EpicGames.Core;
 using EpicGames.Horde;
 using EpicGames.ProjectStore;
 using EpicGames.Serialization;
+using UnrealBuildTool;
 using Microsoft.Extensions.Logging;
 
 namespace AutomationTool.Tasks
@@ -455,6 +456,13 @@ namespace AutomationTool.Tasks
 			return StringId.Sanitize(inString).ToString();
 		}
 
+		private struct CloudConfiguration
+		{
+			public string Host = "";
+			public string BuildsNamespace = "";
+			public CloudConfiguration() {}
+		}
+
 		/// <summary>
 		/// ExecuteAsync the task.
 		/// </summary>
@@ -548,13 +556,27 @@ namespace AutomationTool.Tasks
 			{
 				case SnapshotStorageType.Builds:
 				case SnapshotStorageType.Cloud:
+					ProjectProperties properties = ProjectUtils.GetProjectProperties(_parameters.Project);
+					ConfigHierarchy config = properties.EngineConfigs[HostPlatform.Current.HostEditorPlatform];
+					bool foundConfig = config.TryGetValueGeneric("StorageServers", "Cloud", out CloudConfiguration cloudConfig);
+
 					if (String.IsNullOrEmpty(_parameters.DestinationCloudHost))
 					{
-						throw new AutomationException("Missing destination cloud host");
+						if (!foundConfig || String.IsNullOrEmpty(cloudConfig.Host))
+						{
+							throw new AutomationException("Missing destination cloud host");
+						}
+
+						_parameters.DestinationCloudHost = cloudConfig.Host.Split(";")[0];
 					}
 					if (String.IsNullOrEmpty(_parameters.DestinationCloudNamespace))
 					{
-						throw new AutomationException("Missing destination cloud namespace");
+						if (!foundConfig || String.IsNullOrEmpty(cloudConfig.BuildsNamespace))
+						{
+							throw new AutomationException(String.Format("Missing destination cloud namespace {0}", cloudConfig.BuildsNamespace));
+						}
+
+						_parameters.DestinationCloudNamespace = cloudConfig.BuildsNamespace;
 					}
 					if (String.IsNullOrEmpty(_parameters.DestinationIdentifier))
 					{
