Android itemization report fixes

#rnx

diff --git a/Engine/Source/Programs/AutomationTool/Scripts/ReportItemizedExecutableCode.cs b/Engine/Source/Programs/AutomationTool/Scripts/ReportItemizedExecutableCode.cs
--- a/Engine/Source/Programs/AutomationTool/Scripts/ReportItemizedExecutableCode.cs
+++ b/Engine/Source/Programs/AutomationTool/Scripts/ReportItemizedExecutableCode.cs
@@ -35,6 +35,7 @@ namespace ReportItemizedExecutableCode.Automation
 
 	[Help("Reports itemized binary size in a format consumed by PerfReportServer")]
 	[ParamHelp("ExeToItemize", "Absolute path to the binary to itemize", ParamType = typeof(FileReference))]
+	[ParamHelp("SymbolsExe", "Absolute path to the symbolicated binary", ParamType = typeof(FileReference))]
 	class ReportItemizedExecutableCode : BuildCommand
 	{
 		/** Attempts to assign a reasonable module name from a source file name */
@@ -162,7 +163,7 @@ namespace ReportItemizedExecutableCode.Automation
 		}
 
 		/** Attempts to run bloaty on the executable and returns the results */
-		protected ModuleInfos ItemizeExecutableWithBloaty(FileReference Binary)
+		protected ModuleInfos ItemizeExecutableWithBloaty(FileReference Binary, FileReference SymbolsBinary)
 		{
 			ModuleInfos ModInfos = new ModuleInfos();
 			string BloatyExePath = Unreal.RootDirectory.ToString() + @"\Engine\Extras\ThirdPartyNotUE\Bloaty\bloaty.exe";
@@ -172,6 +173,11 @@ namespace ReportItemizedExecutableCode.Automation
 				if (FileReference.Exists(Binary))
 				{
 					string Args = string.Format("-d compileunits -n 0 --csv {0}", Binary.FullName);
+					if (SymbolsBinary != null && FileReference.Exists(SymbolsBinary))
+					{
+						Args += string.Format(" --debug-file={0}", SymbolsBinary.FullName);
+					}
+
 					// Hack!
 					//string Results = HardCodedResults;
 
@@ -238,28 +244,43 @@ namespace ReportItemizedExecutableCode.Automation
 
 						return ModInfos;
 					}
+					else
+					{
+						throw new AutomationException("bloaty failed with exit code {0} and output '{1}'", ExitCode, Results);
+					}
+				}
+				else
+				{
+					throw new AutomationException("binary to itemize is not found at '{0}'", Binary.FullName);
 				}
 			}
-
-			return ModInfos;
+			else
+			{
+				throw new AutomationException("bloaty.exe is not found at '{0}'", BloatyExePath);
+			}
 		}
 
 		/** Whether the binary on this platform can be itemized with bloaty - expected to be overriden in platform-specific classes */
-		virtual public ModuleInfos PlatformItemizeExecutable(FileReference Binary)
+		virtual public ModuleInfos PlatformItemizeExecutable(FileReference Binary, FileReference SymbolsBinary)
 		{
-			return ItemizeExecutableWithBloaty(Binary);
+			return ItemizeExecutableWithBloaty(Binary, SymbolsBinary);
 		}
 
 		public override void ExecuteBuild()
 		{
 			FileReference ExecutableFileToItemize = ParseRequiredFileReferenceParam("ExeToItemize");
+			FileReference ExecutableFileWithSymbols = ParseOptionalFileReferenceParam("SymbolsExe");
 
 			// Hack for quick iteration on formatting (run bloaty manually and just point this at the csv file)
 			//HardCodedResults = FileReference.ReadAllText(ExecutableFileToItemize);
 			
 			System.Console.WriteLine("\n\n\n\nExecuting itemization of executable: {0}", ExecutableFileToItemize);
+			if (ExecutableFileWithSymbols != null)
+			{
+				System.Console.WriteLine("\n\n\n\nUsing executable for symbolication: {0}", ExecutableFileWithSymbols);
+			}
 
-			ModuleInfos Info = PlatformItemizeExecutable(ExecutableFileToItemize);
+			ModuleInfos Info = PlatformItemizeExecutable(ExecutableFileToItemize, ExecutableFileWithSymbols);
 			if (Info.Modules.Count == 0)
 			{
 				throw new AutomationException("Unable to itemize binary '{0}'", ExecutableFileToItemize.ToString());
