Fix failure to launch application when user folder can't be cleaned
- clean artifact paths first before installing
- iterate to find a folder that is not locked for userdir

#rnx

diff --git a/Engine/Source/Programs/AutomationTool/Gauntlet/Framework/Base/Gauntlet.TargetDeviceDesktopCommon.cs b/Engine/Source/Programs/AutomationTool/Gauntlet/Framework/Base/Gauntlet.TargetDeviceDesktopCommon.cs
--- a/Engine/Source/Programs/AutomationTool/Gauntlet/Framework/Base/Gauntlet.TargetDeviceDesktopCommon.cs
+++ b/Engine/Source/Programs/AutomationTool/Gauntlet/Framework/Base/Gauntlet.TargetDeviceDesktopCommon.cs
@@ -102,7 +102,23 @@ namespace Gauntlet
 
 		public void CleanArtifacts()
 		{
-			CleanArtifactDirectory(Path.Combine(UserDir, "Saved"));
+			if (Directory.Exists(UserDir))
+			{
+				// make sure the user dir is clean and not locked by anything
+				try
+				{
+					CleanArtifactDirectory(Path.Combine(UserDir, "Saved"));
+				}
+				catch (IOException)
+				{
+					// otherwise iterate to get a non existing folder
+					int Incr = 0;
+					while (Directory.Exists(UserDir))
+					{
+						UserDir = Path.Combine(LocalCachePath, $"UserDir_{++Incr}");
+					}
+				}
+			}
 			if (InstallCache != null)
 			{
 				CleanArtifactDirectory(GetInstallArtifactPath());
@@ -416,9 +432,6 @@ namespace Gauntlet
 
 					RunOptions |= CommandUtils.ERunOptions.NoStdOutRedirect;
 				}
-
-				// clear artifact path
-				CleanDeviceArtifacts();
 			}
 
 			bool bExperimental = Globals.Params.ParseParam("ExperimentalLaunchFlow");
diff --git a/Engine/Source/Programs/AutomationTool/Gauntlet/Unreal/Base/Gauntlet.UnrealSession.cs b/Engine/Source/Programs/AutomationTool/Gauntlet/Unreal/Base/Gauntlet.UnrealSession.cs
--- a/Engine/Source/Programs/AutomationTool/Gauntlet/Unreal/Base/Gauntlet.UnrealSession.cs
+++ b/Engine/Source/Programs/AutomationTool/Gauntlet/Unreal/Base/Gauntlet.UnrealSession.cs
@@ -1377,6 +1377,9 @@ namespace Gauntlet
 						IDeviceUsageReporter.RecordStart(Device.Name, Device.Platform, IDeviceUsageReporter.EventType.Install, IDeviceUsageReporter.EventState.Success, BuildSource.BuildName);
 						try
 						{
+							// Make sure the artifact paths are clean before installation
+							Device.CleanArtifacts();
+
 							Install = Device.InstallApplication(AppConfig);
 							IDeviceUsageReporter.RecordEnd(Device.Name, Device.Platform, IDeviceUsageReporter.EventType.Install, IDeviceUsageReporter.EventState.Success);
 						}
@@ -1629,6 +1632,8 @@ namespace Gauntlet
 
 				try
 				{
+					Device.CleanArtifacts();
+
 					if (RoleNeedsInstall(Role))
 					{
 						RolesToConfigs.Remove(Role);
@@ -1680,8 +1685,6 @@ namespace Gauntlet
 						AppConfig = RolesToConfigs[Role];
 					}
 
-					Device.CleanArtifacts();
-
 					if (RolesToConfigs[Role].Build.SupportsAdditionalFileCopy)
 					{
 						Device.CopyAdditionalFiles(RolesToConfigs[Role].FilesToCopy);
