Gauntlet: Fixing up EpicRoot to support Linux

diff --git a/Engine/Source/Programs/AutomationTool/Gauntlet/Unreal/Utils/Gauntlet.UnrealHelpers.cs b/Engine/Source/Programs/AutomationTool/Gauntlet/Unreal/Utils/Gauntlet.UnrealHelpers.cs
--- a/Engine/Source/Programs/AutomationTool/Gauntlet/Unreal/Utils/Gauntlet.UnrealHelpers.cs
+++ b/Engine/Source/Programs/AutomationTool/Gauntlet/Unreal/Utils/Gauntlet.UnrealHelpers.cs
@@ -444,16 +444,23 @@ namespace Gauntlet
 
 			if (BuildHostPlatform.Current.Platform == UnrealTargetPlatform.Mac || BuildHostPlatform.Current.Platform == UnrealTargetPlatform.Linux)
 			{
-				string PosixMountPath;
+				string PosixMountPath = "";
+
+				string OSMountBase = "/Volumes";
+				if (BuildHostPlatform.Current.Platform == UnrealTargetPlatform.Linux)
+				{
+					OSMountBase = "/mnt";
+				}
+
 				if(CommandUtils.IsBuildMachine)
 				{
-					string PrimaryRoot = Path.Combine("/Volumes", "epicgames.net", "root");
-					string SecondaryRoot = "/Volumes";
+					string PrimaryRoot = Path.Combine(OSMountBase, "epicgames.net", "root");
+					string SecondaryRoot = OSMountBase;
 					PosixMountPath = Directory.Exists(PrimaryRoot) ? PrimaryRoot : SecondaryRoot;
 				}
 				else
 				{
-					PosixMountPath = Path.Combine("/Volumes", "root");
+					PosixMountPath = Path.Combine(OSMountBase, "root");
 				}
 
 				if (!InPath.Contains("P:"))
