Copy Role.AdditionalArtifactDirectories before moving Artifact outputs avoiding potential user folders being cleanup if -RetainDeviceArtifacts is not specified.

#rnx

diff --git a/Engine/Source/Programs/AutomationTool/Gauntlet/SelfTest/Unreal/Gauntlet.SelfTest.TestUnrealSession.cs b/Engine/Source/Programs/AutomationTool/Gauntlet/SelfTest/Unreal/Gauntlet.SelfTest.TestUnrealSession.cs
--- a/Engine/Source/Programs/AutomationTool/Gauntlet/SelfTest/Unreal/Gauntlet.SelfTest.TestUnrealSession.cs
+++ b/Engine/Source/Programs/AutomationTool/Gauntlet/SelfTest/Unreal/Gauntlet.SelfTest.TestUnrealSession.cs
@@ -347,7 +347,7 @@ namespace Gauntlet.SelfTest
 				ArtifactPath = ArtifactPath
 			};
 
-			LongProcessResult DummyProcess = new("Dummy", "", CommandUtils.ERunOptions.NoStdOutRedirect);
+			LongProcessResult DummyProcess = new("cmd", "/C echo Dummy", CommandUtils.ERunOptions.NoStdOutRedirect);
 			WindowsAppInstance AppInstance = new(Install, DummyProcess, Log.FullName);
 
 			// This is a pretty dumb workaround, but the log file reader runs on a separate thread.
diff --git a/Engine/Source/Programs/AutomationTool/Gauntlet/Unreal/Base/Gauntlet.UnrealSession.cs b/Engine/Source/Programs/AutomationTool/Gauntlet/Unreal/Base/Gauntlet.UnrealSession.cs
--- a/Engine/Source/Programs/AutomationTool/Gauntlet/Unreal/Base/Gauntlet.UnrealSession.cs
+++ b/Engine/Source/Programs/AutomationTool/Gauntlet/Unreal/Base/Gauntlet.UnrealSession.cs
@@ -1023,7 +1023,32 @@ namespace Gauntlet
 			// Whether this is a Dummy, Client, Server, Editor, etc
 			string RoleName = (InRunningRole.Role.IsDummy() ? "Dummy" : "") + InRunningRole.Role.RoleType.ToString();
 
-			// We only want to move artifacts for editor data if there was a crash on a buildmachine.
+			// Move over any additional artifacts that a role requested (before they are removed potentially later)
+			foreach (EIntendedBaseCopyDirectory AdditionalDirectory in InRunningRole.Role.AdditionalArtifactDirectories)
+			{
+				Dictionary<EIntendedBaseCopyDirectory, string> PlatformMappings = InRunningRole.AppInstance.Device.GetPlatformDirectoryMappings();
+				if (PlatformMappings.ContainsKey(AdditionalDirectory))
+				{
+					DirectoryInfo AdditionalSourceDirectory = new DirectoryInfo(PlatformMappings[AdditionalDirectory]);
+					if (AdditionalSourceDirectory.Exists)
+					{
+						string TargetDirectory = Path.Combine(DestinationDirectory.FullName, AdditionalSourceDirectory.Name);
+
+						try
+						{
+							SystemHelpers.CopyDirectory(AdditionalSourceDirectory.FullName, TargetDirectory);
+						}
+						catch (Exception Exception)
+						{
+							Log.Warning("Encountered an {Exception} when trying to copy additional artifact directory {BaseCopyDirectory}" +
+								" from {AdditionalSourceDirectory} to {TargetDirectory}.",
+								Exception, AdditionalDirectory, AdditionalDirectory, TargetDirectory);
+						}
+					}
+				}
+			}
+
+			// We only want to move any other artifacts for editor data if there was a crash on a buildmachine.
 			// Also, don't move artifacts in dev mode, because peoples saved data could be huuuuuuuge!
 			bool IsEditorBuild = InRunningRole.Role.RoleType.UsesEditor();
 			bool IsBuildMachine = CommandUtils.IsBuildMachine;
@@ -1121,31 +1146,6 @@ namespace Gauntlet
 				}
 			}
 
-			// Next, move over any additional artifacts that a role requested
-			foreach (EIntendedBaseCopyDirectory AdditionalDirectory in InRunningRole.Role.AdditionalArtifactDirectories)
-			{
-				Dictionary<EIntendedBaseCopyDirectory, string> PlatformMappings = InRunningRole.AppInstance.Device.GetPlatformDirectoryMappings();
-				if (PlatformMappings.ContainsKey(AdditionalDirectory))
-				{
-					DirectoryInfo AdditionalSourceDirectory = new DirectoryInfo(PlatformMappings[AdditionalDirectory]);
-					if (AdditionalSourceDirectory.Exists)
-					{
-						string TargetDirectory = Path.Combine(DestinationDirectory.FullName, AdditionalSourceDirectory.Name);
-
-						try
-						{
-							SystemHelpers.CopyDirectory(AdditionalSourceDirectory.FullName, TargetDirectory);
-						}
-						catch (Exception Exception)
-						{
-							Log.Warning("Encountered an {Exception} when trying to copy additional artifact directory {BaseCopyDirectory}" +
-								" from {AdditionalSourceDirectory} to {TargetDirectory}.",
-								Exception, AdditionalDirectory, AdditionalDirectory, TargetDirectory);
-						}
-					}
-				}
-			}
-
 			// Now write the role's log file
 			string ArtifactLogFilePath = Path.GetFullPath(Path.Combine(DestinationDirectory.FullName, RoleName + "Output.log"));
 			try
