Avoid skipping install after agent failure to connect to device

#rnx

diff --git a/Engine/Source/Programs/AutomationTool/Gauntlet/Framework/Devices/Gauntlet.DevicePool.cs b/Engine/Source/Programs/AutomationTool/Gauntlet/Framework/Devices/Gauntlet.DevicePool.cs
--- a/Engine/Source/Programs/AutomationTool/Gauntlet/Framework/Devices/Gauntlet.DevicePool.cs
+++ b/Engine/Source/Programs/AutomationTool/Gauntlet/Framework/Devices/Gauntlet.DevicePool.cs
@@ -235,6 +235,8 @@ namespace Gauntlet
 
 		public static bool DeviceReservationBlock;
 
+		private static bool? IsInstallStep;
+
 		/// <summary>
 		/// Device reservation service URL
 		/// </summary>
@@ -842,17 +844,12 @@ namespace Gauntlet
 					return false;
 				}
 
-				if (DeviceReservation.InstallRequired == true)
-				{
-					SkipInstall = false;
-					FullClean = true;
-					DeviceReservationBlock = true;
-				}
-				else if (DeviceReservation.InstallRequired == false)
+				if (DeviceReservationBlock = DeviceReservation.InstallRequired != null)
 				{
-					SkipInstall = true;
-					FullClean = false;
-					DeviceReservationBlock = true;
+					// InstallRequired is true only for the first reservation attempt of the current step. Cache that value to avoid conflicting value on retry attempt.
+					IsInstallStep = IsInstallStep != null? IsInstallStep : DeviceReservation.InstallRequired == true;
+					SkipInstall = DeviceReservation.InstallRequired == false && IsInstallStep == false;
+					FullClean = !SkipInstall;
 				}
 
 				// Add target devices from reservation
@@ -891,12 +888,10 @@ namespace Gauntlet
 						Log.Info(KnownLogEvents.Gauntlet_DeviceEvent, "Unable to make device registration: device registration failed for {Platform}:{DeviceName}", Def.Platform, Def.Name);
 						return false;
 					}
-					else
-					{
-						ScopeReservedDevices.Add(TargetDevice);
-						ReservedDevices.Add(TargetDevice);
-						UsedPlatforms.Add(TargetDevice.Platform);
-					}
+
+					ScopeReservedDevices.Add(TargetDevice);
+					ReservedDevices.Add(TargetDevice);
+					UsedPlatforms.Add(TargetDevice.Platform);
 
 					ServiceDeviceInfo[TargetDevice] = Def;
 					ServiceReservations[TargetDevice] = DeviceReservation;
