Collect trace files from clients, on demand
#rnx

diff --git a/Engine/Source/Programs/AutomationTool/Gauntlet/Unreal/Game/EpicGame.TestConfig.cs b/Engine/Source/Programs/AutomationTool/Gauntlet/Unreal/Game/EpicGame.TestConfig.cs
--- a/Engine/Source/Programs/AutomationTool/Gauntlet/Unreal/Game/EpicGame.TestConfig.cs
+++ b/Engine/Source/Programs/AutomationTool/Gauntlet/Unreal/Game/EpicGame.TestConfig.cs
@@ -1,13 +1,16 @@
 ﻿// Copyright Epic Games, Inc. All Rights Reserved.
 
+using AutomationTool;
 using Gauntlet;
+using System;
+using System.Collections.Generic;
+using System.Data;
+using System.IO;
 using System.Linq;
 using System.Net;
-using System;
-using UnrealBuildTool;
-using AutomationTool;
 using System.Net.NetworkInformation;
-using System.Collections.Generic;
+using UnrealBuildBase;
+using UnrealBuildTool;
 
 namespace EpicGame
 {
@@ -86,6 +89,12 @@ namespace EpicGame
 		[AutoParam]
 		public bool RequiresLogin = false;
 
+		/// <summary>
+		/// Set traces to collect for editor roles, saved to a generated file.
+		/// </summary>
+		[AutoParam]
+		public string EditorTraces { get; set; }
+
 
 		// incrementing value to ensure we can assign unique values to ports etc
 		static private int NumberOfConfigsCreated = 0;
@@ -273,6 +282,15 @@ namespace EpicGame
 					AppConfig.CommandLine += " -nocrashlytics";
 				}
 			}
+
+			if (ConfigRole.RoleType.IsEditor())
+			{
+				if (!string.IsNullOrEmpty(EditorTraces))
+				{
+					string GeneratedTraceFile = Path.Combine(Unreal.EngineDirectory.FullName, "Programs", "AutomationTool", "Saved", "Logs", "EditorTraces.utrace");
+					AppConfig.CommandLine += $" -trace={EditorTraces} -tracefile={GeneratedTraceFile}";
+				}
+			}
 		}		
 	}
 
diff --git a/Engine/Source/Programs/AutomationTool/Gauntlet/Unreal/RunUnreal.cs b/Engine/Source/Programs/AutomationTool/Gauntlet/Unreal/RunUnreal.cs
--- a/Engine/Source/Programs/AutomationTool/Gauntlet/Unreal/RunUnreal.cs
+++ b/Engine/Source/Programs/AutomationTool/Gauntlet/Unreal/RunUnreal.cs
@@ -1,16 +1,16 @@
 ﻿// Copyright Epic Games, Inc. All Rights Reserved.
 
+using AutomationTool;
+using AutomationTool.DeviceReservation;
+using EpicGames.Core;
+using Microsoft.Extensions.Logging;
 using System;
 using System.Collections.Generic;
+using System.IO;
 using System.Linq;
 using System.Text;
-using AutomationTool;
-using AutomationTool.DeviceReservation;
-using UnrealBuildTool;
-using System.IO;
-using EpicGames.Core;
 using UnrealBuildBase;
-using Microsoft.Extensions.Logging;
+using UnrealBuildTool;
 
 namespace Gauntlet
 {
