[Automation] Adds new RemoteServer flag to EpicGamesTestConfig that prevents adding client/server args related to local connections when enabled

diff --git a/Engine/Source/Programs/AutomationTool/Gauntlet/Unreal/Game/EpicGame.TestConfig.cs b/Engine/Source/Programs/AutomationTool/Gauntlet/Unreal/Game/EpicGame.TestConfig.cs
--- a/Engine/Source/Programs/AutomationTool/Gauntlet/Unreal/Game/EpicGame.TestConfig.cs
+++ b/Engine/Source/Programs/AutomationTool/Gauntlet/Unreal/Game/EpicGame.TestConfig.cs
@@ -90,6 +90,11 @@ namespace EpicGame
 		public bool RequiresLogin = false;
 
 		/// <summary>
+		/// If true, do not apply args to client and server as if they were running on the same host
+		/// </summary>
+		[AutoParam]
+		public bool RemoteServer = false;
+
 		/// Set traces to collect for editor roles, saved to a generated file.
 		/// </summary>
 		[AutoParam]
@@ -142,7 +147,7 @@ namespace EpicGame
 				string McpString = "";
 				bool bIsBuildMachine = CommandUtils.IsBuildMachine; 
 
-				if (ConfigRole.RoleType.IsServer())
+				if (ConfigRole.RoleType.IsServer() && !RemoteServer)
 				{
 					// set explicit server and beacon port for online services
 					// this is important when running tests in parallel to avoid matchmaking collisions
@@ -195,13 +200,13 @@ namespace EpicGame
 
 
 				// Do we need to add the -multihome argument to bind to specific IP?
-				if (ConfigRole.RoleType.IsServer() && (MultipleInterfaces || !string.IsNullOrEmpty(RequestedServerIP)))
+				if (ConfigRole.RoleType.IsServer() && !RemoteServer && (MultipleInterfaces || !string.IsNullOrEmpty(RequestedServerIP)))
 				{
 					AppConfig.CommandLine += string.Format(" -multihome={0} -multihomehttp={0}", ServerIP);
 				}
 
 				// client too, but only desktop platforms
-				if (ConfigRole.RoleType.IsClient() && (MultipleInterfaces || !string.IsNullOrEmpty(RequestedClientIP)))
+				if (ConfigRole.RoleType.IsClient() && !RemoteServer && (MultipleInterfaces || !string.IsNullOrEmpty(RequestedClientIP)))
 				{
 					if (ConfigRole.Platform == UnrealTargetPlatform.Win64 || ConfigRole.Platform == UnrealTargetPlatform.Mac)
 					{
@@ -211,6 +216,11 @@ namespace EpicGame
 
 				if (NoMCP)
 				{
+					if(RemoteServer)
+					{
+						throw new AutomationException("Attempted to use a remote server when running a test with nomcp, this is not supported.");
+					}
+
 					McpString += " -nomcp -notimeouts";
 
 					// if this is a client, and there is a server role, find our PC's IP address and tell it to connect to us
