[Gauntlet] Adds LinuxArm64 BuildSource

diff --git a/Engine/Source/Programs/AutomationTool/Gauntlet/Platform/Linux/Gauntlet.LinuxBuildSources.cs b/Engine/Source/Programs/AutomationTool/Gauntlet/Platform/Linux/Gauntlet.LinuxBuildSources.cs
--- a/Engine/Source/Programs/AutomationTool/Gauntlet/Platform/Linux/Gauntlet.LinuxBuildSources.cs
+++ b/Engine/Source/Programs/AutomationTool/Gauntlet/Platform/Linux/Gauntlet.LinuxBuildSources.cs
@@ -13,4 +13,13 @@ namespace Gauntlet
 
 		public override string PlatformFolderPrefix { get { return "Linux"; } }
 	}
+
+	public class LinuxArmBuildSource : StagedBuildSource<StagedBuild>
+	{
+		public override string BuildName { get { return "LinuxArm64StagedBuild"; } }
+
+		public override UnrealTargetPlatform Platform { get { return UnrealTargetPlatform.LinuxArm64; } }
+
+		public override string PlatformFolderPrefix { get { return "LinuxArm64"; } }
+	}
 }
