[Gauntlet] Adds an optional and configurable RolePlatformOverrides dictionary to UnrealTestContext that allows for overriding which roles use a particular platform.
This allows for branches of RunUnreal to specify something like running a Linux server instead of forcing usage of the host platform.
This feature already exists when running with -ClientPlatform= or -ServerPlatform=, but this offers a means of doing it in code.
If specified on commandline, that commandline value will stomp any values placed set by code

diff --git a/Engine/Source/Programs/AutomationTool/Gauntlet/Unreal/Base/Gauntlet.UnrealTestContext.cs b/Engine/Source/Programs/AutomationTool/Gauntlet/Unreal/Base/Gauntlet.UnrealTestContext.cs
--- a/Engine/Source/Programs/AutomationTool/Gauntlet/Unreal/Base/Gauntlet.UnrealTestContext.cs
+++ b/Engine/Source/Programs/AutomationTool/Gauntlet/Unreal/Base/Gauntlet.UnrealTestContext.cs
@@ -85,6 +85,11 @@ namespace Gauntlet
 		/// </summary>
 		public List<ArgumentWithParams> PlatformList = new List<ArgumentWithParams>();
 
+		/// <summary>
+		/// Optional set of platform overrides for a given role
+		/// </summary>
+		public Dictionary<UnrealTargetRole, UnrealTargetPlatform> RolePlatformOverrides = new Dictionary<UnrealTargetRole, UnrealTargetPlatform>();
+
 		/// <summary>
 		/// Tests to run and their params.
 		/// </summary>
diff --git a/Engine/Source/Programs/AutomationTool/Gauntlet/Unreal/RunUnreal.cs b/Engine/Source/Programs/AutomationTool/Gauntlet/Unreal/RunUnreal.cs
--- a/Engine/Source/Programs/AutomationTool/Gauntlet/Unreal/RunUnreal.cs
+++ b/Engine/Source/Programs/AutomationTool/Gauntlet/Unreal/RunUnreal.cs
@@ -212,6 +212,7 @@ namespace Gauntlet
 					if (string.IsNullOrEmpty(PlatformRoleString) == false)
 					{
 						RequestedPlatform = UnrealTargetPlatform.Parse(PlatformRoleString);
+						ContextOptions.RolePlatformOverrides[Type] = RequestedPlatform;
 					}
 
 					if (string.IsNullOrEmpty(ConfigString) == false)
@@ -255,10 +256,13 @@ namespace Gauntlet
 					// TODO - the below is a bit rigid, but maybe that's good enough since the "actually use the editor.." option
 					// is specific to clients and servers
 
-					// client can override platform and config
+					if (ContextOptions.RolePlatformOverrides.ContainsKey(Type))
+					{
+						Role.Platform = ContextOptions.RolePlatformOverrides[Type];
+					}
+					
 					if (Type.IsClient())
 					{
-						Role.Platform = RequestedPlatform;
 						Role.Configuration = RequestedConfiguration;
 
 						if (UsesEditor)
@@ -273,7 +277,6 @@ namespace Gauntlet
 					}
 					else if (Type.IsServer())
 					{
-						// server can only override config
 						Role.Configuration = RequestedConfiguration;
 
 						if (UsesEditor)
