[Gauntlet] Adds Experimental feature flag to globals that enables experimental aspects of Gauntlet when true
Also replaces old references of ExperimentalLaunchFlow with the flag

diff --git a/Engine/Source/Programs/AutomationTool/Gauntlet/Framework/Base/Gauntlet.TargetDeviceDesktopCommon.cs b/Engine/Source/Programs/AutomationTool/Gauntlet/Framework/Base/Gauntlet.TargetDeviceDesktopCommon.cs
--- a/Engine/Source/Programs/AutomationTool/Gauntlet/Framework/Base/Gauntlet.TargetDeviceDesktopCommon.cs
+++ b/Engine/Source/Programs/AutomationTool/Gauntlet/Framework/Base/Gauntlet.TargetDeviceDesktopCommon.cs
@@ -434,8 +434,7 @@ namespace Gauntlet
 				}
 			}
 
-			bool bExperimental = Globals.Params.ParseParam("ExperimentalLaunchFlow");
-			if(bExperimental && Globals.IsRunningDev && AppConfig.OverlayExecutable.GetOverlay(ExecutablePath, out string OverlayExecutable))
+			if(Globals.UseExperimentaFeatures && Globals.IsRunningDev && AppConfig.OverlayExecutable.GetOverlay(ExecutablePath, out string OverlayExecutable))
 			{
 				CommandArguments += string.Format(" -basedir=\"{0}\"", Path.GetDirectoryName(ExecutablePath));
 				ExecutablePath = OverlayExecutable;
diff --git a/Engine/Source/Programs/AutomationTool/Gauntlet/Framework/Gauntlet.Utils.cs b/Engine/Source/Programs/AutomationTool/Gauntlet/Framework/Gauntlet.Utils.cs
--- a/Engine/Source/Programs/AutomationTool/Gauntlet/Framework/Gauntlet.Utils.cs
+++ b/Engine/Source/Programs/AutomationTool/Gauntlet/Framework/Gauntlet.Utils.cs
@@ -56,6 +56,14 @@ namespace Gauntlet
 			}
 		}
 
+		/// <summary>
+		/// This acts as a broad feature flag to allow for a smoother integration of large flow changes without impacting licensees.
+		/// When true, Gauntlet will execute alternative "experimental" code paths in several areas of project, primarily effecting executions of RunUnreal
+		/// Most notable areas this is used:
+		///		- UnrealSession	| Modifies the code path of LaunchSession to better reserve, setup, and handle errors for devices.
+		///		- DevicePool	| Removes the strict expectation that Horde is used as a device service. Allows users to specify new services, such as a cloud provider.
+		/// </summary>
+		public static bool UseExperimentaFeatures => Params.ParseParam("Experimental") || Params.ParseParam("ExperimentalLaunchFlow");
 
 		public static string TempDir
 		{
diff --git a/Engine/Source/Programs/AutomationTool/Gauntlet/Unreal/Base/Gauntlet.UnrealSession.cs b/Engine/Source/Programs/AutomationTool/Gauntlet/Unreal/Base/Gauntlet.UnrealSession.cs
--- a/Engine/Source/Programs/AutomationTool/Gauntlet/Unreal/Base/Gauntlet.UnrealSession.cs
+++ b/Engine/Source/Programs/AutomationTool/Gauntlet/Unreal/Base/Gauntlet.UnrealSession.cs
@@ -816,7 +816,7 @@ namespace Gauntlet
 		/// <returns></returns>
 		public UnrealSessionInstance LaunchSession()
 		{
-			if(!Globals.Params.ParseParam("ExperimentalLaunchFlow"))
+			if(!Globals.UseExperimentaFeatures)
 			{
 				return Legacy_LaunchSession();
 			}
diff --git a/Engine/Source/Programs/AutomationTool/LowLevelTests/RunLowLevelTests.cs b/Engine/Source/Programs/AutomationTool/LowLevelTests/RunLowLevelTests.cs
--- a/Engine/Source/Programs/AutomationTool/LowLevelTests/RunLowLevelTests.cs
+++ b/Engine/Source/Programs/AutomationTool/LowLevelTests/RunLowLevelTests.cs
@@ -359,7 +359,7 @@ namespace LowLevelTests
 		/// </summary>
 		public IAppInstance InstallAndRunNativeTestApp()
 		{
-			if(!Globals.Params.ParseParam("ExperimentalLaunchFlow"))
+			if(!Globals.UseExperimentaFeatures)
 			{
 				return Legacy_InstallAndRunNativeTestApp();
 			}
