[Gauntlet] Updates UnrealBuildSource to search external assemblies for BuildSource types
Also better sanitize log formatting of located builds, notably skipping builds that are discarded by failures of MakeBuildAvailable

diff --git a/Engine/Source/Programs/AutomationTool/Gauntlet/Unreal/BuildSource/Gauntlet.UnrealBuildSource.cs b/Engine/Source/Programs/AutomationTool/Gauntlet/Unreal/BuildSource/Gauntlet.UnrealBuildSource.cs
--- a/Engine/Source/Programs/AutomationTool/Gauntlet/Unreal/BuildSource/Gauntlet.UnrealBuildSource.cs
+++ b/Engine/Source/Programs/AutomationTool/Gauntlet/Unreal/BuildSource/Gauntlet.UnrealBuildSource.cs
@@ -367,6 +367,9 @@ namespace Gauntlet
 					DiscoveredBuilds[NewBuild.Platform] = new List<IBuild>();
 				}
 
+				string Flavor = string.IsNullOrEmpty(NewBuild.Flavor) ? "None" : NewBuild.Flavor;
+				Log.Info("Adding {Platform} {BuildType} | Configuration: {Configuration} | Flavor: {Flavor} | Flags: {Flags} | PreferenceOrder: {PreferenceOrder}", 
+					NewBuild.Platform, NewBuild.GetType().Name, NewBuild.Configuration, Flavor, NewBuild.Flags, NewBuild.PreferenceOrder);
 				DiscoveredBuilds[NewBuild.Platform].Add(NewBuild);
 			}
 		}
@@ -455,7 +458,7 @@ namespace Gauntlet
 				{
 					foreach (string Path in BuildPaths)
 					{
-						IEnumerable<IFolderBuildSource> BuildSources = Gauntlet.Utils.InterfaceHelpers.FindImplementations<IFolderBuildSource>()
+						IEnumerable<IFolderBuildSource> BuildSources = Gauntlet.Utils.InterfaceHelpers.FindImplementations<IFolderBuildSource>(true)
 																			.Where(BS => BS.CanSupportPlatform(InPlatform));
 
 						foreach (var BS in BuildSources)
@@ -469,7 +472,6 @@ namespace Gauntlet
 
 							foreach (IBuild Build in Builds)
 							{
-								Log.Info("Adding build {0} with flags {1} priority {2}", BS.BuildName, Build.Flags, Build.PreferenceOrder);
 								AddBuild(Build);
 							}
 						}
