[Gauntlet] Several LaunchSession flow improvements
- Prevent further attempts to reserve devices when encountering exceptions during device setup. Only attempt with new devices when the caught exception was a DeviceException.
This ensures configuration errors that should be thrown by the like of AutomationExceptions don't result in needless re-attempts at running against the same device

- When running in experimental mode, UnrealTestNode does not reserve devices as part of the ready check. This responsibilty falls into LaunchSession, which has more clear error handling

- Fixes a null reference exception in UnrealSession when the HttpResponse has been disposed or is null due to failed cast

- Fix TryReserveDevices failure from always reporting 0 attempts were made

diff --git a/Engine/Source/Programs/AutomationTool/Gauntlet/Unreal/Base/Gauntlet.UnrealSession.cs b/Engine/Source/Programs/AutomationTool/Gauntlet/Unreal/Base/Gauntlet.UnrealSession.cs
--- a/Engine/Source/Programs/AutomationTool/Gauntlet/Unreal/Base/Gauntlet.UnrealSession.cs
+++ b/Engine/Source/Programs/AutomationTool/Gauntlet/Unreal/Base/Gauntlet.UnrealSession.cs
@@ -752,9 +752,10 @@ namespace Gauntlet
 			{
 				if (ex.InnerException != null && ex.InnerException is WebException
 					&& (ex.InnerException as WebException).Response != null
-					&& ((ex.InnerException as WebException).Response as HttpWebResponse).StatusCode == HttpStatusCode.Conflict)
+					&& ((ex.InnerException as WebException).Response is HttpWebResponse WebResponse)
+					&& WebResponse.StatusCode == HttpStatusCode.Conflict)
 				{
-					// Reservation failed for lack of devices
+					// Reservation failed, usually for lack of devices
 					// Return false to let attempt loop to kick in
 					Log.Info(ex.Message);
 					return false;
@@ -766,7 +767,7 @@ namespace Gauntlet
 
 		public bool TryReserveDevices(int Attempts)
 		{
-			for (; Attempts > 0; --Attempts)
+			for (int RemainingAttempts = Attempts; RemainingAttempts > 0; --RemainingAttempts)
 			{
 				if(Globals.CancelSignalled)
 				{
@@ -784,7 +785,7 @@ namespace Gauntlet
 				}
 			}
 
-			Log.Error("Failed to reserve devices after {Attempts}", Attempts);
+			Log.Error("Failed to reserve devices after {Attempts} attempts", Attempts);
 			ReleaseSessionDevices();
 			return false;
 		}
@@ -862,7 +863,7 @@ namespace Gauntlet
 				}
 				catch (Exception Ex)
 				{
-					if(IsOutOfSpaceException(Ex) || IsOverlayException(Ex))
+					if(IsOutOfSpaceException(Ex) || IsOverlayException(Ex) || Ex is not DeviceException)
 					{
 						RemainingAttempts = 0;
 						ReleaseSessionDevices();
diff --git a/Engine/Source/Programs/AutomationTool/Gauntlet/Unreal/Base/Gauntlet.UnrealTestNode.cs b/Engine/Source/Programs/AutomationTool/Gauntlet/Unreal/Base/Gauntlet.UnrealTestNode.cs
--- a/Engine/Source/Programs/AutomationTool/Gauntlet/Unreal/Base/Gauntlet.UnrealTestNode.cs
+++ b/Engine/Source/Programs/AutomationTool/Gauntlet/Unreal/Base/Gauntlet.UnrealTestNode.cs
@@ -653,7 +653,14 @@ namespace Gauntlet
 				PrepareUnrealApp();
 			}
 
-			return UnrealApp.TryReserveDevices();
+			if(Globals.UseExperimentaFeatures)
+			{
+				return true;
+			}
+			else
+			{
+				return UnrealApp.TryReserveDevices();
+			}
 		}
 
 		/// <summary>
