Made Create/Retrieve artifacts able to read its default values for apporiate values from Engine / Game configs similar to how oplog exports work. This required changing the meaning of the project field from the name of the project to the uproject file path.
Also added baseline branch as a default value setable in the ini.

Updated build and test with publishing cloud artifacts for staged builds as well (similar to what we already can do for packaged builds).
Also tweaked build name used to match what is usually expected out of the build name.
Lastly added `RetrieveFromCloudDDC` option to change artifact retrival to use cloud artifacts.

diff --git a/Engine/Build/Graph/Tasks/BuildAndTestProject.xml b/Engine/Build/Graph/Tasks/BuildAndTestProject.xml
--- a/Engine/Build/Graph/Tasks/BuildAndTestProject.xml
+++ b/Engine/Build/Graph/Tasks/BuildAndTestProject.xml
@@ -18,6 +18,7 @@
 	<Option Name="Horde" DefaultValue="$(IsBuildMachine) And $(IsHordeEnv)" Description="Whether the build is running under Horde"/>
 	
 	<Option Name="PublishToCloudDDC" DefaultValue="false" Description="Enable to publish to cloud ddc instead of fileshare" />
+	<Option Name="RetrieveFromCloudDDC" DefaultValue="false" Description="Enable to retrieve artifacts from cloud ddc" />
 
 	<!-- Project Options -->
 	<Option Name="ProjectFile" DefaultValue="" Restrict=".*" Description="Path to the project file. Can be used instead of specifying ProjectName and ProjectPath separately."/>
@@ -457,6 +458,14 @@
 				<Property Name="TestPrerequisites" Value="$(TestPrerequisites);$(HordeArtifactPlatformRetrieve)" />
 				<Property Name="CookedEditorBuildPath" Value="$(ProjectOutputDirectory)/$(EditorPlatformPath)CookedEditor" />
 			</Do>
+			<Do If="$(RetrieveFromCloudDDC)">
+				<Property Name="CloudArtifactPlatformRetrieve" Value="$(PreNodeName)Retrieve CookedEditor Build Artifact $(EditorPlatform) (CloudDDC)"/>
+				<Node Name="$(CloudArtifactPlatformRetrieve)" Requires="$(TestPrerequisites);$(PlatformPublishNodeName)">
+					<RetrieveCloudArtifact Name="$(BuildNamePath)" Type="staged-build" OutputDir="$(ProjectOutputDirectory)/$(EditorPlatformPath)CookedEditor" Project="$(TargetProject)" Platform="$(EditorPlatform)" />
+				</Node>
+				<Property Name="TestPrerequisites" Value="$(TestPrerequisites);$(CloudArtifactPlatformRetrieve)" />
+				<Property Name="CookedEditorBuildPath" Value="$(ProjectOutputDirectory)/$(EditorPlatformPath)CookedEditor" />
+			</Do>
 			<!-- Declare a boot test for the editor-->
 			<Property Name="PlatformBootTestNodeName" Value="$(PreNodeName)CookedEditor $(EditorPlatform) Test=BootTest" />
 			<Do If="'$(AutomationReportOutputDirectory)' != ''">
@@ -570,12 +579,14 @@
 		<Property Name="PlatformPublishPackagedNodeName" Value="$(PreNodeNameWithPlatform)Publish Packaged $(TargetPlatform)" />
 		<Property Name="PlatformCheckBuildSizeNodeName" Value="$(PreNodeNameWithPlatform)Check Build Size $(TargetPlatform)" />
 		<Property Name="PlatformStagedArtifactNodeName" Value="$(PlatformPublishStagedNodeName) (Horde)" />
+		<Property Name="CloudPlatformStagedArtifactNodeName" Value="$(PlatformPublishStagedNodeName) (CloudDDC)" />
 		<Property Name="HordeArtifactPlatformPackagedNodeName" Value="$(PlatformPublishPackagedNodeName) (Horde)" />
 		<Property Name="HordeArtifactPlatformRetrieve" Value="$(PreNodeNameWithPlatform)Retrieve Build Artifact $(TargetPlatform)" />
 		<Property Name="CloudDDCArtifactPlatformPackagedNodeName" Value="$(PlatformPublishPackagedNodeName) (CloudDDC)" />
 		<Property Name="CloudDDCArtifactPlatformRetrieve" Value="$(PreNodeNameWithPlatform)Retrieve Build Artifact $(TargetPlatform)" />
 		<Property Name="RequirePlatformStaged" Value="$(PlatformStageNodeName)" />
 		<Property Name="RequirePlatformStaged" Value="$(RequirePlatformStaged);$(PlatformStagedArtifactNodeName)" If="$(Horde)"/>
+		<Property Name="RequirePlatformStaged" Value="$(RequirePlatformStaged);$(CloudPlatformStagedArtifactNodeName)" If="$(PublishToCloudDDC)"/>
 		<Property Name="RequirePlatformStaged" Value="$(RequirePlatformStaged);$(PlatformPublishStagedNodeName)" If="!$(SkipPublish) And $(PublishBuilds)"/>
 		<Property Name="RequirePlatformPackaged" Value="$(PlatformPackageNodeName)" />
 		<Property Name="RequirePlatformPackaged" Value="$(RequirePlatformPackaged);$(HordeArtifactPlatformPackagedNodeName)" If="$(Horde)"/>
@@ -743,6 +754,16 @@
 				<Copy Files="..." From="$(ShaderSymbolsDirectory)/*$(TargetPlatform)*" To="$(NetworkOutputDirectory)/$(StagedPlatformFolder)/ShaderSymbols"/>
 			</Node>
 
+			<!-- Publish (Staged) (CloudDDC) -->
+			<Do If="$(PublishToCloudDDC)">
+				<Node Name="$(CloudPlatformStagedArtifactNodeName)" Requires="$(PlatformStageNodeName)">
+					<ForEach Name="PackageConfiguration" Values="$(PackageConfigurations)" Separator="+">
+						<Tag Files="$(ProjectOutputDirectory)/$(StagedPlatformFolder)/..." With="#$(PlatformStageNodeName) Files"/>
+						<CreateCloudArtifact Name="$(BuildNamePath)" Type="staged-build" BaseDir="$(ProjectOutputDirectory)/$(StagedPlatformFolder)" Metadata="Configuration=$(PackageConfiguration);BuildName=$(BuildNamePath);dashboard-category=$(PackageConfiguration) Builds" Project="$(TargetProject)" Platform="$(StagedPlatformFolder)" HttpVersion="$(UE-CloudPublishHttpVersion)" Files="#$(PlatformStageNodeName) Files"/>
+					</ForEach>
+				</Node>
+			</Do>
+
 			<!-- Publish (Staged) (Horde) -->
 			<Do If="$(Horde)">
 				<Node Name="$(PlatformStagedArtifactNodeName)" Requires="$(PlatformStageNodeName)">
@@ -794,7 +815,7 @@
 						<Command Name="BuildCookRun" Arguments="$(BCRArgs) -configuration=$(PackageConfiguration) $(StageAndPackageArguments) -skipbuild -skipcook -skippak $(CloudPackageSkipBCRArguments) -archive -archivedirectory=&quot;$(ProjectOutputDirectory)/Packages/$(StagedPlatformFolder)/$(PackageConfiguration)&quot;" />
 						<Tag Files="$(ProjectOutputDirectory)/Packages/$(StagedPlatformFolder)/$(PackageConfiguration)/..." With="#$(PlatformPackageNodeName) $(PackageConfiguration) Files"/>
 
-						<CreateCloudArtifact Name="$(ProjectName)-Packaged-$(StagedPlatformFolder)-$(PackageConfiguration)" Type="packaged-build" BaseDir="$(ProjectOutputDirectory)/Packages/$(StagedPlatformFolder)/$(PackageConfiguration)" Metadata="Project=$(ProjectName);Stream=$(Branch);Platform=$(StagedPlatformFolder);Configuration=$(PackageConfiguration);BuildName=$(BuildNamePath);dashboard-category=$(PackageConfiguration) Builds" Project="$(ProjectName)" Platform="$(StagedPlatformFolder)" BaselineBranch="$(UE-CloudPublishBaselineBranch)" Files="#$(PlatformPackageNodeName) $(PackageConfiguration) Files"/>
+						<CreateCloudArtifact Name="$(BuildNamePath)" Type="packaged-build" BaseDir="$(ProjectOutputDirectory)/Packages/$(StagedPlatformFolder)/$(PackageConfiguration)" Metadata="Configuration=$(PackageConfiguration);BuildName=$(BuildNamePath);dashboard-category=$(PackageConfiguration) Builds" Project="$(TargetProject)" Platform="$(StagedPlatformFolder)" HttpVersion="$(UE-CloudPublishHttpVersion)" Files="#$(PlatformPackageNodeName) $(PackageConfiguration) Files"/>
 					</ForEach>
 				</Node>
 			</Do>
diff --git a/Engine/Source/Programs/AutomationTool/BuildGraph/Tasks/CreateCloudArtifactTask.cs b/Engine/Source/Programs/AutomationTool/BuildGraph/Tasks/CreateCloudArtifactTask.cs
--- a/Engine/Source/Programs/AutomationTool/BuildGraph/Tasks/CreateCloudArtifactTask.cs
+++ b/Engine/Source/Programs/AutomationTool/BuildGraph/Tasks/CreateCloudArtifactTask.cs
@@ -19,6 +19,7 @@ using EpicGames.Horde.Commits;
 using EpicGames.Horde.Streams;
 using EpicGames.Serialization;
 using Microsoft.Extensions.Logging;
+using UnrealBuildTool;
 using IoHash = EpicGames.Core.IoHash;
 
 #nullable enable
@@ -57,8 +58,8 @@ namespace AutomationTool.Tasks
 		/// <summary>
 		/// The project this artifact is created for
 		/// </summary>
-		[TaskParameter]
-		public string Project { get; set; } = null!;
+		[TaskParameter(Optional = true)]
+		public FileReference Project { get; set; } = null!;
 
 		/// <summary>
 		/// The platform the artifact is created for
@@ -172,9 +173,13 @@ namespace AutomationTool.Tasks
 			List<string> metadataParam = (_parameters.Metadata ?? String.Empty).Split(';', StringSplitOptions.RemoveEmptyEntries | StringSplitOptions.TrimEntries).ToList();
 
 			Dictionary<string, object> metadata = new Dictionary<string, object>();
-			string project = _parameters.Project;
+			string project = _parameters.Project.GetFileNameWithoutExtension();
 			string platform = _parameters.Platform;
 
+			ProjectProperties properties = ProjectUtils.GetProjectProperties(_parameters.Project);
+			ConfigHierarchy config = properties.EngineConfigs[HostPlatform.Current.HostEditorPlatform];
+			bool foundConfig = config.TryGetValueGeneric("StorageServers", "Cloud", out CloudConfiguration cloudConfig);
+
 			foreach (string line in metadataParam)
 			{
 				string[] s = line.Split(":");
@@ -232,9 +237,14 @@ namespace AutomationTool.Tasks
 			metadata.Add("stream", streamId);
 			
 			string baselineBranch = _parameters.BaselineBranch;
+			if (foundConfig && String.IsNullOrEmpty(baselineBranch))
+			{
+				// no explicit baseline branch set and we have found a storage server config, using that value
+				baselineBranch = cloudConfig.BuildsBaselineBranch;
+			}
 			if (String.IsNullOrEmpty(baselineBranch))
 			{
-				Logger.LogWarning("No baseline branch set, falling back to current stream as baseline. This will generate in-optimal de-dup for branches, make sure to set this to your mainline");
+				Logger.LogWarning("No baseline branch set, falling back to current stream as baseline. This will generate in-optimal de-dup for branches, make sure to set this to your mainline by specifying the property or setting this in your StorageServer section in your engine ini");
 				baselineBranch = streamId.ToString();
 			}
 
@@ -263,58 +273,74 @@ namespace AutomationTool.Tasks
 				// no commit specified, not required so ignoring
 			}
 
+			string publishHostEnvVar = "UE-CloudPublishHost";
+			if (!OperatingSystem.IsWindows())
+			{
+				publishHostEnvVar = publishHostEnvVar.Replace("-", "_", StringComparison.OrdinalIgnoreCase);
+			}
 			string host;
 			if (!String.IsNullOrEmpty(_parameters.Host))
 			{
 				host = _parameters.Host;
 			}
-			else
+			else if (!String.IsNullOrEmpty(Environment.GetEnvironmentVariable(publishHostEnvVar)))
 			{
-				string? cloudHostEnvVar = Environment.GetEnvironmentVariable("UE-CloudPublishHost");
-				if (!String.IsNullOrEmpty(cloudHostEnvVar))
-				{
-					host = cloudHostEnvVar;
-				}
-				else
+				host = Environment.GetEnvironmentVariable(publishHostEnvVar)!;
+			}
+			else if (foundConfig && !String.IsNullOrEmpty(cloudConfig.Host))
+			{
+				host = cloudConfig.Host;
+				if (host.Contains(';', StringComparison.OrdinalIgnoreCase))
 				{
-					throw new AutomationException("Missing UE-CloudPublishHost environment variable; unable to determine cloud host.");
+					// if its a list pick the first element
+					host = host.Split(";").First();
 				}
 			}
+			else
+			{
+				throw new AutomationException("Missing UE-CloudPublishHost environment variable; unable to determine cloud host. Please specify this environment variable or define it in your StorageServer section in your engine ini");
+			}
 
+			string publishNamespaceEnvVar = "UE-CloudPublishNamespace";
+			if (!OperatingSystem.IsWindows())
+			{
+				publishNamespaceEnvVar = publishNamespaceEnvVar.Replace("-", "_", StringComparison.OrdinalIgnoreCase);
+			}
 			string ns;
 			if (!String.IsNullOrEmpty(_parameters.Namespace))
 			{
 				ns = _parameters.Namespace;
 			}
+			else if (!String.IsNullOrEmpty(Environment.GetEnvironmentVariable(publishNamespaceEnvVar)))
+			{
+				ns = Environment.GetEnvironmentVariable(publishNamespaceEnvVar)!;
+			}
+			else if (foundConfig && !String.IsNullOrEmpty(cloudConfig.BuildsNamespace))
+			{
+				ns = cloudConfig.BuildsNamespace;
+			}
 			else
 			{
-				string? cloudDefaultNamespace = Environment.GetEnvironmentVariable("UE-CloudPublishNamespace");
-				if (!String.IsNullOrEmpty(cloudDefaultNamespace))
-				{
-					ns = cloudDefaultNamespace;
-				}
-				else
-				{
-					throw new AutomationException("Missing UE-CloudPublishNamespace environment variable; unable to default namespace please specify it in the task.");
-				}
+				throw new AutomationException("Missing UE-CloudPublishNamespace environment variable; unable to default namespace please specify it in the task or define it in your StorageServer section in your engine ini");
 			}
 
+			string accessTokenEnvVar = "UE-CloudDataCacheAccessToken";
+			if (!OperatingSystem.IsWindows())
+			{
+				accessTokenEnvVar = accessTokenEnvVar.Replace("-", "_", StringComparison.OrdinalIgnoreCase);
+			}
 			string accessToken;
 			if (!String.IsNullOrEmpty(_parameters.AccessToken))
 			{
 				accessToken = _parameters.AccessToken;
 			}
+			else if (!String.IsNullOrEmpty(Environment.GetEnvironmentVariable(accessTokenEnvVar)))
+			{
+				accessToken = Environment.GetEnvironmentVariable(accessTokenEnvVar)!;
+			}
 			else
 			{
-				string? cloudAccessTokenEnvVar = Environment.GetEnvironmentVariable("UE-CloudDataCacheAccessToken");
-				if (!String.IsNullOrEmpty(cloudAccessTokenEnvVar))
-				{
-					accessToken = cloudAccessTokenEnvVar;
-				}
-				else
-				{
-					throw new AutomationException("Missing UE-CloudDataCacheAccessToken environment variable; unable to find access token to use.");
-				}
+				throw new AutomationException("Missing UE-CloudDataCacheAccessToken environment variable; unable to find access token to use.");
 			}
 
 			// Resolve the files to include
diff --git a/Engine/Source/Programs/AutomationTool/BuildGraph/Tasks/RetrieveCloudArtifactTask.cs b/Engine/Source/Programs/AutomationTool/BuildGraph/Tasks/RetrieveCloudArtifactTask.cs
--- a/Engine/Source/Programs/AutomationTool/BuildGraph/Tasks/RetrieveCloudArtifactTask.cs
+++ b/Engine/Source/Programs/AutomationTool/BuildGraph/Tasks/RetrieveCloudArtifactTask.cs
@@ -17,6 +17,7 @@ using EpicGames.Horde.Commits;
 using EpicGames.Horde.Streams;
 using EpicGames.Serialization;
 using Microsoft.Extensions.Logging;
+using UnrealBuildTool;
 using Exception = System.Exception;
 
 #nullable enable
@@ -77,11 +78,10 @@ namespace AutomationTool.Tasks
 		public string Platform { get; set; } = null!;
 
 		/// <summary>
-		/// The project this artifact is created for
-		/// TODO: Can this be defaulted?
+		/// The path to the uproject this artifact is created for
 		/// </summary>
 		[TaskParameter]
-		public string Project { get; set; } = null!;
+		public FileReference Project { get; set; } = null!;
 
 		/// <summary>
 		/// The platform the artifact is created for
@@ -154,9 +154,13 @@ namespace AutomationTool.Tasks
 				throw new AutomationException("Cannot specify both Commit and MaxCommit parameters for retrieving an artifact.");
 			}
 
-			string project = _parameters.Project;
+			string project = _parameters.Project.GetFileNameWithoutExtension();
 			string platform = _parameters.Platform;
 
+			ProjectProperties properties = ProjectUtils.GetProjectProperties(_parameters.Project);
+			ConfigHierarchy config = properties.EngineConfigs[HostPlatform.Current.HostEditorPlatform];
+			bool foundConfig = config.TryGetValueGeneric("StorageServers", "Cloud", out CloudConfiguration cloudConfig);
+
 			// Figure out the current change and stream id
 			StreamId streamId;
 			if (!String.IsNullOrEmpty(_parameters.StreamId))
@@ -181,6 +185,10 @@ namespace AutomationTool.Tasks
 			{
 				host = _parameters.Host;
 			}
+			else if (foundConfig && !String.IsNullOrEmpty(cloudConfig.Host))
+			{
+				host = cloudConfig.Host;
+			}
 			else
 			{
 				string? cloudHostEnvVar = Environment.GetEnvironmentVariable("UE-CloudPublishHost");
@@ -190,7 +198,7 @@ namespace AutomationTool.Tasks
 				}
 				else
 				{
-					throw new AutomationException("Missing UE-CloudPublishHost environment variable; unable to determine cloud host.");
+					throw new AutomationException("Missing UE-CloudPublishHost environment variable; unable to determine cloud host. Specify the environment variable or define it in your StorageServer section in your engine ini");
 				}
 			}
 
@@ -199,6 +207,10 @@ namespace AutomationTool.Tasks
 			{
 				ns = _parameters.Namespace;
 			}
+			else if (foundConfig && !String.IsNullOrEmpty(cloudConfig.BuildsNamespace))
+			{
+				ns = cloudConfig.BuildsNamespace;
+			}
 			else
 			{
 				string? cloudDefaultNamespace = Environment.GetEnvironmentVariable("UE-CloudPublishNamespace");
@@ -208,7 +220,7 @@ namespace AutomationTool.Tasks
 				}
 				else
 				{
-					throw new AutomationException("Missing UE-CloudPublishNamespace environment variable; unable to default namespace please specify it in the task.");
+					throw new AutomationException("Missing UE-CloudPublishNamespace environment variable; unable to default namespace please specify it in the task or define it in your StorageServer section in your engine ini.");
 				}
 			}
 
diff --git a/Engine/Source/Programs/AutomationTool/BuildGraph/Tasks/ZenExportSnapshotTask.cs b/Engine/Source/Programs/AutomationTool/BuildGraph/Tasks/ZenExportSnapshotTask.cs
--- a/Engine/Source/Programs/AutomationTool/BuildGraph/Tasks/ZenExportSnapshotTask.cs
+++ b/Engine/Source/Programs/AutomationTool/BuildGraph/Tasks/ZenExportSnapshotTask.cs
@@ -456,15 +456,6 @@ namespace AutomationTool.Tasks
 			return StringId.Sanitize(inString).ToString();
 		}
 
-		private struct CloudConfiguration
-		{
-#pragma warning disable IDE1006 // Static analysis wants these to be named differently, but they must be named the same as the config file properties
-			public string Host = "";
-			public string BuildsNamespace = "";
-#pragma warning restore IDE1006
-			public CloudConfiguration() {}
-		}
-
 		/// <summary>
 		/// ExecuteAsync the task.
 		/// </summary>
@@ -636,7 +627,7 @@ namespace AutomationTool.Tasks
 						string destinationId;
 						if (destinationStorageType == SnapshotStorageType.Builds)
 						{
-							if (!string.IsNullOrEmpty(_parameters.Metadata))
+							if (!String.IsNullOrEmpty(_parameters.Metadata))
 							{
 								exportSingleSourceCommandline.AppendFormat(" --builds-metadata \"{0}\"", _parameters.Metadata);
 							}
@@ -827,4 +818,15 @@ namespace AutomationTool.Tasks
 			yield break;
 		}
 	}
+
+	struct CloudConfiguration
+	{
+#pragma warning disable IDE1006 // Static analysis wants these to be named differently, but they must be named the same as the config file properties
+		public string Host = "";
+		public string BuildsNamespace = "";
+		public string BuildsBaselineBranch = "";
+#pragma warning restore IDE1006
+
+		public CloudConfiguration() {}
+	}
 }
