Added createdAt metadata when uploading cloud artifact, using current utc time formatted as iso 8601.
Also made baselineBranch optional as we can now default this from engine ini.

diff --git a/Engine/Source/Programs/AutomationTool/BuildGraph/Tasks/CreateCloudArtifactTask.cs b/Engine/Source/Programs/AutomationTool/BuildGraph/Tasks/CreateCloudArtifactTask.cs
--- a/Engine/Source/Programs/AutomationTool/BuildGraph/Tasks/CreateCloudArtifactTask.cs
+++ b/Engine/Source/Programs/AutomationTool/BuildGraph/Tasks/CreateCloudArtifactTask.cs
@@ -117,7 +117,7 @@ namespace AutomationTool.Tasks
 		/// <summary>
 		/// The streamid used for your mainline, used to build content as patches on top of this branch if no other baselines exist.
 		/// </summary>
-		[TaskParameter]
+		[TaskParameter(Optional = true)]
 		public string BaselineBranch { get; set; } = null!;
 
 		/// <summary>
@@ -370,6 +370,8 @@ namespace AutomationTool.Tasks
 			long totalSize = files.Sum(x => x.ToFileInfo().Length);
 
 			metadata.Add("totalSize", totalSize);
+			string createdAtString = DateTime.UtcNow.ToString("O");
+			metadata.Add("createdAt", createdAtString);
 			if (_parameters.UseZenCli)
 			{
 				BuildUploader.UploadBuild(host, baseDir, partName, ns, bucketId, metadata, accessToken, new BuildUploader.BuildInfo
