[Automation] Fix issue with client role not being overriden by -platform

diff --git a/Engine/Source/Programs/AutomationTool/Gauntlet/Unreal/RunUnreal.cs b/Engine/Source/Programs/AutomationTool/Gauntlet/Unreal/RunUnreal.cs
--- a/Engine/Source/Programs/AutomationTool/Gauntlet/Unreal/RunUnreal.cs
+++ b/Engine/Source/Programs/AutomationTool/Gauntlet/Unreal/RunUnreal.cs
@@ -256,13 +256,18 @@ namespace Gauntlet
 					// TODO - the below is a bit rigid, but maybe that's good enough since the "actually use the editor.." option
 					// is specific to clients and servers
 
-					if (ContextOptions.RolePlatformOverrides.ContainsKey(Type))
+					bool bPlatformOverride = ContextOptions.RolePlatformOverrides.ContainsKey(Type);
+					if (bPlatformOverride)
 					{
 						Role.Platform = ContextOptions.RolePlatformOverrides[Type];
 					}
 					
 					if (Type.IsClient())
 					{
+						if(!bPlatformOverride)
+						{
+							Role.Platform = RequestedPlatform;
+						}
 						Role.Configuration = RequestedConfiguration;
 
 						if (UsesEditor)
