AutomationTool: If an error is encountered before a log file is set, write to ErrorLog.tx before exiting

#rnx

diff --git a/Engine/Source/Programs/AutomationTool/Program.cs b/Engine/Source/Programs/AutomationTool/Program.cs
--- a/Engine/Source/Programs/AutomationTool/Program.cs
+++ b/Engine/Source/Programs/AutomationTool/Program.cs
@@ -458,6 +458,14 @@ namespace AutomationToolDriver
             }
             finally
             {
+				// If there was an error before a log writer was attached, do so now
+				if (!Log.HasFileWriter() && ReturnCode != ExitCode.Success)
+				{
+					FileReference errorLogFile = FileReference.Combine(Unreal.WritableEngineDirectory, "Programs/AutomationTool/Saved/Logs/ErrorLog.txt");
+					Log.AddFileWriterWithoutBackup("UnhandledFileLogListener", errorLogFile);
+					Logger.LogInformation("AutomationTool encountered an error before a log file was set, see '{File}' for more details", errorLogFile);
+				}
+
 				// Write the exit code
                 Logger.LogInformation("AutomationTool executed for {Time}", Timer.Elapsed.ToString("h'h 'm'm 's's'"));
                 Logger.LogInformation("AutomationTool exiting with ExitCode={ExitCode} ({ExitReason})", (int)ReturnCode, ReturnCode);
@@ -486,7 +494,7 @@ namespace AutomationToolDriver
 					Rules.RulesFileType.AutomationModule, ScriptsForProject, AdditionalScriptDirs, bForceCompile, bNoCompile, bUseBuildRecords, 
 					out bBuildSuccess, (int Count) =>
                     {
-						Logger.LogInformation("Building {Count} projects (see Log 'Engine/Programs/AutomationTool/Saved/Logs/Log.txt' for more details)", Count);
+						Logger.LogInformation("Building {Count} projects...", Count);
 					},
 					Log.Logger);
 
