BuildGraph: Append the Horde job ID to the RetrieveArtifact query

- The behavior of the RetrieveArtifact task is to return the artifact for the most recent job. Add an new optional parameter named ScopeToJobId to search for artifacts from only the currently running job, using the job ID set in the environment variable UE_HORDE_JOBID.
- To use ScopeToJobId add to the RetrieveArtifact element the attribute ScopeToJobId="true".
- Fix the wrong exception being thrown when an artifact is not found. AutomationException will be thrown instead of System.InvalidOperationException.

diff --git a/Engine/Source/Programs/AutomationTool/BuildGraph/Tasks/RetrieveArtifactTask.cs b/Engine/Source/Programs/AutomationTool/BuildGraph/Tasks/RetrieveArtifactTask.cs
--- a/Engine/Source/Programs/AutomationTool/BuildGraph/Tasks/RetrieveArtifactTask.cs
+++ b/Engine/Source/Programs/AutomationTool/BuildGraph/Tasks/RetrieveArtifactTask.cs
@@ -66,6 +66,12 @@ namespace AutomationTool.Tasks
 		/// </summary>
 		[TaskParameter]
 		public string? OutputDir { get; set; }
+
+		/// <summary>
+		/// Append the environment variable UE_HORDE_JOBID to the artifact query
+		/// </summary>
+		[TaskParameter(Optional = true)]
+		public bool ScopeToJobId { get; set; }
 	}
 
 	/// <summary>
@@ -144,6 +150,15 @@ namespace AutomationTool.Tasks
 			ArtifactType type = new ArtifactType(_parameters.Type);
 			List<string> keys = (_parameters.Keys ?? String.Empty).Split(';', StringSplitOptions.RemoveEmptyEntries | StringSplitOptions.TrimEntries).ToList();
 
+			if (_parameters.ScopeToJobId)
+			{
+				string? jobIdEnvVar = Environment.GetEnvironmentVariable("UE_HORDE_JOBID");
+				if (!String.IsNullOrEmpty(jobIdEnvVar))
+				{
+					keys.Add($"job:{jobIdEnvVar}");
+				}
+			}
+
 			DirectoryReference outputDir = ResolveDirectory(_parameters.OutputDir);
 
 			// Format the list of search parameters
@@ -177,7 +192,7 @@ namespace AutomationTool.Tasks
 
 			// Find the new artifact
 			IHordeClient hordeClient = CommandUtils.ServiceProvider.GetRequiredService<IHordeClient>();
-			IArtifact? artifact = await hordeClient.Artifacts.FindAsync(streamId, minCommitId, maxCommitId, name, type, keys, 1).FirstAsync();
+			IArtifact? artifact = await hordeClient.Artifacts.FindAsync(streamId, minCommitId, maxCommitId, name, type, keys, 1).FirstOrDefaultAsync();
 			if (artifact == null)
 			{
 				throw new AutomationException("Unable to find any artifact matching given criteria.");
