[Gauntlet] Fix issue where android artifacts were not being copied because HaxExited only properly checks app run state every 10 seconds.
Instead, whenever HasExited is called, check the status of the app

diff --git a/Engine/Source/Programs/AutomationTool/Gauntlet/Platform/Android/Gauntlet.TargetDeviceAndroid.cs b/Engine/Source/Programs/AutomationTool/Gauntlet/Platform/Android/Gauntlet.TargetDeviceAndroid.cs
--- a/Engine/Source/Programs/AutomationTool/Gauntlet/Platform/Android/Gauntlet.TargetDeviceAndroid.cs
+++ b/Engine/Source/Programs/AutomationTool/Gauntlet/Platform/Android/Gauntlet.TargetDeviceAndroid.cs
@@ -1499,7 +1499,7 @@ namespace Gauntlet
 					return true;
 				}
 
-				return IsActivityRunning();
+				return !IsActivityRunning();
 			}
 		}
 
@@ -1535,8 +1535,6 @@ namespace Gauntlet
 
 		protected bool bHaveSavedArtifacts;
 
-		private static readonly TimeSpan ActivityCheckDelta = TimeSpan.FromSeconds(10);
-		private DateTime ActivityCheckTime = DateTime.UtcNow;
 		private bool ActivityExited = false;
 
 		public AdbScreenRecorder Recorder;
@@ -1624,28 +1622,15 @@ namespace Gauntlet
 		}
 
 		/// <summary>
-		/// Checks on device whether the activity is running, this is an expensive shell with output operation
-		/// the result is cached, with checks at ActivityCheckDelta seconds
+		/// Checks on device whether the activity is running
 		/// </summary>
 		private bool IsActivityRunning()
 		{
-			if (ActivityExited)
-			{
-				return true;
-			}
-
-			if ((DateTime.UtcNow - ActivityCheckTime) < ActivityCheckDelta)
+			if (ActivityExited || AndroidDevice.Disposed)
 			{
 				return false;
 			}
 
-			if (AndroidDevice.Disposed)
-			{
-				return false;
-			}
-
-			ActivityCheckTime = DateTime.UtcNow;
-
 			// get activities filtered by our package name
 			IProcessResult ActivityQuery = AndroidDevice.RunAdbDeviceCommand("shell dumpsys activity -p " + Install.AndroidPackageName + " a");
 
@@ -1665,10 +1650,10 @@ namespace Gauntlet
 				Log.VeryVerbose("{0}: process exited, Activity running={1}, Activity in foreground={2} ", ToString(), bActivityPresent.ToString(), bActivityInForeground.ToString());
 			}
 
-			return bHasExited;
-
+			return !bHasExited;
 		}
 	}
+
 	public class AdbScreenRecorder : IDisposable
 	{
 		private Process ADBProcess;
