[Gauntlet] Fix typo in experimental feature flag property name

diff --git a/Engine/Source/Programs/AutomationTool/Gauntlet/Framework/Base/Gauntlet.TargetDeviceDesktopCommon.cs b/Engine/Source/Programs/AutomationTool/Gauntlet/Framework/Base/Gauntlet.TargetDeviceDesktopCommon.cs
--- a/Engine/Source/Programs/AutomationTool/Gauntlet/Framework/Base/Gauntlet.TargetDeviceDesktopCommon.cs
+++ b/Engine/Source/Programs/AutomationTool/Gauntlet/Framework/Base/Gauntlet.TargetDeviceDesktopCommon.cs
@@ -434,7 +434,7 @@ namespace Gauntlet
 				}
 			}
 
-			if(Globals.UseExperimentaFeatures && Globals.IsRunningDev && AppConfig.OverlayExecutable.GetOverlay(ExecutablePath, out string OverlayExecutable))
+			if(Globals.UseExperimentalFeatures && Globals.IsRunningDev && AppConfig.OverlayExecutable.GetOverlay(ExecutablePath, out string OverlayExecutable))
 			{
 				CommandArguments += string.Format(" -basedir=\"{0}\"", Path.GetDirectoryName(ExecutablePath));
 				ExecutablePath = OverlayExecutable;
diff --git a/Engine/Source/Programs/AutomationTool/Gauntlet/Framework/Gauntlet.Utils.cs b/Engine/Source/Programs/AutomationTool/Gauntlet/Framework/Gauntlet.Utils.cs
--- a/Engine/Source/Programs/AutomationTool/Gauntlet/Framework/Gauntlet.Utils.cs
+++ b/Engine/Source/Programs/AutomationTool/Gauntlet/Framework/Gauntlet.Utils.cs
@@ -63,7 +63,7 @@ namespace Gauntlet
 		///		- UnrealSession	| Modifies the code path of LaunchSession to better reserve, setup, and handle errors for devices.
 		///		- DevicePool	| Removes the strict expectation that Horde is used as a device service. Allows users to specify new services, such as a cloud provider.
 		/// </summary>
-		public static bool UseExperimentaFeatures => Params.ParseParam("Experimental") || Params.ParseParam("ExperimentalLaunchFlow");
+		public static bool UseExperimentalFeatures => Params.ParseParam("Experimental") || Params.ParseParam("ExperimentalLaunchFlow");
 
 		public static string TempDir
 		{
diff --git a/Engine/Source/Programs/AutomationTool/Gauntlet/Unreal/Base/Gauntlet.UnrealSession.cs b/Engine/Source/Programs/AutomationTool/Gauntlet/Unreal/Base/Gauntlet.UnrealSession.cs
--- a/Engine/Source/Programs/AutomationTool/Gauntlet/Unreal/Base/Gauntlet.UnrealSession.cs
+++ b/Engine/Source/Programs/AutomationTool/Gauntlet/Unreal/Base/Gauntlet.UnrealSession.cs
@@ -817,7 +817,7 @@ namespace Gauntlet
 		/// <returns></returns>
 		public UnrealSessionInstance LaunchSession()
 		{
-			if(!Globals.UseExperimentaFeatures)
+			if(!Globals.UseExperimentalFeatures)
 			{
 				return Legacy_LaunchSession();
 			}
diff --git a/Engine/Source/Programs/AutomationTool/Gauntlet/Unreal/Base/Gauntlet.UnrealTestNode.cs b/Engine/Source/Programs/AutomationTool/Gauntlet/Unreal/Base/Gauntlet.UnrealTestNode.cs
--- a/Engine/Source/Programs/AutomationTool/Gauntlet/Unreal/Base/Gauntlet.UnrealTestNode.cs
+++ b/Engine/Source/Programs/AutomationTool/Gauntlet/Unreal/Base/Gauntlet.UnrealTestNode.cs
@@ -653,7 +653,7 @@ namespace Gauntlet
 				PrepareUnrealApp();
 			}
 
-			if(Globals.UseExperimentaFeatures)
+			if(Globals.UseExperimentalFeatures)
 			{
 				return true;
 			}
diff --git a/Engine/Source/Programs/AutomationTool/LowLevelTests/RunLowLevelTests.cs b/Engine/Source/Programs/AutomationTool/LowLevelTests/RunLowLevelTests.cs
--- a/Engine/Source/Programs/AutomationTool/LowLevelTests/RunLowLevelTests.cs
+++ b/Engine/Source/Programs/AutomationTool/LowLevelTests/RunLowLevelTests.cs
@@ -359,7 +359,7 @@ namespace LowLevelTests
 		/// </summary>
 		public IAppInstance InstallAndRunNativeTestApp()
 		{
-			if(!Globals.UseExperimentaFeatures)
+			if(!Globals.UseExperimentalFeatures)
 			{
 				return Legacy_InstallAndRunNativeTestApp();
 			}
