Added waiter to check if player in matchmaking session in TestResetForNextProject
Added suppress of timeout exception in DoUntil/DoWhile function

#rnx

diff --git a/Engine/Source/Programs/AutomationTool/Gauntlet/Framework/RpcFramework/Gauntlet.RpcLibrary.cs b/Engine/Source/Programs/AutomationTool/Gauntlet/Framework/RpcFramework/Gauntlet.RpcLibrary.cs
--- a/Engine/Source/Programs/AutomationTool/Gauntlet/Framework/RpcFramework/Gauntlet.RpcLibrary.cs
+++ b/Engine/Source/Programs/AutomationTool/Gauntlet/Framework/RpcFramework/Gauntlet.RpcLibrary.cs
@@ -200,7 +200,7 @@ namespace Gauntlet
 		public delegate void OnBeforeTimeoutDelegate();
 
 		public static bool DoWhile(string FuncName, Func<bool> RpcCall, 
-			int DelayInMilliseconds = 5 * 1000, int TimeoutInMilliseconds = 60 * 1000, bool bSuppressExceptions = false, OnBeforeTimeoutDelegate OnBeforeTimeout = null)
+			int DelayInMilliseconds = 5 * 1000, int TimeoutInMilliseconds = 60 * 1000, bool bSuppressExceptions = false, bool bSuppressTimeoutException = false, OnBeforeTimeoutDelegate OnBeforeTimeout = null)
 		{
 			Stopwatch StopTimer = Stopwatch.StartNew();
 			while (true)
@@ -224,11 +224,18 @@ namespace Gauntlet
 				if (StopTimer.ElapsedMilliseconds > TimeoutInMilliseconds)
 				{
 					OnBeforeTimeout?.Invoke();
-					throw new TestException($"{FuncName} timed out after {TimeoutInMilliseconds}ms");
+					string TimeoutErrorMessage = $"{FuncName} timed out after {TimeoutInMilliseconds}ms";
+					if (!bSuppressTimeoutException)
+					{
+						throw new TestException(TimeoutErrorMessage);
+					}
+					Log.Info(TimeoutErrorMessage);
+					break;
 				}
 
 				Thread.Sleep(DelayInMilliseconds);
 			}
+			return false;
 		}
 
 		public static bool DoUntil(string FuncName, Func<bool> RpcCall,
