LaunchMultiServer UAT script option -attachtoservers now also attaches to proxy servers, if any.

diff --git a/Engine/Source/Programs/AutomationTool/Scripts/LaunchMultiServer.Automation.cs b/Engine/Source/Programs/AutomationTool/Scripts/LaunchMultiServer.Automation.cs
--- a/Engine/Source/Programs/AutomationTool/Scripts/LaunchMultiServer.Automation.cs
+++ b/Engine/Source/Programs/AutomationTool/Scripts/LaunchMultiServer.Automation.cs
@@ -400,6 +400,11 @@ namespace AutomationTool
 							Args += "-notimeouts ";
 						}
 
+						if (IsAttachingDebugger(LaunchType))
+						{
+							Args += "-WaitForDebuggerNoBreak ";
+						}
+
 						Args += "-ProxyGameServers=";
 
 						foreach (var ServerInfo in ServerInfos.Values)
@@ -421,6 +426,11 @@ namespace AutomationTool
 
 					// Pause between starting processes to enforce startup determinism.
 					System.Threading.Thread.Sleep(PauseBetweenProcessMS);
+
+					if (IsAttachingDebugger(LaunchType))
+					{
+						ProcessDebugger.DebugProcess(NewProcess.ProcessObject.Id);
+					}
 				}
 				else if (LaunchType == ProcessLaunchType.ProxyClient)
 				{
@@ -474,7 +484,7 @@ namespace AutomationTool
 
 		private bool IsAttachingDebugger(ProcessLaunchType LaunchType)
 		{
-			return (ParseParam("attachtoservers") && LaunchType == ProcessLaunchType.Server) ||
+			return (ParseParam("attachtoservers") && LaunchType == ProcessLaunchType.Server || LaunchType == ProcessLaunchType.ProxyServer) ||
 				(ParseParam("attachtoclients") && (LaunchType == ProcessLaunchType.DirectClient || LaunchType == ProcessLaunchType.ProxyClient));
 		}
 
