Fixed LaunchMultiServer commandlet trying to attach the proxy to debugger every time

diff --git a/Engine/Source/Programs/AutomationTool/Scripts/LaunchMultiServer.Automation.cs b/Engine/Source/Programs/AutomationTool/Scripts/LaunchMultiServer.Automation.cs
--- a/Engine/Source/Programs/AutomationTool/Scripts/LaunchMultiServer.Automation.cs
+++ b/Engine/Source/Programs/AutomationTool/Scripts/LaunchMultiServer.Automation.cs
@@ -484,7 +484,7 @@ namespace AutomationTool
 
 		private bool IsAttachingDebugger(ProcessLaunchType LaunchType)
 		{
-			return (ParseParam("attachtoservers") && LaunchType == ProcessLaunchType.Server || LaunchType == ProcessLaunchType.ProxyServer) ||
+			return (ParseParam("attachtoservers") && (LaunchType == ProcessLaunchType.Server || LaunchType == ProcessLaunchType.ProxyServer)) ||
 				(ParseParam("attachtoclients") && (LaunchType == ProcessLaunchType.DirectClient || LaunchType == ProcessLaunchType.ProxyClient));
 		}
 
