Cloud Artifact uploads now run validate part after uploads to make sure all content is available and downloadable. Temporary while we iron out some issues. Also set the part name to "default" normally, to enable us to upload other parts from build graph if we want to append things to the build.

diff --git a/Engine/Source/Programs/AutomationTool/BuildGraph/Tasks/CreateCloudArtifactTask.cs b/Engine/Source/Programs/AutomationTool/BuildGraph/Tasks/CreateCloudArtifactTask.cs
--- a/Engine/Source/Programs/AutomationTool/BuildGraph/Tasks/CreateCloudArtifactTask.cs
+++ b/Engine/Source/Programs/AutomationTool/BuildGraph/Tasks/CreateCloudArtifactTask.cs
@@ -369,9 +369,15 @@ namespace AutomationTool.Tasks
 			metadata.Add("totalSize", totalSize);
 			string createdAtString = DateTime.UtcNow.ToString("O");
 			metadata.Add("createdAt", createdAtString);
+			
+			CbObjectId buildId = CbObjectId.NewObjectId();
+			CbObjectId partId = CbObjectId.NewObjectId();
+
+			// set part name to default by default - should be exposed as a option to allow for uploading separate parts
+			string partName = "default";
 			if (_parameters.UseZenCli)
 			{
-				BuildUploader.UploadBuild(host, baseDir, ns, bucketId, metadata, accessToken, new BuildUploader.BuildInfo
+				BuildUploader.UploadBuild(host, baseDir, ns, bucketId, buildId, partName, partId, metadata, accessToken, new BuildUploader.BuildInfo
 				{
 					Name = name.ToString(),
 					ProjectName = project,
@@ -379,10 +385,12 @@ namespace AutomationTool.Tasks
 					BaselineBranch = baselineBranch,
 					Platform = platform
 				}, allowMultipart: _parameters.AllowMultipart, assumeHttp2: String.Equals(_parameters.HttpVersion, "http2-only", StringComparison.OrdinalIgnoreCase), forceClean: _parameters.Force);
+
+				BuildUploader.ValidatePart(host, ns, bucketId, buildId, partId, accessToken, assumeHttp2: String.Equals(_parameters.HttpVersion, "http2-only", StringComparison.OrdinalIgnoreCase));
+
 				return;
 			}
 
-			CbObjectId buildId = CbObjectId.NewObjectId();
 			using HttpClient httpClient = BuildHttpClient(host, ns, bucketId, buildId, accessToken);
 			Logger.LogInformation("Creating artifact {ArtifactId} '{ArtifactName}' ({ArtifactType}) from {NumFiles} files ({TotalSize:n1}mb). Namespace: {NamespaceId}, Bucket: {Bucket}", buildId, name, type, files.Count, totalSize / (1024.0 * 1024.0), ns, bucketId);
 
@@ -707,73 +715,6 @@ namespace AutomationTool.Tasks
 			return contentLength;
 		}
 
-		// TODO: Multipart uploads not fully implemented in Jupiter yet
-		/*private static async Task MultipartUploadFileAsync(HttpClient httpClient, CbObjectId partId, IoHash hashOfFile, Stream streamToSend, MediaTypeHeaderValue contentType, int idealChunkSize)
-		{
-			if (!streamToSend.CanSeek)
-				throw new ArgumentException("Stream must be seekable");
-			Logger.LogInformation("Multipart uploading blob {Hash}.", hashOfFile);
-
-			IoHash hashOfCompressedContent = await IoHash.ComputeAsync(streamToSend);
-			// reset the stream so we can send it
-			streamToSend.Position = 0;
-
-			string uploadId;
-			string blobName;
-			List<string> partQueryStrings;
-			// start the multipart upload
-			{
-				using HttpContent requestContent = new ByteArrayContent(Array.Empty<byte>());
-				requestContent.Headers.ContentType = new MediaTypeHeaderValue(CustomMediaTypeNames.UnrealCompactBinary);
-
-				using HttpResponseMessage result = await httpClient!.PostAsync(new Uri($"parts/{partId}/blobs/{hashOfFile}/startMultipartUpload", UriKind.Relative), requestContent);
-				result.EnsureSuccessStatusCode();
-				MultipartUploadIdResponse response = await result.Content.ReadAsCompactBinaryAsync<MultipartUploadIdResponse>();
-				uploadId = response.UploadId;
-				blobName = response.BlobName;
-				partQueryStrings = response.PartQueryStrings;
-			}
-
-			// upload parts
-			// split the large file into the smallest chunks allowed
-			long bytesLeft = streamToSend.Length;
-
-			byte[] chunk = Array.Empty<byte>();
-			long bytesToRead = Math.Min(idealChunkSize, bytesLeft);
-			while (bytesToRead > 0)
-			{
-				try
-				{
-					chunk = ArrayPool<byte>.Shared.Rent((int)bytesToRead);
-					using HttpContent requestContent = new ByteArrayContent(chunk);
-					requestContent.Headers.ContentType = new MediaTypeHeaderValue(MediaTypeNames.Application.Octet);
-
-					using HttpResponseMessage result = await httpClient!.PutAsync(new Uri($"parts/{partId}/blobs/{hashOfFile}/uploadMultipart?{queryString}", UriKind.Relative), requestContent);
-					
-					result.EnsureSuccessStatusCode();
-				}
-				finally
-				{
-					bytesLeft -= bytesToRead;
-					bytesToRead = Math.Min(idealChunkSize, bytesLeft);
-					ArrayPool<byte>.Shared.Return(chunk);
-				}
-			}
-			
-			// complete the multipart upload
-			{
-				CompleteMultipartUploadRequest request = new CompleteMultipartUploadRequest() {BlobName = blobName, IsCompressed = isCompressed, UploadId = uploadId};
-				byte[] buf = CbSerializer.SerializeToByteArray(request);
-				using HttpContent requestContent = new ByteArrayContent(buf);
-				requestContent.Headers.ContentType = contentType;
-
-				using HttpResponseMessage result = await httpClient!.PostAsync(new Uri($"parts/{partId}/blobs/{hashOfFile}/completeMultipart", UriKind.Relative), requestContent);
-				result.EnsureSuccessStatusCode();
-			}
-
-			Logger.LogInformation("Multipart uploading of blob {Hash} finished.", hashOfFile);
-		}*/
-
 		private const long UseTempFileAtBytes = Int32.MaxValue; // use a backing file when the uncompressed file is too large to fit into memory
 
 		/// <inheritdoc/>
@@ -799,8 +740,8 @@ namespace AutomationTool.Tasks
 			public string Branch { get; set; }
 			public string BaselineBranch { get; set; }
 		}
-		public static void UploadBuild(string host, DirectoryReference baseDir, string namespaceId, string bucketId,
-			Dictionary<string, object> metadata, string accessToken, BuildInfo info, bool allowMultipart = false, bool assumeHttp2 = false,
+		public static void UploadBuild(string host, DirectoryReference baseDir, string namespaceId, string bucketId, CbObjectId buildId, string partName,
+			CbObjectId partId, Dictionary<string, object> metadata, string accessToken, BuildInfo info, bool allowMultipart = false, bool assumeHttp2 = false,
 			bool forceClean = false)
 		{
 			string metadataJsonPath = Path.Combine(Path.GetTempPath(), Path.GetTempFileName());
@@ -841,7 +782,7 @@ namespace AutomationTool.Tasks
 				string metadataString = File.ReadAllText(metadataJsonPath);
 				string http2Options = assumeHttp2 ? "--assume-http2" : "";
 				// pass the access token via environment variable to avoid showing it on the cli
-				string cmdline = $"builds upload --url {host} --namespace {namespaceId} --bucket {bucketId} --metadata-path {metadataJsonPath} --access-token-env UE-CloudDataCacheAccessToken --create-build --plain-progress --allow-multipart={allowMultipart} {http2Options} --clean={forceClean} {baseDir}";
+				string cmdline = $"builds upload --url {host} --namespace {namespaceId} --bucket {bucketId} --build-id {buildId} --build-part-name {partName} --build-part-id {partId} --metadata-path {metadataJsonPath} --access-token-env UE-CloudDataCacheAccessToken --create-build --plain-progress --allow-multipart={allowMultipart} {http2Options} --clean={forceClean} --verbose {baseDir}";
 				Log.Logger.LogInformation("Using metadata '{Metadata}'", metadataString);
 				CommandUtils.RunAndLog(CommandUtils.CmdEnv, zenExe.FullName,cmdline , Options: CommandUtils.ERunOptions.Default, EnvVars: new Dictionary<string, string> {{"UE-CloudDataCacheAccessToken", accessToken}});
 			}
@@ -850,6 +791,15 @@ namespace AutomationTool.Tasks
 				File.Delete(metadataJsonPath);
 			}
 		}
+
+		public static void ValidatePart(string host, string namespaceId, string bucketId, CbObjectId buildId, CbObjectId partId, string accessToken, bool assumeHttp2)
+		{
+			FileInfo zenExe = new FileInfo("Engine/Binaries/Win64/zen.exe");
+			string http2Options = assumeHttp2 ? "--assume-http2" : "";
+			// pass the access token via environment variable to avoid showing it on the cli
+			string cmdline = $"builds validate-part --url {host} --namespace {namespaceId} --bucket {bucketId} --build-id {buildId} --build-part-id {partId} --access-token-env UE-CloudDataCacheAccessToken --plain-progress {http2Options} --verbose";
+			CommandUtils.RunAndLog(CommandUtils.CmdEnv, zenExe.FullName,cmdline , Options: CommandUtils.ERunOptions.Default, EnvVars: new Dictionary<string, string> {{"UE-CloudDataCacheAccessToken", accessToken}});
+		}
 	}
 	namespace CloudDDC
 	{
