[DSTM] Updating Launcher UAT script to support automatic debugger atachmement using Rider as an option (currently that only works with visual studio).

#rnx

diff --git a/Engine/Source/Programs/AutomationTool/Scripts/LaunchMultiServer.Automation.cs b/Engine/Source/Programs/AutomationTool/Scripts/LaunchMultiServer.Automation.cs
--- a/Engine/Source/Programs/AutomationTool/Scripts/LaunchMultiServer.Automation.cs
+++ b/Engine/Source/Programs/AutomationTool/Scripts/LaunchMultiServer.Automation.cs
@@ -9,6 +9,7 @@ using EpicGames.Core;
 using System.Linq;
 using Microsoft.Extensions.Logging;
 using System.Collections.Generic;
+using System.Diagnostics;
 using System.Runtime.InteropServices;
 
 using Microsoft.VisualStudio.OLE.Interop;
@@ -26,7 +27,71 @@ namespace AutomationTool
 		[DllImport("ole32.dll")]
 		public static extern int GetRunningObjectTable(int reserved, out IRunningObjectTable prot);
 
-		public static void DebugProcess(int ProcessID)
+		public static void DebugProcess(int ProcessID, BuildCommand Command, ProjectParams Params)
+		{
+			if (OperatingSystem.IsWindows())
+			{
+				if (CommandUtils.ParseParam(Command.Params, "Rider"))
+				{
+					bool bUsesUProject = CommandUtils.ParseParam(Command.Params,"RiderUProject");
+					DebugProcessWithRider(ProcessID, Params.RawProjectPath, bUsesUProject);
+				}
+				else
+				{
+					DebugProcessWithVisualStudio(ProcessID);
+				}
+			}
+		}
+
+		public static void DebugProcessWithRider(int ProcessId, FileReference ProjectPath, bool bUsesUProject)
+		{
+			// Rider allows to attach to any process using commandline arguments. If you already have a rider instance opened with the provided project path, that
+			// running instance will be used instead of starting a new one.
+
+			string FinalSolutionOrProjectPath = bUsesUProject ? ProjectPath.ToString() : null;
+			if (!bUsesUProject)
+			{
+				DirectoryReference CurrentDirectory = ProjectPath.Directory;
+				do
+				{
+					string PathToEvaluate = CurrentDirectory + "\\UE5.sln";
+					if (File.Exists(PathToEvaluate))
+					{
+						FinalSolutionOrProjectPath = PathToEvaluate;
+						break;
+					}
+
+					CurrentDirectory = CurrentDirectory.ParentDirectory;
+				} 
+				while(CurrentDirectory != null && !CurrentDirectory.IsRootDirectory());
+			}
+
+			if (FinalSolutionOrProjectPath == null)
+			{
+				Log.Logger.LogInformation($"Failed to find a project solution file path. We cannot attach using Rider.");
+				return;
+			}
+
+			string RiderPath = Environment.GetEnvironmentVariable("RIDERINSTALLDIR", EnvironmentVariableTarget.Machine);
+
+			if (RiderPath == null)
+			{
+				Log.Logger.LogError($"Failed to find Rider's binary path. Is Rider executable location added to the RIDERINSTALLDIR system environment variable?.");
+				return;
+			}
+
+			var RiderProcess = new Process();
+			
+			RiderProcess.StartInfo.FileName = RiderPath + "/Rider64.exe";
+			RiderProcess.StartInfo.UseShellExecute = false;
+			RiderProcess.StartInfo.Arguments = $"attach-to-process Native {ProcessId} {FinalSolutionOrProjectPath}";
+			if (!RiderProcess.Start())
+			{
+				Log.Logger.LogError($"Failed to start or connect to Rider. We cannot attach to the selected process.");
+			}
+		}
+		
+		public static void DebugProcessWithVisualStudio(int ProcessID)
 		{
 			if (OperatingSystem.IsWindows())
 			{
@@ -379,7 +444,7 @@ namespace AutomationTool
 
 						if (IsAttachingDebugger(LaunchType))
 						{
-							ProcessDebugger.DebugProcess(NewProcess.ProcessObject.Id);
+							ProcessDebugger.DebugProcess(NewProcess.ProcessObject.Id, this, Params);
 						}
 
 						ServerNum++;
@@ -429,7 +494,7 @@ namespace AutomationTool
 
 					if (IsAttachingDebugger(LaunchType))
 					{
-						ProcessDebugger.DebugProcess(NewProcess.ProcessObject.Id);
+						ProcessDebugger.DebugProcess(NewProcess.ProcessObject.Id, this, Params);
 					}
 				}
 				else if (LaunchType == ProcessLaunchType.ProxyClient)
@@ -468,7 +533,7 @@ namespace AutomationTool
 
 					if (IsAttachingDebugger(LaunchType))
 					{
-						ProcessDebugger.DebugProcess(NewProcess.ProcessObject.Id);
+						ProcessDebugger.DebugProcess(NewProcess.ProcessObject.Id, this, Params);
 					}
 				}
 			}
