Add CSV output to the bloaty itemization job

#tests Preflighted and checked the csv output was valid

diff --git a/Engine/Source/Programs/AutomationTool/Scripts/ReportItemizedExecutableCode.cs b/Engine/Source/Programs/AutomationTool/Scripts/ReportItemizedExecutableCode.cs
--- a/Engine/Source/Programs/AutomationTool/Scripts/ReportItemizedExecutableCode.cs
+++ b/Engine/Source/Programs/AutomationTool/Scripts/ReportItemizedExecutableCode.cs
@@ -7,6 +7,7 @@ using System.IO;
 using System;
 using UnrealBuildBase;
 using System.Linq;
+using System.Text;
 
 namespace ReportItemizedExecutableCode.Automation
 {
@@ -321,9 +322,14 @@ namespace ReportItemizedExecutableCode.Automation
 			System.Console.WriteLine("-------------------------------------------------------------------------------------------------------------------------------------------------------------------");
 			System.Console.WriteLine("Module, Megabytes (float), Bytes (uint)");
 			ModInfos.SortBy(ModInfo => ModInfo.Key);
+
+			// Create a new string builder for CSV output
+			StringBuilder Builder = new StringBuilder("Module,Megabytes,Bytes\n");
+			
 			foreach (KeyValuePair<string, uint> ModInfo in ModInfos)
 			{
 				System.Console.WriteLine("{0}, {1:F2}, {2}", ModInfo.Key, ModInfo.Value / (1024.0 * 1024.0), ModInfo.Value);
+				Builder.AppendLine(string.Format("{0},{1:F2},{2}", ModInfo.Key, ModInfo.Value / (1024.0 * 1024.0), ModInfo.Value));
 			}
 
 			System.Console.WriteLine("\n\n\n\n-------------------------------------------------------------------------------------------------------------------------------------------------------------------");
@@ -337,6 +343,13 @@ namespace ReportItemizedExecutableCode.Automation
 				System.Console.WriteLine("{0}, {1:F2}, {2}", ModInfo.Key, ModInfo.Value / (1024.0 * 1024.0), ModInfo.Value);
 			}
 
+			string CSVName = ExecutableFileToItemize.GetFileNameWithoutExtension() + "_ItemizedModules.csv";
+			CSVName = Path.Combine(CmdEnv.LogFolder, CSVName);
+			
+			System.Console.WriteLine($"\nWriting itemized module information to: {CSVName}");
+
+			File.WriteAllText(CSVName, Builder.ToString());
+
 			System.Console.WriteLine("\n\n");
 		}
 
