- Improved error handling by catching an exception on correct thread (hopefully, very hard to test this)

diff --git a/Engine/Source/Programs/AutomationTool/Apple/ExportAppleBuild.Automation.cs b/Engine/Source/Programs/AutomationTool/Apple/ExportAppleBuild.Automation.cs
--- a/Engine/Source/Programs/AutomationTool/Apple/ExportAppleBuild.Automation.cs
+++ b/Engine/Source/Programs/AutomationTool/Apple/ExportAppleBuild.Automation.cs
@@ -22,6 +22,8 @@ using System.Xml.Linq;
 using System.Net.Http.Headers;
 using System.Text.Json;
 using System.Text.RegularExpressions;
+using System.Threading.Tasks;
+using System.Net.Http.Json;
 
 class ExportAppleBuild : BuildCommand
 {
@@ -205,14 +207,17 @@ class ExportAppleBuild : BuildCommand
 	record LocalizationCreateRequest(LocalizationCreateInfo data);
 
 
-
-
-
 	static CancellationToken CancellationToken = new();
 
 	// shared client object
 	HttpClient HttpClient = null;
 
+	private async Task<TResponse> GetAsync<TResponse>(string RelativePath, CancellationToken CancellationToken = default)
+	{
+		TResponse Response = await HttpClient.GetFromJsonAsync<TResponse>(RelativePath, CancellationToken);
+		return Response ?? throw new InvalidCastException($"Expected non-null response from GET to {RelativePath}");
+	}
+
 	string GetAppId(string BundleId)
 	{
 		AppResponse Response = HttpClient.GetAsync<AppResponse>($"apps?filter[bundleId]={BundleId}&fields[apps]=name", CancellationToken).Result;
@@ -227,12 +232,12 @@ class ExportAppleBuild : BuildCommand
 		bool bHasProcessedAll = false;
 		while (!bHasProcessedAll)
 		{
-			PagedBuildsResponse Response = HttpClient.GetAsync<PagedBuildsResponse>(SubURL, CancellationToken).Result;			
+			PagedBuildsResponse Response = GetAsync<PagedBuildsResponse>(SubURL, CancellationToken).Result;			
 			foreach (BuildInfo Build in Response.data)
 			{
 				// follow link to preReleaseVersion
 				SubURL = Build.relationships.preReleaseVersion.links.related.Replace(HttpClient.BaseAddress.ToString(), "");
-				VersionResponse VersionResponse = HttpClient.GetAsync<VersionResponse>(SubURL, CancellationToken).Result;
+				VersionResponse VersionResponse = GetAsync<VersionResponse>(SubURL, CancellationToken).Result;
 
 				// if it matches, return this Build!
 				if (VersionResponse.data.attributes.version == PreReleaseVersion)
@@ -260,7 +265,7 @@ class ExportAppleBuild : BuildCommand
 	string GetLocalizationId(string BuildId)
 	{
 		string SubURL = $"builds/{BuildId}/betaBuildLocalizations";
-		LocalizationResponse Response = HttpClient.GetAsync<LocalizationResponse>(SubURL, CancellationToken).Result;
+		LocalizationResponse Response = GetAsync<LocalizationResponse>(SubURL, CancellationToken).Result;
 
 		if (Response.data.Count == 0)
 		{
@@ -297,7 +302,7 @@ class ExportAppleBuild : BuildCommand
 
 					// get prerelease version id
 					string SubURL = $"preReleaseVersions?filter[app]={AppId}&filter[version]={AppVersion}&fields[preReleaseVersions]=version";
-					VersionsResponse VersionResponse = HttpClient.GetAsync<VersionsResponse>(SubURL, CancellationToken).Result;
+					VersionsResponse VersionResponse = GetAsync<VersionsResponse>(SubURL, CancellationToken).Result;
 					// if the app version doesn't even exist, then the bundle verrsion can't exist in it!
 					if (VersionResponse.data.Count == 0)
 					{
@@ -312,7 +317,7 @@ class ExportAppleBuild : BuildCommand
 					bool bHasProcessedAll = false;
 					while (!bHasProcessedAll)
 					{
-						PagedBuildsResponse Response = HttpClient.GetAsync<PagedBuildsResponse>(SubURL, CancellationToken).Result;
+						PagedBuildsResponse Response = GetAsync<PagedBuildsResponse>(SubURL, CancellationToken).Result;
 
 						foreach (BuildInfo Build in Response.data)
 						{
