Allow callers to specify a remote zen host/port to import into in ZenImportSnapshotTask.

[RN] none

diff --git a/Engine/Source/Programs/AutomationTool/BuildGraph/Tasks/ZenImportSnapshotTask.cs b/Engine/Source/Programs/AutomationTool/BuildGraph/Tasks/ZenImportSnapshotTask.cs
--- a/Engine/Source/Programs/AutomationTool/BuildGraph/Tasks/ZenImportSnapshotTask.cs
+++ b/Engine/Source/Programs/AutomationTool/BuildGraph/Tasks/ZenImportSnapshotTask.cs
@@ -62,6 +62,18 @@ namespace AutomationTool.Tasks
 		/// </summary>
 		[TaskParameter(Optional = true)]
 		public bool AsyncImport { get; set; } = false;
+
+		/// <summary>
+		/// Optional. Remote zenserver host to import snapshots into.
+		/// </summary>
+		[TaskParameter(Optional = true)]
+		public string RemoteZenHost { get; set; }
+
+		/// <summary>
+		/// Optional. Port on remote host which zenserver is listening on.
+		/// </summary>
+		[TaskParameter(Optional = true)]
+		public int RemoteZenPort { get; set; } = 8558;
 	}
 
 	/// <summary>
@@ -297,7 +309,7 @@ namespace AutomationTool.Tasks
 			return descriptors;
 		}
 
-		private FileReference WriteProjectStoreFile(string projectId, string platform)
+		private FileReference WriteProjectStoreFile(string projectId, string platform, string host, int port)
 		{
 			DirectoryReference platformCookedDirectory;
 
@@ -320,12 +332,20 @@ namespace AutomationTool.Tasks
 			ProjectStoreData data = new ProjectStoreData();
 			data.ZenServer = new ZenServerStoreData();
 
-			data.ZenServer.IsLocalHost = true;
+			data.ZenServer.IsLocalHost = String.IsNullOrEmpty(host);
 			data.ZenServer.HostName = "[::1]";
-			data.ZenServer.HostPort = 8558;
+			data.ZenServer.HostPort = port;
 			data.ZenServer.ProjectId = projectId;
 			data.ZenServer.OplogId = platform;
-			data.ZenServer.RemoteHostNames.AddRange(GetHostAddresses());
+
+			if (!String.IsNullOrEmpty(host))
+			{
+				data.ZenServer.RemoteHostNames.Add(host);
+			}
+			else
+			{
+				data.ZenServer.RemoteHostNames.AddRange(GetHostAddresses());
+			}
 
 			string json = JsonSerializer.Serialize(data, GetDefaultJsonSerializerOptions());
 			FileReference.WriteAllText(projectStoreFile, json);
@@ -364,7 +384,16 @@ namespace AutomationTool.Tasks
 			_parameters.Project = FileReference.FindCorrectCase(_parameters.Project);
 			SnapshotDescriptorCollection descriptors = GetDescriptors(_parameters);
 
-			ZenLaunch(_parameters.Project);
+			string host;
+			if (String.IsNullOrEmpty(_parameters.RemoteZenHost))
+			{
+				ZenLaunch(_parameters.Project);
+				host = "localhost:8558";
+			}
+			else
+			{
+				host = _parameters.RemoteZenHost + ":" + _parameters.RemoteZenPort.ToString();
+			}
 
 			string projectId = GetProjectID();
 
@@ -374,21 +403,24 @@ namespace AutomationTool.Tasks
 
 			StringBuilder projectCreateCommandLine = new StringBuilder();
 			projectCreateCommandLine.AppendFormat("project-create --force-update {0} {1} {2} {3} {4}", projectId, rootDir.FullName, engineDir.FullName, projectDir.FullName, _parameters.Project.FullName);
+			projectCreateCommandLine.AppendFormat(" --hosturl {0}", host);
 
 			FileReference zenExe = ZenExeFileReference();
 			RunAndLogWithoutSpew(zenExe.FullName, projectCreateCommandLine.ToString());
 
 			foreach (SnapshotDescriptor descriptor in descriptors.Snapshots)
 			{
-				FileReference projectStoreFile = WriteProjectStoreFile(projectId, descriptor.TargetPlatform);
+				FileReference projectStoreFile = WriteProjectStoreFile(projectId, descriptor.TargetPlatform, _parameters.RemoteZenHost, _parameters.RemoteZenPort);
 
 				string oplogName = descriptor.TargetPlatform;
 
 				StringBuilder oplogCreateCommandLine = new StringBuilder();
 				oplogCreateCommandLine.AppendFormat("oplog-create --force-update {0} {1}", projectId, oplogName);
+				oplogCreateCommandLine.AppendFormat(" --hosturl {0}", host);
 
 				StringBuilder oplogImportCommandLine = new StringBuilder();
 				oplogImportCommandLine.AppendFormat("oplog-import {0} {1} {2} --ignore-missing-attachments --clean", projectId, oplogName, projectStoreFile.FullName);
+				oplogImportCommandLine.AppendFormat(" --hosturl {0}", host);
 
 				if (_parameters.ForceImport)
 				{
