Zen workspaces:
- UAT automatic workspace share creation during staging pak builds
- client side Zen workspace discovery & creation

diff --git a/Engine/Source/Programs/AutomationTool/AutomationUtils/ZenUtils.cs b/Engine/Source/Programs/AutomationTool/AutomationUtils/ZenUtils.cs
new file mode 100644
--- /dev/null
+++ b/Engine/Source/Programs/AutomationTool/AutomationUtils/ZenUtils.cs
@@ -0,0 +1,187 @@
+// Copyright Epic Games, Inc. All Rights Reserved.
+
+using AutomationTool;
+using EpicGames.Core;
+using EpicGames.Horde.Jobs;
+using EpicGames.ProjectStore;
+using System;
+using System.Collections.Generic;
+using System.IO;
+using System.Linq;
+using System.Net.Http;
+using System.Runtime.InteropServices;
+using System.Text;
+using System.Text.Json;
+using System.Text.Json.Nodes;
+using System.Threading.Tasks;
+using UnrealBuildBase;
+using UnrealBuildTool;
+
+namespace AutomationUtils
+{
+	public class ZenWorkspace
+	{
+		public string id { get; set; }
+		public string root_path { get; set; }
+		public bool allow_share_creation_from_http { get; set; }
+		public List<Share> shares { get; set; }
+	}
+
+	public class Share
+	{
+		public string id { get; set; }
+		public string share_path { get; set; }
+	}
+
+	public static class ZenUtils
+	{
+		private static FileReference GetZenExeLocation(string ExecutableName)
+		{
+			string PathString = String.Format("{0}/Binaries/{1}/{2}{3}", Unreal.EngineDirectory, HostPlatform.Current.HostEditorPlatform.ToString(), ExecutableName,
+				RuntimeInformation.IsOSPlatform(OSPlatform.Windows) ? ".exe" : "");
+
+			if (Path.IsPathRooted(PathString))
+			{
+				return new FileReference(PathString);
+			}
+			else
+			{
+				return new FileReference(Path.Combine(CommandUtils.CmdEnv.LocalRoot, PathString));
+			}
+		}
+
+		public static JsonSerializerOptions GetDefaultJsonSerializerOptions()
+		{
+			JsonSerializerOptions Options = new();
+			Options.AllowTrailingCommas = true;
+			Options.PropertyNameCaseInsensitive = true;
+			Options.ReadCommentHandling = JsonCommentHandling.Skip;
+			return Options;
+		}
+
+		public static void ZenLaunch(FileReference SponsorFile)
+		{
+			FileReference ZenLaunchExe = GetZenExeLocation("ZenLaunch");
+
+			string LaunchArguments = String.Format("{0} -SponsorProcessID={1}", CommandUtils.MakePathSafeToUseWithCommandLine(SponsorFile.FullName), Environment.ProcessId);
+
+			CommandUtils.RunAndLog(CommandUtils.CmdEnv, ZenLaunchExe.FullName, LaunchArguments, Options: CommandUtils.ERunOptions.Default);
+		}
+
+		public static List<ZenWorkspace> GetWorkspaces(string HostName, int HostPort)
+		{
+			HttpResponseMessage HttpGetResult = null;
+			HttpClient Client = new HttpClient();
+			string Uri = string.Format("http://{0}:{1}/ws", HostName, HostPort);
+			using var Request = new HttpRequestMessage(HttpMethod.Get, Uri);
+			Request.Headers.Add("Accept", "application/json");
+
+			try
+			{
+				HttpGetResult = Client.Send(Request);
+			}
+			catch
+			{
+				return new List<ZenWorkspace>();
+			}
+			
+			if (!HttpGetResult.IsSuccessStatusCode)
+			{
+				return new List<ZenWorkspace>();
+			}
+
+			Task<string> ResultContent = HttpGetResult.Content.ReadAsStringAsync();
+			ResultContent.Wait();
+
+			JsonSerializerOptions JsonOptions = new();
+			JsonOptions.AllowTrailingCommas = true;
+			JsonOptions.PropertyNameCaseInsensitive = true;
+			JsonOptions.ReadCommentHandling = JsonCommentHandling.Skip;
+			JsonNode WorkspaceJsonObject = JsonNode.Parse(ResultContent.Result);
+			JsonArray WorkspaceArray = WorkspaceJsonObject["workspaces"]!.AsArray();
+
+			List<ZenWorkspace> Result = JsonSerializer.Deserialize<List<ZenWorkspace>>(WorkspaceArray, JsonOptions);
+
+			return Result;
+		}
+
+		public static ZenWorkspace GetWorkspaceForDir(in List<ZenWorkspace> Workspaces, DirectoryReference BaseDir)
+		{
+			if (Workspaces == null)
+			{
+				return null;
+			}
+
+			foreach(ZenWorkspace Workspace in Workspaces)
+			{
+				if (Workspace.root_path == BaseDir.ToString())
+				{
+					return Workspace;
+				}
+			}
+
+			return null;
+		}
+
+		public static Share GetShareForDir(in ZenWorkspace Workspace, string ShareDir)
+		{
+			if (Workspace.shares == null)
+			{
+				return null;
+			}
+
+			string RelativeShare = ShareDir;
+			if (System.IO.Path.IsPathRooted(ShareDir))
+			{
+				Uri WorkspaceBase = new Uri(Workspace.root_path, UriKind.Absolute);
+				Uri ShareBase = new Uri(ShareDir, UriKind.Absolute);
+				RelativeShare = WorkspaceBase.MakeRelativeUri(ShareBase).ToString();
+			}
+
+			foreach (Share ZenShare in Workspace.shares)
+			{
+				if (ZenShare.share_path == RelativeShare)
+				{
+					return ZenShare;
+				}
+			}
+
+			return null;
+		}
+
+		public static ZenWorkspace CreateWorkspace(ZenServerStoreData ServerContext, DirectoryReference BaseDir, bool Dynamic = false)
+		{
+			FileReference ZenExe = GetZenExeLocation("zen");
+			string Args = String.Format("workspace create {0}", BaseDir.ToString());
+			if (Dynamic)
+			{
+				Args += " --allow-share-create-from-http";
+			}
+
+			CommandUtils.RunAndLog(ZenExe.FullName, Args);
+
+			ZenWorkspace Result = GetWorkspaceForDir(GetWorkspaces(ServerContext.HostName, ServerContext.HostPort), BaseDir);
+			return Result;
+		}
+
+		public static Share CreateShare(ZenServerStoreData ServerContext, ZenWorkspace Workspace, string SharePath)
+		{
+			if (!System.IO.Directory.Exists(System.IO.Path.Combine(Workspace.root_path, SharePath)))
+			{
+				return null;
+			}
+
+			FileReference ZenExe = GetZenExeLocation("zen");
+			string Args = String.Format("workspace-share create {0} \"{1}\"", Workspace.id, SharePath);
+
+			// Create the workspace in Zen server
+			CommandUtils.RunAndLog(ZenExe.FullName, Args);
+
+			// now refresh the workspaces as the state has changed externally
+			var Workspaces = GetWorkspaces(ServerContext.HostName, ServerContext.HostPort);
+			Workspace = GetWorkspaceForDir(Workspaces, new DirectoryReference(Workspace.root_path));
+			Share Result = GetShareForDir(Workspace, SharePath);
+			return Result;
+		}
+	}
+}
diff --git a/Engine/Source/Programs/AutomationTool/Scripts/CopyBuildToStagingDirectory.Automation.cs b/Engine/Source/Programs/AutomationTool/Scripts/CopyBuildToStagingDirectory.Automation.cs
--- a/Engine/Source/Programs/AutomationTool/Scripts/CopyBuildToStagingDirectory.Automation.cs
+++ b/Engine/Source/Programs/AutomationTool/Scripts/CopyBuildToStagingDirectory.Automation.cs
@@ -25,6 +25,8 @@ using Microsoft.Extensions.Logging;
 using static AutomationTool.CommandUtils;
 using System.Text.Json.Serialization;
 using System.Text.Json.Nodes;
+using AutomationUtils;
+using System.Net;
 
 namespace AutomationScripts
 {
@@ -1019,6 +1021,7 @@ namespace AutomationScripts
 				Options.ReadCommentHandling = JsonCommentHandling.Skip;
 				Options.PropertyNameCaseInsensitive = true;
 				Options.Converters.Add(new JsonStringEnumConverter());
+				Options.DefaultIgnoreCondition = JsonIgnoreCondition.WhenWritingNull;
 				ProjectStoreData ParsedProjectStore = JsonSerializer.Deserialize<ProjectStoreData>(ProjectStoreBytes, Options)!;
 
 				if ((ParsedProjectStore != null) && (ParsedProjectStore.ZenServer != null))
@@ -1043,20 +1046,15 @@ namespace AutomationScripts
 			throw new AutomationException(String.Format("A package store manifest or project store is required when staging to IoStore. Expected to find {0} or {1}. Ensure that cooking was successful.", PackageStoreManifestFile.FullName, ProjectStoreFile.FullName));
 		}
 
-
-
-		/// <summary>
-		/// Creates an updated .projectstore file with a new field containing the stage RemapDirectories, allowing StorageServer to apply these when requesting files
-		/// </summary>
-		/// <returns>New .projectstore file on success, source .projectsource file on failure</returns>
-		private static FileReference MakeProjectStoreWithRemappings(DeploymentContext SC, FileReference SourceProjectStorePath)
+		private static FileReference UpdateProjectStore(ProjectParams Params, DeploymentContext SC, FileReference SourceProjectStorePath)
 		{
 			if (!EpicGames.Core.JsonObject.TryRead(SourceProjectStorePath, out EpicGames.Core.JsonObject ProjectStore))
 			{
 				return SourceProjectStorePath;
 			}
 
-			FileReference UpdatedProjectStorePath = FileReference.Combine(SC.ProjectRoot, "Intermediate", "Staging", SC.StageTargetPlatform.PlatformType.ToString(), "remap_" + SourceProjectStorePath.GetFileName());
+			FileReference UpdatedProjectStorePath = FileReference.Combine(SC.ProjectRoot, "Intermediate", "Staging", SC.StageTargetPlatform.PlatformType.ToString(),
+				"updated_" + SourceProjectStorePath.GetFileName());
 			if (FileReference.Exists(UpdatedProjectStorePath) && !InternalUtils.SafeDeleteFile(UpdatedProjectStorePath.FullName))
 			{
 				return SourceProjectStorePath;
@@ -1067,7 +1065,54 @@ namespace AutomationScripts
 				return SourceProjectStorePath;
 			}
 
-			string MakeStorageServerPath( StagedDirectoryReference Source )
+			if (Params.UsePak(SC.StageTargetPlatform) && !Params.Unattended)
+			{
+				AppendWorkspaceToProjectStore(SC, ProjectStore);
+			}
+			else if (SC.RemapDirectories.Any())
+			{
+				MakeProjectStoreWithRemappings(SC, ProjectStore);
+			}
+			File.WriteAllText(UpdatedProjectStorePath.FullName, ProjectStore.ToJsonString());
+
+			return UpdatedProjectStorePath;
+		}
+
+		private static void AppendWorkspaceToProjectStore(DeploymentContext SC, EpicGames.Core.JsonObject ProjectStore)
+		{
+			if (!ProjectStore.TryGetObjectField("zenserver", out EpicGames.Core.JsonObject ZenServerObject))
+			{
+				return;
+			}
+
+			if (!ZenServerObject.TryGetStringField("HostName", out string HostName) || !ZenServerObject.TryGetIntegerField("HostPort", out int HostPort))
+			{
+				return;
+			}
+
+			List<ZenWorkspace> ZenWorkspaces = ZenUtils.GetWorkspaces(HostName, HostPort);
+			DirectoryReference StageBaseDir = DirectoryReference.Combine(SC.ProjectRoot, "Saved", "StagedBuilds");
+			ZenWorkspace BaseWorkspace = ZenUtils.GetWorkspaceForDir(ZenWorkspaces, StageBaseDir);
+			Share WorkspaceShare = ZenUtils.GetShareForDir(BaseWorkspace, SC.PlatformDir);
+
+			if (WorkspaceShare != null)
+			{
+				DirectoryReference ShareDir = DirectoryReference.Combine(StageBaseDir, WorkspaceShare.share_path);
+				ZenServerObject.AddOrSetFieldValue("workspaceid", BaseWorkspace.id);
+				ZenServerObject.AddOrSetFieldValue("shareid", WorkspaceShare.id);
+				ZenServerObject.AddOrSetFieldValue("workspacesharepath", ShareDir.ToString());
+				ZenServerObject.AddOrSetFieldValue("operatingmode", ZenOperatingMode.Filesystem);
+			}
+
+			ProjectStore.AddOrSetFieldValue("zenserver", ZenServerObject);
+		}
+
+		/// <summary>
+		/// Updates project store json data with a new field containing the stage RemapDirectories, allowing StorageServer to apply these when requesting files
+		/// </summary>
+		private static void MakeProjectStoreWithRemappings(DeploymentContext SC, EpicGames.Core.JsonObject ProjectStore)
+		{
+			string MakeStorageServerPath(StagedDirectoryReference Source)
 			{
 				string Result = Source.ToString();
 				string[] PathFragments = Result.Split(new[] { '\\', '/' }, StringSplitOptions.RemoveEmptyEntries);
@@ -1081,7 +1126,7 @@ namespace AutomationScripts
 					{
 						PathFragments[0] = "{project}";
 					}
-					Result = '/' + string.Join('/',PathFragments) + '/';
+					Result = '/' + string.Join('/', PathFragments) + '/';
 				}
 				return Result;
 			}
@@ -1091,16 +1136,12 @@ namespace AutomationScripts
 			{
 				EpicGames.Core.JsonObject JsonRemapDirectory = new();
 				JsonRemapDirectory.AddOrSetFieldValue("from", MakeStorageServerPath(RemapDirectory.Item1));
-				JsonRemapDirectory.AddOrSetFieldValue("to",   MakeStorageServerPath(RemapDirectory.Item2));
+				JsonRemapDirectory.AddOrSetFieldValue("to", MakeStorageServerPath(RemapDirectory.Item2));
 				JsonRemapDirectories.Add(JsonRemapDirectory);
 			}
 			ProjectStore.AddOrSetFieldValue("remapDirectories", JsonRemapDirectories.ToArray());
-
-			File.WriteAllText(UpdatedProjectStorePath.FullName, ProjectStore.ToJsonString());
-			return UpdatedProjectStorePath;
 		}
 
-
 		private static StagedFileType? GetStagingTypeForCookedFile(FileReference File)
 		{
 			// json files have never been staged
@@ -1768,13 +1809,11 @@ namespace AutomationScripts
 
 							if (CookedFile.HasExtension(".projectstore"))
 							{
-								if (!Params.UsePak(SC.StageTargetPlatform))
+								if (!Params.UsePak(SC.StageTargetPlatform) || (Params.UsePak(SC.StageTargetPlatform) && !Params.Unattended))
 								{
-									FileReference SrcCookedFile = CookedFile;
-									if (SC.RemapDirectories.Any())
-									{
-										SrcCookedFile = MakeProjectStoreWithRemappings(SC, CookedFile);
-									}
+									FileReference SrcCookedFile = FileReference.Combine(SC.ProjectRoot, "Intermediate", "Staging", SC.StageTargetPlatform.PlatformType.ToString(), "updated_" + CookedFile.GetFileName());
+
+									SrcCookedFile = UpdateProjectStore(Params, SC, CookedFile);
 
 									SC.StageFile(StagedFileType.NonUFS, SrcCookedFile, new StagedFileReference(CookedFile.MakeRelativeTo(SC.PlatformCookDir)));
 								}
@@ -2998,6 +3037,60 @@ namespace AutomationScripts
 			CreatePaks(Params, SC, PakInputs, PakCryptoSettings, CryptoKeysCacheFilename);
 		}
 
+		private static void PrepareZenWorkspaceEnv(ProjectParams Params, DeploymentContext SC)
+		{
+			if (!Params.UsePak(SC.StageTargetPlatform) || IsBuildMachine)
+			{
+				return;
+			}
+
+			Logger.LogInformation("Preparing zen workspace environment.");
+
+			ZenUtils.ZenLaunch(SC.RawProjectPath);
+
+			FileReference ZenMarkerFile = FileReference.Combine(SC.PlatformCookDir, "ue.projectstore");
+			System.IO.FileInfo ProjectStoreFileInfo = ZenMarkerFile.ToFileInfo();
+			if (!ProjectStoreFileInfo.Exists)
+			{
+				return;
+			}
+
+			byte[] MarkerFileData = FileReference.ReadAllBytes(ZenMarkerFile);
+			JsonSerializerOptions JsonOptions = new();
+			JsonOptions.AllowTrailingCommas = true;
+			JsonOptions.ReadCommentHandling = JsonCommentHandling.Skip;
+			JsonOptions.PropertyNameCaseInsensitive = true;
+			JsonOptions.Converters.Add(new JsonStringEnumConverter());
+			ProjectStoreData ExtractedProjectStore = JsonSerializer.Deserialize<ProjectStoreData>(MarkerFileData, JsonOptions);
+
+			if (ExtractedProjectStore != null && ExtractedProjectStore.ZenServer != null)
+			{
+				ExtractedProjectStore.ZenServer.OperatingMode = ZenOperatingMode.Stream;
+
+				if (ShouldCreatePak(Params))
+				{
+
+					List<ZenWorkspace> ZenWorkspaces = ZenUtils.GetWorkspaces(ExtractedProjectStore.ZenServer.HostName, ExtractedProjectStore.ZenServer.HostPort);
+					DirectoryReference StageBaseDir = DirectoryReference.Combine(SC.ProjectRoot, "Saved", "StagedBuilds");
+					ZenWorkspace BaseWorkspace = ZenUtils.GetWorkspaceForDir(ZenWorkspaces, StageBaseDir);
+					Share WorkspaceShare = null;
+					if (BaseWorkspace != null)
+					{
+						WorkspaceShare = ZenUtils.GetShareForDir(BaseWorkspace, SC.PlatformDir);
+					}
+					else
+					{
+						BaseWorkspace = ZenUtils.CreateWorkspace(ExtractedProjectStore.ZenServer, StageBaseDir);
+					}
+
+					if (WorkspaceShare == null)
+					{
+						WorkspaceShare = ZenUtils.CreateShare(ExtractedProjectStore.ZenServer, BaseWorkspace, SC.PlatformDir);
+					}
+				}
+			}
+		}
+
 		/// <summary>
 		/// Creates a standalone pak file for crash reporter
 		/// </summary>
@@ -6020,6 +6113,7 @@ namespace AutomationScripts
 
 						// write out the commandline file now so it can go into the manifest
 						WriteStageCommandline(Params, SC);
+						PrepareZenWorkspaceEnv(Params, SC);
 						CreateStagingManifest(Params, SC);
 						CleanStagingDirectory(Params, SC);
 					}
diff --git a/Engine/Source/Programs/Shared/EpicGames.ProjectStore/ProjectStoreTypes.cs b/Engine/Source/Programs/Shared/EpicGames.ProjectStore/ProjectStoreTypes.cs
--- a/Engine/Source/Programs/Shared/EpicGames.ProjectStore/ProjectStoreTypes.cs
+++ b/Engine/Source/Programs/Shared/EpicGames.ProjectStore/ProjectStoreTypes.cs
@@ -18,6 +18,17 @@ using UnrealBuildBase;
 
 namespace EpicGames.ProjectStore
 {
+	/// <summary>
+	/// How the client should operate. Whether it should stream from zenserver
+	/// or use the local filesystem
+	/// </summary>
+	[JsonConverter(typeof(JsonStringEnumConverter))]
+	public enum ZenOperatingMode
+	{
+		Stream,
+		Filesystem
+	}
+
 	/// <summary>
 	/// A specification of storage details with zenserver
 	/// </summary>
@@ -34,7 +45,7 @@ namespace EpicGames.ProjectStore
 		/// <summary>
 		/// The list of remote host names for zenserver
 		/// </summary>
-		public List<string> RemoteHostNames { get; } = new List<string>();
+		public List<string> RemoteHostNames { get; set; } = new List<string>();
 		/// <summary>
 		/// The port for zenserver on the host
 		/// </summary>
@@ -47,6 +58,18 @@ namespace EpicGames.ProjectStore
 		/// The oplog identifier for the stored data in zenserver
 		/// </summary>
 		public string? OplogId { get; set; }
+		/// <summary>
+		/// The zen workspace the build's share is available from
+		/// </summary>
+		public string? WorkspaceId { get; set; }
+		/// <summary>
+		/// The zen workspace share pointing to the build
+		/// </summary>
+		public string? ShareId { get; set; }
+		/// <summary>
+		/// Selected mode of operation. If missing, assume ZenOperatingMode.Stream
+		/// </summary>
+		public ZenOperatingMode? OperatingMode { get; set; } = ZenOperatingMode.Stream;
 	}
 	/// <summary>
 	/// A desctiptor of a project store
diff --git a/Engine/Source/Runtime/StorageServerClient/Private/StorageServerConnection.cpp b/Engine/Source/Runtime/StorageServerClient/Private/StorageServerConnection.cpp
--- a/Engine/Source/Runtime/StorageServerClient/Private/StorageServerConnection.cpp
+++ b/Engine/Source/Runtime/StorageServerClient/Private/StorageServerConnection.cpp
@@ -109,6 +109,110 @@ bool FStorageServerConnection::Initialize(TArrayView<const FString> HostAddresse
 	return false;
 }
 
+TIoStatusOr<FStorageServerConnection::Workspaces> FStorageServerConnection::GetWorkspaces()
+{
+	IStorageServerHttpClient::FResult ResultTuple = HttpClient->RequestSync(BaseURI, EStorageServerContentType::CbObject);
+	TIoStatusOr<FIoBuffer> RequestResult = ResultTuple.Get<0>();
+	if (RequestResult.IsOk())
+	{
+		FMemoryReaderView Reader(RequestResult.ValueOrDie().GetView());
+		FCbObject ResponseObj = LoadCompactBinary(Reader).AsObject();
+
+		Workspaces Result;
+		FCbArrayView WorkspacesArray = ResponseObj["workspaces"].AsArrayView();
+		for (FCbFieldView WorkspaceField : WorkspacesArray)
+		{
+			FCbObjectView WorkspaceObject = WorkspaceField.AsObjectView();
+
+			Workspaces::Workspace Workspace = {
+				.Id = *WriteToString<64>(WorkspaceObject["id"].AsObjectId()),
+				.Root = FString(WorkspaceObject["root_path"].AsString()),
+				.AllowShareCreationFromHttp = WorkspaceObject["allow_share_creation_from_http"].AsBool()
+			};
+
+			if (!Workspace.Id.IsEmpty())
+			{
+				FCbArrayView SharesArray = WorkspaceObject["shares"].AsArrayView();
+				for (FCbFieldView ShareField : SharesArray)
+				{
+					FCbObjectView ShareObject = ShareField.AsObjectView();
+
+					Workspaces::Share Share = {
+						.Id = *WriteToString<64>(ShareObject["id"].AsObjectId()),
+						.Path = FString(ShareObject["share_path"].AsString()),
+						.Alias = FString(ShareObject["alias"].AsString())
+					};
+
+					if (!Share.Id.IsEmpty())
+					{
+						Workspace.Shares.Add(Share);
+					}
+				}
+				Result.Workspaces.Add(Workspace);
+			}
+		}
+		return Result;
+	}
+	return RequestResult.Status();
+}
+
+static void PercentEncodeString(FAnsiStringBuilderBase& EncodedBuilder, const FString& String)
+{
+	for (FString::ElementType C : String)
+	{
+		switch (C)
+		{
+		case '!':
+		case '#':
+		case '$':
+		case '&':
+		case '\'':
+		case '(':
+		case ')':
+		case '*':
+		case '+':
+		case ',':
+		case '/':
+		case ':':
+		case ';':
+		case '=':
+		case '?':
+		case '@':
+		case '[':
+		case ']':
+			EncodedBuilder.Append("%02X", C);
+			break;
+		default:
+			EncodedBuilder.Appendf("%c", C);
+		}
+	}
+}
+
+TIoStatusOr<FString> FStorageServerConnection::CreateShare(const FString& WorkspaceId, const FString& SharePath, const FString& Alias)
+{
+	TAnsiStringBuilder<256> ResourceBuilder;
+	ResourceBuilder.Append(BaseURI);
+	ResourceBuilder.Append("/");
+	ResourceBuilder.Append(WorkspaceId);
+	ResourceBuilder.Append("/");
+	ResourceBuilder.Append("000000000000000000000000");
+	ResourceBuilder.Append("?share_path=");
+	PercentEncodeString(ResourceBuilder, SharePath);
+	if (!Alias.IsEmpty())
+	{
+		ResourceBuilder.Append("&alias=");
+		PercentEncodeString(ResourceBuilder, Alias);
+	}
+	IStorageServerHttpClient::FResult ResultTuple = HttpClient->RequestSync(*ResourceBuilder, EStorageServerContentType::Unknown, "PUT");
+	TIoStatusOr<FIoBuffer> Result = ResultTuple.Get<0>();
+	if (Result.IsOk())
+	{
+		FMemoryView TextView(Result.ValueOrDie().GetView());
+		return FString::ConstructFromPtrSize(reinterpret_cast<const ANSICHAR*>(TextView.GetData()), TextView.GetSize());
+	}
+	return Result.Status();
+}
+
 void FStorageServerConnection::PackageStoreRequest(TFunctionRef<void(FPackageStoreEntryResource&&)> Callback)
 {
 	TAnsiStringBuilder<256> ResourceBuilder;
@@ -821,7 +925,7 @@ void FStorageServerConnection::GetDefaultCacheConfiguration(FCacheConfiguration&
 	OutConfiguration.bEnable = bPlatformSupportsCaching && bPlatformEnablesCachingByDefault;
 
 	// always check it first to ensure we disable cache if cmd arg is provided
-	if ((FCString::Strstr(CmdLine, TEXT("-ZenDisableCache")) != nullptr))
+	if ((FCString::Strstr(CmdLine, TEXT("-ZenDisableCache")) != nullptr) || UE::IsUsingZenPakFileStreaming())
 	{
 		OutConfiguration.bEnable = false;
 	}
diff --git a/Engine/Source/Runtime/StorageServerClient/Private/StorageServerConnection.h b/Engine/Source/Runtime/StorageServerClient/Private/StorageServerConnection.h
--- a/Engine/Source/Runtime/StorageServerClient/Private/StorageServerConnection.h
+++ b/Engine/Source/Runtime/StorageServerClient/Private/StorageServerConnection.h
@@ -30,7 +30,7 @@
 #endif
 
 #ifndef HAS_STORAGE_SERVER_RPC_GETCHUNKS_API // zen server 5.5.16 introduced a new API to request chunks
-	#define HAS_STORAGE_SERVER_RPC_GETCHUNKS_API 1
+	#define HAS_STORAGE_SERVER_RPC_GETCHUNKS_API 0
 #endif
 
 #if !UE_BUILD_SHIPPING
@@ -47,6 +47,29 @@ public:
 
 	bool Initialize(TArrayView<const FString> HostAddresses, const int32 Port, const FAnsiStringView& InBaseURI);
 
+	struct Workspaces
+	{
+		struct Share
+		{
+			FString Id;
+			FString Path;
+			FString Alias;
+		};
+
+		struct Workspace
+		{
+			FString Id;
+			FString Root;
+			bool AllowShareCreationFromHttp = false;
+
+			TArray<Share> Shares;
+		};
+		TArray<Workspace> Workspaces;
+	};
+
+	TIoStatusOr<Workspaces> GetWorkspaces();
+	TIoStatusOr<FString> CreateShare(const FString& WorkspaceId, const FString& SharePath, const FString& Alias);
+
 	void PackageStoreRequest(TFunctionRef<void(FPackageStoreEntryResource&&)> Callback);
 	void FileManifestRequest(TFunctionRef<void(FIoChunkId Id, FStringView Path, int64 RawSize)> Callback);
 	void ChunkInfosRequest(TFunctionRef<void(FIoChunkId Id, FIoHash RawHash, int64 RawSize)> Callback);
diff --git a/Engine/Source/Runtime/StorageServerClient/Private/StorageServerPlatformFile.cpp b/Engine/Source/Runtime/StorageServerClient/Private/StorageServerPlatformFile.cpp
--- a/Engine/Source/Runtime/StorageServerClient/Private/StorageServerPlatformFile.cpp
+++ b/Engine/Source/Runtime/StorageServerClient/Private/StorageServerPlatformFile.cpp
@@ -914,9 +914,114 @@ TUniquePtr<FArchive> FStorageServerPlatformFile::TryFindProjectStoreMarkerFile(I
 	return nullptr;
 }
 
+static FAnsiString FindWorkspaceSharePath(const FString& WorkspaceSharePath, const FStorageServerConnection::Workspaces& Workspaces)
+{
+	const bool IsRelative = FPaths::IsRelative(WorkspaceSharePath);
+	for (const auto& Workspace : Workspaces.Workspaces)
+	{
+		if (!Workspace.Root.IsEmpty())
+		{
+			FString WorkspaceRoot = Workspace.Root;
+			FPaths::NormalizeDirectoryName(WorkspaceRoot);
+			for (const auto& Share : Workspace.Shares)
+			{
+				FString SharePath(Share.Path);
+				FPaths::NormalizeDirectoryName(SharePath);
+				if (FPaths::IsSamePath(IsRelative ? SharePath : FPaths::Combine(WorkspaceRoot, SharePath), WorkspaceSharePath))
+				{
+					TAnsiStringBuilder<256> SharePathBuilder;
+					SharePathBuilder.Append("/ws/");
+					SharePathBuilder.Append(Workspace.Id);
+					SharePathBuilder.Append("/");
+					SharePathBuilder.Append(Share.Id);
+					return SharePathBuilder.ToString();
+				}
+			}
+		}
+	}
+	return {};
+}
+
+static FString GetAsSubPath(const FString& WorkspaceRootPath, const FString& WorkspaceSharePath)
+{
+	FString TestRoot = FPaths::GetPath(WorkspaceSharePath);
+	while (TestRoot.Len() >= WorkspaceRootPath.Len())
+	{
+		if (FPaths::IsSamePath(TestRoot, WorkspaceRootPath))
+		{
+			return WorkspaceSharePath.Mid(TestRoot.Len() + 1);
+		}
+		TestRoot = FPaths::GetPath(TestRoot);
+	}
+	return {};
+}
+
+static FAnsiString CreateWorkspaceShare(FStorageServerConnection& QueryConnection, const FString& WorkspaceSharePath, const FStorageServerConnection::Workspaces& Workspaces)
+{
+	const bool IsRelative = FPaths::IsRelative(WorkspaceSharePath);
+	for (const auto& Workspace : Workspaces.Workspaces)
+	{
+		if (Workspace.AllowShareCreationFromHttp)
+		{
+			FString WorkspaceRoot = Workspace.Root;
+			FPaths::NormalizeDirectoryName(WorkspaceRoot);
+			if (!Workspace.Root.IsEmpty())
+			{
+				FString SharePath = IsRelative ? WorkspaceSharePath : GetAsSubPath(WorkspaceRoot, WorkspaceSharePath);
+				FPaths::NormalizeDirectoryName(SharePath);
+				if (!SharePath.IsEmpty())
+				{
+					TIoStatusOr<FString> CreateResult = QueryConnection.CreateShare(Workspace.Id, SharePath, "");
+					if (CreateResult.IsOk())
+					{
+						TAnsiStringBuilder<256> SharePathBuilder;
+						SharePathBuilder.Append("/ws/");
+						SharePathBuilder.Append(Workspace.Id);
+						SharePathBuilder.Append("/");
+						SharePathBuilder.Append(CreateResult.ValueOrDie());
+						return SharePathBuilder.ToString();
+					}
+				}
+			}
+		}
+	}
+	return {};
+}
+
 FAnsiString FStorageServerPlatformFile::MakeBaseURI()
 {
 	TAnsiStringBuilder<256> BaseURIBuilder;
+	if (UE::IsUsingZenPakFileStreaming() && !WorkspaceSharePath.IsEmpty())
+	{
+		FPaths::NormalizeDirectoryName(WorkspaceSharePath);
+		FStorageServerConnection QueryConnection;
+		if (QueryConnection.Initialize(HostAddrs, HostPort, "/ws"))
+		{
+			TIoStatusOr<FStorageServerConnection::Workspaces> WorkspacesResponse = QueryConnection.GetWorkspaces();
+			if (WorkspacesResponse.IsOk())
+			{
+				const FStorageServerConnection::Workspaces& Workspaces = WorkspacesResponse.ValueOrDie();
+
+				if (FAnsiString ExistingShareURI = FindWorkspaceSharePath(WorkspaceSharePath, Workspaces); !ExistingShareURI.IsEmpty())
+				{
+					return ExistingShareURI;
+				}
+				if (FAnsiString NewShareURI = CreateWorkspaceShare(QueryConnection, WorkspaceSharePath, Workspaces); !NewShareURI.IsEmpty())
+				{
+					return NewShareURI;
+				}
+				UE_LOG(LogStorageServerPlatformFile, Error, TEXT("Failed to to resolve or create workspace share path '%s' from %s"), *WorkspaceSharePath, *FString(QueryConnection.GetHostAddr()));
+			}
+			else
+			{
+				UE_LOG(LogStorageServerPlatformFile, Error, TEXT("Failed to get list of workspaces from %s. Status: %s"), *FString(QueryConnection.GetHostAddr()), *WorkspacesResponse.Status().ToString());
+			}
+		}
+		else
+		{
+			UE_LOG(LogStorageServerPlatformFile, Error, TEXT("Failed to connect to %s to get list of workspace shares"), *FString(QueryConnection.GetHostAddr()));
+		}
+	}
 	if (!BaseURI.IsEmpty())
 	{
 		BaseURIBuilder.Append(BaseURI);
@@ -954,6 +1059,8 @@ bool FStorageServerPlatformFile::ShouldBeUsed(IPlatformFile* Inner, const TCHAR*
 	if (FParse::Param(FCommandLine::Get(), TEXT("SkipZenStore")))
 		return false;
 
+	bool bPreferFileSystem = false;
+
 #if WITH_COTF
 	UE::Cook::ICookOnTheFlyModule& CookOnTheFlyModule = FModuleManager::LoadModuleChecked<UE::Cook::ICookOnTheFlyModule>(TEXT("CookOnTheFly"));
 	TSharedPtr<UE::Cook::ICookOnTheFlyServerConnection> DefaultConnection = CookOnTheFlyModule.GetDefaultServerConnection();
@@ -975,6 +1082,13 @@ bool FStorageServerPlatformFile::ShouldBeUsed(IPlatformFile* Inner, const TCHAR*
 			if (ProjectStoreObject->TryGetObjectField(TEXT("zenserver"), ZenServerObjectPtr) && (ZenServerObjectPtr != nullptr))
 			{
 				const TSharedPtr<FJsonObject>& ZenServerObject = *ZenServerObjectPtr;
+
+				FString FilesystemOperatingMode;
+				if (ZenServerObject->TryGetStringField(TEXT("operatingmode"), FilesystemOperatingMode) && (FilesystemOperatingMode == "Filesystem"))
+				{
+					bPreferFileSystem = true;
+				}
+
 #if PLATFORM_DESKTOP || PLATFORM_ANDROID
 				FString HostName;
 				if (ZenServerObject->TryGetStringField(TEXT("hostname"), HostName) && !HostName.IsEmpty())
@@ -998,7 +1112,7 @@ bool FStorageServerPlatformFile::ShouldBeUsed(IPlatformFile* Inner, const TCHAR*
 									// As this is the fastest connection when on USB-C, set this as the first to test
 									// TODO: what about when using USB2?? Should we detect device type?
 									HostAddrs.Insert(MacHostName + TEXT(".local"), 0);
-								
+
 									// Some macs drop the ".local", so try that as well, but as a last resort
 									HostAddrs.Add(MacHostName);
 								}
@@ -1039,7 +1153,13 @@ bool FStorageServerPlatformFile::ShouldBeUsed(IPlatformFile* Inner, const TCHAR*
 	{
 		UE_LOG(LogStorageServerPlatformFile, Display, TEXT("Using connection settings from command line: -ZenStorePort='%d'"), HostPort);
 	}
-	return HostAddrs.Num() > 0;
+
+	if (!bPreferFileSystem || UE::IsUsingZenPakFileStreaming())
+	{
+		return HostAddrs.Num() > 0;
+	}
+
+	return false;
 }
 
 bool FStorageServerPlatformFile::Initialize(IPlatformFile* Inner, const TCHAR* CmdLine)
@@ -1093,6 +1213,10 @@ bool FStorageServerPlatformFile::Initialize(IPlatformFile* Inner, const TCHAR* C
 					{
 						BaseURI.Empty();
 					}
+					if (!ZenServerObject->TryGetStringField(TEXT("workspacesharepath"), WorkspaceSharePath))
+					{
+						WorkspaceSharePath.Empty();
+					}
 					UE_LOG(LogStorageServerPlatformFile, Display, TEXT("Using settings from ue.projectstore: ServerProject='%s' and ServerPlatform='%s'"), *ServerProject, *ServerPlatform);
 				}
 
@@ -1123,6 +1247,10 @@ bool FStorageServerPlatformFile::Initialize(IPlatformFile* Inner, const TCHAR* C
 		{
 			UE_LOG(LogStorageServerPlatformFile, Display, TEXT("Using settings from command line: -ZenStoreBaseURI='%s'"), *BaseURI);
 		}
+		if (FParse::Value(CmdLine, TEXT("-ZenWorkspaceSharePath="), WorkspaceSharePath))
+		{
+			UE_LOG(LogStorageServerPlatformFile, Display, TEXT("Using settings from command line: -ZenWorkspaceSharePath='%s'"), *WorkspaceSharePath);
+		}
 
 		if (UE::IsUsingZenPakFileStreaming())
 		{
diff --git a/Engine/Source/Runtime/StorageServerClient/Private/StorageServerPlatformFile.h b/Engine/Source/Runtime/StorageServerClient/Private/StorageServerPlatformFile.h
--- a/Engine/Source/Runtime/StorageServerClient/Private/StorageServerPlatformFile.h
+++ b/Engine/Source/Runtime/StorageServerClient/Private/StorageServerPlatformFile.h
@@ -161,6 +161,7 @@ private:
 	FString BaseURI;
 	FString AbsProjectDir;
 	FString AbsEngineDir;
+	FString WorkspaceSharePath;
 	TDirectoryTree<FString> RemapDirectoriesTree;
 	mutable TArray<FString> HostAddrs;
 	mutable uint16 HostPort = 8558;
diff --git a/Engine/Source/Runtime/StorageServerClientDebug/Private/StorageServerConnectionDebug.cpp b/Engine/Source/Runtime/StorageServerClientDebug/Private/StorageServerConnectionDebug.cpp
--- a/Engine/Source/Runtime/StorageServerClientDebug/Private/StorageServerConnectionDebug.cpp
+++ b/Engine/Source/Runtime/StorageServerClientDebug/Private/StorageServerConnectionDebug.cpp
@@ -85,7 +85,7 @@ bool FStorageServerConnectionDebug::OnTick(float)
 		if (GZenShowStats && GEngine)
 		{
 			FString ZenConnectionDebugMsg;
-			ZenConnectionDebugMsg = FString::Printf(TEXT("ZenServer streaming from %s [%.2fMbps]"), *HostAddress, Throughput);
+			ZenConnectionDebugMsg = FString::Printf(TEXT("ZenServer %s from %s [%.2fMbps]"), UE::IsUsingZenPakFileStreaming() ? TEXT("pak streaming") : TEXT("streaming"), * HostAddress, Throughput);
 			GEngine->AddOnScreenDebugMessage((uint64)this, 86400.0f, FColor::White, ZenConnectionDebugMsg, false);
 		}
 		
