Generate error fingerprint in test phase + Keep json file output small by ignoring null value on serialization

#rnx

diff --git a/Engine/Source/Programs/AutomationTool/Gauntlet/Framework/Utils/Gauntlet.HordeReport.cs b/Engine/Source/Programs/AutomationTool/Gauntlet/Framework/Utils/Gauntlet.HordeReport.cs
--- a/Engine/Source/Programs/AutomationTool/Gauntlet/Framework/Utils/Gauntlet.HordeReport.cs
+++ b/Engine/Source/Programs/AutomationTool/Gauntlet/Framework/Utils/Gauntlet.HordeReport.cs
@@ -1289,8 +1289,8 @@ namespace Gauntlet
 				public string dateTime { get; set; }
 				public float timeElapseSec { get; set; }
 				public TestPhaseOutcome outcome { get; set; }
-				public bool hasError { get; set; }
-				public bool hasWarning { get; set; }
+				public string errorFingerprint { get; set; }
+				public bool? hasWarning { get; set; }
 				public List<string> deviceKeys { get; set; }
 				public string eventStreamPath { get; set; }
 
@@ -1304,6 +1304,7 @@ namespace Gauntlet
 					name = InName;
 					deviceKeys = new List<string>();
 					outcome = TestPhaseOutcome.Unknown;
+					errorFingerprint = null;
 					eventStream = new TestEventStream(session?.report, this);
 					AttachToSession(InSession);
 				}
@@ -1391,6 +1392,17 @@ namespace Gauntlet
 					}
 				}
 
+				/// <summary>
+				/// Compile the aggregate error fingerprints
+				/// </summary>
+				public void CompileErrorFingerprint()
+				{
+					if (!string.IsNullOrEmpty(errorFingerprint) && errorFingerprint.Length > 8)
+					{
+						errorFingerprint = Hasher.ComputeHash(errorFingerprint, Hasher.DefaultAlgo, 8);
+					}
+				}
+
 				/// <summary>
 				/// Get the event stream associated with this phase
 				/// </summary>
@@ -1554,6 +1566,13 @@ namespace Gauntlet
 			/// </summary>
 			public class TestEventStream
 			{
+				// Redact numbers and timestamp like things
+				static Regex Regex_numbertimestampstring = new Regex(@"\d[\d.,:]*", RegexOptions.Compiled);
+				// Redact path
+				static Regex Regex_pathstring = new Regex(@"(?<![A-Z])(?:[A-Z]:|/)[^ :]+[/\\]SYNC[/\\]", RegexOptions.Compiled);
+				// Redact hex strings
+				static Regex Regex_hexstring = new Regex(@"(?:0X|#)[0-9A-F]+", RegexOptions.Compiled);
+
 				private AutomatedTestSessionData report { get; set; }
 				private TestPhase phase { get; set; }
 				private List<LogEvent> stream { get; set; }
@@ -1602,7 +1621,7 @@ namespace Gauntlet
 					switch (Event.Level)
 					{
 						case Logging.LogLevel.Error:
-							phase.hasError = true;
+							phase.errorFingerprint = (phase.errorFingerprint ?? "") + GenerateHash(Event.Message);
 							break;
 
 						case Logging.LogLevel.Warning:
@@ -1611,6 +1630,20 @@ namespace Gauntlet
 					}
 				}
 
+				private string FilterEvent(string Text)
+				{
+					Text = Text.Trim().ToUpperInvariant();
+					Text = Regex_pathstring.Replace(Text, "");
+					Text = Regex_hexstring.Replace(Text, "H");
+					Text = Regex_numbertimestampstring.Replace(Text, "n");
+
+					return Text;
+				}
+				private string GenerateHash(string InText)
+				{
+					return Hasher.ComputeHash(FilterEvent(InText), Hasher.DefaultAlgo, 8);
+				}
+
 				/// <summary>
 				/// Add event with optional formatting inputs
 				/// </summary>
@@ -2129,6 +2162,8 @@ namespace Gauntlet
 						{
 							Log.Error("Failed to save Event Stream for '{0}'. {1}", Phase.key, Ex);
 						}
+						// finalize error fingerprint
+						Phase.CompileErrorFingerprint(); 
 					}
 					// update phase
 					Phase.eventStreamPath = OutputFilePath;
@@ -2321,14 +2356,12 @@ namespace Gauntlet
 		}
 		private static JsonSerializerOptions GetDefaultJsonOptions()
 		{
-			if (CommandUtils.IsBuildMachine)
-			{
-				return new JsonSerializerOptions { };
-			}
-			// Pretty print when running locally
 			return new JsonSerializerOptions
 			{
-				WriteIndented = true
+				// Pretty print when running locally
+				WriteIndented = !CommandUtils.IsBuildMachine,
+				// Keep file size small by ignoring null value
+				DefaultIgnoreCondition = JsonIgnoreCondition.WhenWritingNull
 			};
 		}
 
