Fix for missing zen workspace during staging

#rnx

diff --git a/Engine/Source/Programs/AutomationTool/Scripts/CopyBuildToStagingDirectory.Automation.cs b/Engine/Source/Programs/AutomationTool/Scripts/CopyBuildToStagingDirectory.Automation.cs
--- a/Engine/Source/Programs/AutomationTool/Scripts/CopyBuildToStagingDirectory.Automation.cs
+++ b/Engine/Source/Programs/AutomationTool/Scripts/CopyBuildToStagingDirectory.Automation.cs
@@ -1093,8 +1093,12 @@ namespace AutomationScripts
 			List<ZenWorkspace> ZenWorkspaces = ZenUtils.GetWorkspaces(HostName, HostPort);
 			DirectoryReference StageBaseDir = DirectoryReference.Combine(SC.ProjectRoot, "Saved", "StagedBuilds");
 			ZenWorkspace BaseWorkspace = ZenUtils.GetWorkspaceForDir(ZenWorkspaces, StageBaseDir);
-			Share WorkspaceShare = ZenUtils.GetShareForDir(BaseWorkspace, SC.PlatformDir);
+			if (BaseWorkspace == null)
+			{
+				return;
+			}
 
+			Share WorkspaceShare = ZenUtils.GetShareForDir(BaseWorkspace, SC.PlatformDir);
 			if (WorkspaceShare != null)
 			{
 				DirectoryReference ShareDir = DirectoryReference.Combine(StageBaseDir, WorkspaceShare.share_path);
