[Backout] - CL40452907
#rnx
Original CL Desc
-----------------------------------------------------------------
Zen workspaces:
- UAT automatic workspace share creation during staging pak builds
- client side Zen workspace discovery & creation

diff --git a/Engine/Source/Programs/AutomationTool/AutomationUtils/ZenUtils.cs b/Engine/Source/Programs/AutomationTool/AutomationUtils/ZenUtils.cs
deleted file mode 100644
diff --git a/Engine/Source/Programs/AutomationTool/Scripts/CopyBuildToStagingDirectory.Automation.cs b/Engine/Source/Programs/AutomationTool/Scripts/CopyBuildToStagingDirectory.Automation.cs
--- a/Engine/Source/Programs/AutomationTool/Scripts/CopyBuildToStagingDirectory.Automation.cs
+++ b/Engine/Source/Programs/AutomationTool/Scripts/CopyBuildToStagingDirectory.Automation.cs
@@ -25,8 +25,6 @@ using Microsoft.Extensions.Logging;
 using static AutomationTool.CommandUtils;
 using System.Text.Json.Serialization;
 using System.Text.Json.Nodes;
-using AutomationUtils;
-using System.Net;
 
 namespace AutomationScripts
 {
@@ -1021,7 +1019,6 @@ namespace AutomationScripts
 				Options.ReadCommentHandling = JsonCommentHandling.Skip;
 				Options.PropertyNameCaseInsensitive = true;
 				Options.Converters.Add(new JsonStringEnumConverter());
-				Options.DefaultIgnoreCondition = JsonIgnoreCondition.WhenWritingNull;
 				ProjectStoreData ParsedProjectStore = JsonSerializer.Deserialize<ProjectStoreData>(ProjectStoreBytes, Options)!;
 
 				if ((ParsedProjectStore != null) && (ParsedProjectStore.ZenServer != null))
@@ -1046,15 +1043,20 @@ namespace AutomationScripts
 			throw new AutomationException(String.Format("A package store manifest or project store is required when staging to IoStore. Expected to find {0} or {1}. Ensure that cooking was successful.", PackageStoreManifestFile.FullName, ProjectStoreFile.FullName));
 		}
 
-		private static FileReference UpdateProjectStore(ProjectParams Params, DeploymentContext SC, FileReference SourceProjectStorePath)
+
+
+		/// <summary>
+		/// Creates an updated .projectstore file with a new field containing the stage RemapDirectories, allowing StorageServer to apply these when requesting files
+		/// </summary>
+		/// <returns>New .projectstore file on success, source .projectsource file on failure</returns>
+		private static FileReference MakeProjectStoreWithRemappings(DeploymentContext SC, FileReference SourceProjectStorePath)
 		{
 			if (!EpicGames.Core.JsonObject.TryRead(SourceProjectStorePath, out EpicGames.Core.JsonObject ProjectStore))
 			{
 				return SourceProjectStorePath;
 			}
 
-			FileReference UpdatedProjectStorePath = FileReference.Combine(SC.ProjectRoot, "Intermediate", "Staging", SC.StageTargetPlatform.PlatformType.ToString(),
-				"updated_" + SourceProjectStorePath.GetFileName());
+			FileReference UpdatedProjectStorePath = FileReference.Combine(SC.ProjectRoot, "Intermediate", "Staging", SC.StageTargetPlatform.PlatformType.ToString(), "remap_" + SourceProjectStorePath.GetFileName());
 			if (FileReference.Exists(UpdatedProjectStorePath) && !InternalUtils.SafeDeleteFile(UpdatedProjectStorePath.FullName))
 			{
 				return SourceProjectStorePath;
@@ -1065,58 +1067,7 @@ namespace AutomationScripts
 				return SourceProjectStorePath;
 			}
 
-			if (Params.UsePak(SC.StageTargetPlatform) && !Params.Unattended)
-			{
-				AppendWorkspaceToProjectStore(SC, ProjectStore);
-			}
-			else if (SC.RemapDirectories.Any())
-			{
-				MakeProjectStoreWithRemappings(SC, ProjectStore);
-			}
-			File.WriteAllText(UpdatedProjectStorePath.FullName, ProjectStore.ToJsonString());
-
-			return UpdatedProjectStorePath;
-		}
-
-		private static void AppendWorkspaceToProjectStore(DeploymentContext SC, EpicGames.Core.JsonObject ProjectStore)
-		{
-			if (!ProjectStore.TryGetObjectField("zenserver", out EpicGames.Core.JsonObject ZenServerObject))
-			{
-				return;
-			}
-
-			if (!ZenServerObject.TryGetStringField("HostName", out string HostName) || !ZenServerObject.TryGetIntegerField("HostPort", out int HostPort))
-			{
-				return;
-			}
-
-			List<ZenWorkspace> ZenWorkspaces = ZenUtils.GetWorkspaces(HostName, HostPort);
-			DirectoryReference StageBaseDir = DirectoryReference.Combine(SC.ProjectRoot, "Saved", "StagedBuilds");
-			ZenWorkspace BaseWorkspace = ZenUtils.GetWorkspaceForDir(ZenWorkspaces, StageBaseDir);
-			if (BaseWorkspace == null)
-			{
-				return;
-			}
-
-			Share WorkspaceShare = ZenUtils.GetShareForDir(BaseWorkspace, SC.PlatformDir);
-			if (WorkspaceShare != null)
-			{
-				DirectoryReference ShareDir = DirectoryReference.Combine(StageBaseDir, WorkspaceShare.share_path);
-				ZenServerObject.AddOrSetFieldValue("workspaceid", BaseWorkspace.id);
-				ZenServerObject.AddOrSetFieldValue("shareid", WorkspaceShare.id);
-				ZenServerObject.AddOrSetFieldValue("workspacesharepath", ShareDir.ToString());
-				ZenServerObject.AddOrSetFieldValue("operatingmode", ZenOperatingMode.Filesystem);
-			}
-
-			ProjectStore.AddOrSetFieldValue("zenserver", ZenServerObject);
-		}
-
-		/// <summary>
-		/// Updates project store json data with a new field containing the stage RemapDirectories, allowing StorageServer to apply these when requesting files
-		/// </summary>
-		private static void MakeProjectStoreWithRemappings(DeploymentContext SC, EpicGames.Core.JsonObject ProjectStore)
-		{
-			string MakeStorageServerPath(StagedDirectoryReference Source)
+			string MakeStorageServerPath( StagedDirectoryReference Source )
 			{
 				string Result = Source.ToString();
 				string[] PathFragments = Result.Split(new[] { '\\', '/' }, StringSplitOptions.RemoveEmptyEntries);
@@ -1130,7 +1081,7 @@ namespace AutomationScripts
 					{
 						PathFragments[0] = "{project}";
 					}
-					Result = '/' + string.Join('/', PathFragments) + '/';
+					Result = '/' + string.Join('/',PathFragments) + '/';
 				}
 				return Result;
 			}
@@ -1140,12 +1091,16 @@ namespace AutomationScripts
 			{
 				EpicGames.Core.JsonObject JsonRemapDirectory = new();
 				JsonRemapDirectory.AddOrSetFieldValue("from", MakeStorageServerPath(RemapDirectory.Item1));
-				JsonRemapDirectory.AddOrSetFieldValue("to", MakeStorageServerPath(RemapDirectory.Item2));
+				JsonRemapDirectory.AddOrSetFieldValue("to",   MakeStorageServerPath(RemapDirectory.Item2));
 				JsonRemapDirectories.Add(JsonRemapDirectory);
 			}
 			ProjectStore.AddOrSetFieldValue("remapDirectories", JsonRemapDirectories.ToArray());
+
+			File.WriteAllText(UpdatedProjectStorePath.FullName, ProjectStore.ToJsonString());
+			return UpdatedProjectStorePath;
 		}
 
+
 		private static StagedFileType? GetStagingTypeForCookedFile(FileReference File)
 		{
 			// json files have never been staged
@@ -1813,11 +1768,13 @@ namespace AutomationScripts
 
 							if (CookedFile.HasExtension(".projectstore"))
 							{
-								if (!Params.UsePak(SC.StageTargetPlatform) || (Params.UsePak(SC.StageTargetPlatform) && !Params.Unattended))
+								if (!Params.UsePak(SC.StageTargetPlatform))
 								{
-									FileReference SrcCookedFile = FileReference.Combine(SC.ProjectRoot, "Intermediate", "Staging", SC.StageTargetPlatform.PlatformType.ToString(), "updated_" + CookedFile.GetFileName());
-
-									SrcCookedFile = UpdateProjectStore(Params, SC, CookedFile);
+									FileReference SrcCookedFile = CookedFile;
+									if (SC.RemapDirectories.Any())
+									{
+										SrcCookedFile = MakeProjectStoreWithRemappings(SC, CookedFile);
+									}
 
 									SC.StageFile(StagedFileType.NonUFS, SrcCookedFile, new StagedFileReference(CookedFile.MakeRelativeTo(SC.PlatformCookDir)));
 								}
@@ -3041,60 +2998,6 @@ namespace AutomationScripts
 			CreatePaks(Params, SC, PakInputs, PakCryptoSettings, CryptoKeysCacheFilename);
 		}
 
-		private static void PrepareZenWorkspaceEnv(ProjectParams Params, DeploymentContext SC)
-		{
-			if (!Params.UsePak(SC.StageTargetPlatform) || IsBuildMachine)
-			{
-				return;
-			}
-
-			Logger.LogInformation("Preparing zen workspace environment.");
-
-			ZenUtils.ZenLaunch(SC.RawProjectPath);
-
-			FileReference ZenMarkerFile = FileReference.Combine(SC.PlatformCookDir, "ue.projectstore");
-			System.IO.FileInfo ProjectStoreFileInfo = ZenMarkerFile.ToFileInfo();
-			if (!ProjectStoreFileInfo.Exists)
-			{
-				return;
-			}
-
-			byte[] MarkerFileData = FileReference.ReadAllBytes(ZenMarkerFile);
-			JsonSerializerOptions JsonOptions = new();
-			JsonOptions.AllowTrailingCommas = true;
-			JsonOptions.ReadCommentHandling = JsonCommentHandling.Skip;
-			JsonOptions.PropertyNameCaseInsensitive = true;
-			JsonOptions.Converters.Add(new JsonStringEnumConverter());
-			ProjectStoreData ExtractedProjectStore = JsonSerializer.Deserialize<ProjectStoreData>(MarkerFileData, JsonOptions);
-
-			if (ExtractedProjectStore != null && ExtractedProjectStore.ZenServer != null)
-			{
-				ExtractedProjectStore.ZenServer.OperatingMode = ZenOperatingMode.Stream;
-
-				if (ShouldCreatePak(Params))
-				{
-
-					List<ZenWorkspace> ZenWorkspaces = ZenUtils.GetWorkspaces(ExtractedProjectStore.ZenServer.HostName, ExtractedProjectStore.ZenServer.HostPort);
-					DirectoryReference StageBaseDir = DirectoryReference.Combine(SC.ProjectRoot, "Saved", "StagedBuilds");
-					ZenWorkspace BaseWorkspace = ZenUtils.GetWorkspaceForDir(ZenWorkspaces, StageBaseDir);
-					Share WorkspaceShare = null;
-					if (BaseWorkspace != null)
-					{
-						WorkspaceShare = ZenUtils.GetShareForDir(BaseWorkspace, SC.PlatformDir);
-					}
-					else
-					{
-						BaseWorkspace = ZenUtils.CreateWorkspace(ExtractedProjectStore.ZenServer, StageBaseDir);
-					}
-
-					if (WorkspaceShare == null)
-					{
-						WorkspaceShare = ZenUtils.CreateShare(ExtractedProjectStore.ZenServer, BaseWorkspace, SC.PlatformDir);
-					}
-				}
-			}
-		}
-
 		/// <summary>
 		/// Creates a standalone pak file for crash reporter
 		/// </summary>
@@ -6117,7 +6020,6 @@ namespace AutomationScripts
 
 						// write out the commandline file now so it can go into the manifest
 						WriteStageCommandline(Params, SC);
-						PrepareZenWorkspaceEnv(Params, SC);
 						CreateStagingManifest(Params, SC);
 						CleanStagingDirectory(Params, SC);
 					}
diff --git a/Engine/Source/Programs/Shared/EpicGames.ProjectStore/ProjectStoreTypes.cs b/Engine/Source/Programs/Shared/EpicGames.ProjectStore/ProjectStoreTypes.cs
--- a/Engine/Source/Programs/Shared/EpicGames.ProjectStore/ProjectStoreTypes.cs
+++ b/Engine/Source/Programs/Shared/EpicGames.ProjectStore/ProjectStoreTypes.cs
@@ -18,17 +18,6 @@ using UnrealBuildBase;
 
 namespace EpicGames.ProjectStore
 {
-	/// <summary>
-	/// How the client should operate. Whether it should stream from zenserver
-	/// or use the local filesystem
-	/// </summary>
-	[JsonConverter(typeof(JsonStringEnumConverter))]
-	public enum ZenOperatingMode
-	{
-		Stream,
-		Filesystem
-	}
-
 	/// <summary>
 	/// A specification of storage details with zenserver
 	/// </summary>
@@ -45,7 +34,7 @@ namespace EpicGames.ProjectStore
 		/// <summary>
 		/// The list of remote host names for zenserver
 		/// </summary>
-		public List<string> RemoteHostNames { get; set; } = new List<string>();
+		public List<string> RemoteHostNames { get; } = new List<string>();
 		/// <summary>
 		/// The port for zenserver on the host
 		/// </summary>
@@ -58,18 +47,6 @@ namespace EpicGames.ProjectStore
 		/// The oplog identifier for the stored data in zenserver
 		/// </summary>
 		public string? OplogId { get; set; }
-		/// <summary>
-		/// The zen workspace the build's share is available from
-		/// </summary>
-		public string? WorkspaceId { get; set; }
-		/// <summary>
-		/// The zen workspace share pointing to the build
-		/// </summary>
-		public string? ShareId { get; set; }
-		/// <summary>
-		/// Selected mode of operation. If missing, assume ZenOperatingMode.Stream
-		/// </summary>
-		public ZenOperatingMode? OperatingMode { get; set; } = ZenOperatingMode.Stream;
 	}
 	/// <summary>
 	/// A desctiptor of a project store
diff --git a/Engine/Source/Runtime/StorageServerClient/Private/StorageServerConnection.cpp b/Engine/Source/Runtime/StorageServerClient/Private/StorageServerConnection.cpp
--- a/Engine/Source/Runtime/StorageServerClient/Private/StorageServerConnection.cpp
+++ b/Engine/Source/Runtime/StorageServerClient/Private/StorageServerConnection.cpp
@@ -109,110 +109,6 @@ bool FStorageServerConnection::Initialize(TArrayView<const FString> HostAddresse
 	return false;
 }
 
-TIoStatusOr<FStorageServerConnection::Workspaces> FStorageServerConnection::GetWorkspaces()
-{
-	IStorageServerHttpClient::FResult ResultTuple = HttpClient->RequestSync(BaseURI, EStorageServerContentType::CbObject);
-	TIoStatusOr<FIoBuffer> RequestResult = ResultTuple.Get<0>();
-	if (RequestResult.IsOk())
-	{
-		FMemoryReaderView Reader(RequestResult.ValueOrDie().GetView());
-		FCbObject ResponseObj = LoadCompactBinary(Reader).AsObject();
-
-		Workspaces Result;
-		FCbArrayView WorkspacesArray = ResponseObj["workspaces"].AsArrayView();
-		for (FCbFieldView WorkspaceField : WorkspacesArray)
-		{
-			FCbObjectView WorkspaceObject = WorkspaceField.AsObjectView();
-
-			Workspaces::Workspace Workspace = {
-				.Id = *WriteToString<64>(WorkspaceObject["id"].AsObjectId()),
-				.Root = FString(WorkspaceObject["root_path"].AsString()),
-				.AllowShareCreationFromHttp = WorkspaceObject["allow_share_creation_from_http"].AsBool()
-			};
-
-			if (!Workspace.Id.IsEmpty())
-			{
-				FCbArrayView SharesArray = WorkspaceObject["shares"].AsArrayView();
-				for (FCbFieldView ShareField : SharesArray)
-				{
-					FCbObjectView ShareObject = ShareField.AsObjectView();
-
-					Workspaces::Share Share = {
-						.Id = *WriteToString<64>(ShareObject["id"].AsObjectId()),
-						.Path = FString(ShareObject["share_path"].AsString()),
-						.Alias = FString(ShareObject["alias"].AsString())
-					};
-
-					if (!Share.Id.IsEmpty())
-					{
-						Workspace.Shares.Add(Share);
-					}
-				}
-				Result.Workspaces.Add(Workspace);
-			}
-		}
-		return Result;
-	}
-	return RequestResult.Status();
-}
-
-static void PercentEncodeString(FAnsiStringBuilderBase& EncodedBuilder, const FString& String)
-{
-	for (FString::ElementType C : String)
-	{
-		switch (C)
-		{
-		case '!':
-		case '#':
-		case '$':
-		case '&':
-		case '\'':
-		case '(':
-		case ')':
-		case '*':
-		case '+':
-		case ',':
-		case '/':
-		case ':':
-		case ';':
-		case '=':
-		case '?':
-		case '@':
-		case '[':
-		case ']':
-			EncodedBuilder.Append("%02X", C);
-			break;
-		default:
-			EncodedBuilder.Appendf("%c", C);
-		}
-	}
-}
-
-TIoStatusOr<FString> FStorageServerConnection::CreateShare(const FString& WorkspaceId, const FString& SharePath, const FString& Alias)
-{
-	TAnsiStringBuilder<256> ResourceBuilder;
-	ResourceBuilder.Append(BaseURI);
-	ResourceBuilder.Append("/");
-	ResourceBuilder.Append(WorkspaceId);
-	ResourceBuilder.Append("/");
-	ResourceBuilder.Append("000000000000000000000000");
-	ResourceBuilder.Append("?share_path=");
-	PercentEncodeString(ResourceBuilder, SharePath);
-	if (!Alias.IsEmpty())
-	{
-		ResourceBuilder.Append("&alias=");
-		PercentEncodeString(ResourceBuilder, Alias);
-	}
-	IStorageServerHttpClient::FResult ResultTuple = HttpClient->RequestSync(*ResourceBuilder, EStorageServerContentType::Unknown, "PUT");
-	TIoStatusOr<FIoBuffer> Result = ResultTuple.Get<0>();
-	if (Result.IsOk())
-	{
-		FMemoryView TextView(Result.ValueOrDie().GetView());
-		return FString::ConstructFromPtrSize(reinterpret_cast<const ANSICHAR*>(TextView.GetData()), TextView.GetSize());
-	}
-	return Result.Status();
-}
-
 void FStorageServerConnection::PackageStoreRequest(TFunctionRef<void(FPackageStoreEntryResource&&)> Callback)
 {
 	TAnsiStringBuilder<256> ResourceBuilder;
@@ -925,7 +821,7 @@ void FStorageServerConnection::GetDefaultCacheConfiguration(FCacheConfiguration&
 	OutConfiguration.bEnable = bPlatformSupportsCaching && bPlatformEnablesCachingByDefault;
 
 	// always check it first to ensure we disable cache if cmd arg is provided
-	if ((FCString::Strstr(CmdLine, TEXT("-ZenDisableCache")) != nullptr) || UE::IsUsingZenPakFileStreaming())
+	if ((FCString::Strstr(CmdLine, TEXT("-ZenDisableCache")) != nullptr))
 	{
 		OutConfiguration.bEnable = false;
 	}
diff --git a/Engine/Source/Runtime/StorageServerClient/Private/StorageServerConnection.h b/Engine/Source/Runtime/StorageServerClient/Private/StorageServerConnection.h
--- a/Engine/Source/Runtime/StorageServerClient/Private/StorageServerConnection.h
+++ b/Engine/Source/Runtime/StorageServerClient/Private/StorageServerConnection.h
@@ -30,7 +30,7 @@
 #endif
 
 #ifndef HAS_STORAGE_SERVER_RPC_GETCHUNKS_API // zen server 5.5.16 introduced a new API to request chunks
-	#define HAS_STORAGE_SERVER_RPC_GETCHUNKS_API 0
+	#define HAS_STORAGE_SERVER_RPC_GETCHUNKS_API 1
 #endif
 
 #if !UE_BUILD_SHIPPING
@@ -47,29 +47,6 @@ public:
 
 	bool Initialize(TArrayView<const FString> HostAddresses, const int32 Port, const FAnsiStringView& InBaseURI);
 
-	struct Workspaces
-	{
-		struct Share
-		{
-			FString Id;
-			FString Path;
-			FString Alias;
-		};
-
-		struct Workspace
-		{
-			FString Id;
-			FString Root;
-			bool AllowShareCreationFromHttp = false;
-
-			TArray<Share> Shares;
-		};
-		TArray<Workspace> Workspaces;
-	};
-
-	TIoStatusOr<Workspaces> GetWorkspaces();
-	TIoStatusOr<FString> CreateShare(const FString& WorkspaceId, const FString& SharePath, const FString& Alias);
-
 	void PackageStoreRequest(TFunctionRef<void(FPackageStoreEntryResource&&)> Callback);
 	void FileManifestRequest(TFunctionRef<void(FIoChunkId Id, FStringView Path, int64 RawSize)> Callback);
 	void ChunkInfosRequest(TFunctionRef<void(FIoChunkId Id, FIoHash RawHash, int64 RawSize)> Callback);
diff --git a/Engine/Source/Runtime/StorageServerClient/Private/StorageServerPlatformFile.cpp b/Engine/Source/Runtime/StorageServerClient/Private/StorageServerPlatformFile.cpp
--- a/Engine/Source/Runtime/StorageServerClient/Private/StorageServerPlatformFile.cpp
+++ b/Engine/Source/Runtime/StorageServerClient/Private/StorageServerPlatformFile.cpp
@@ -914,114 +914,9 @@ TUniquePtr<FArchive> FStorageServerPlatformFile::TryFindProjectStoreMarkerFile(I
 	return nullptr;
 }
 
-static FAnsiString FindWorkspaceSharePath(const FString& WorkspaceSharePath, const FStorageServerConnection::Workspaces& Workspaces)
-{
-	const bool IsRelative = FPaths::IsRelative(WorkspaceSharePath);
-	for (const auto& Workspace : Workspaces.Workspaces)
-	{
-		if (!Workspace.Root.IsEmpty())
-		{
-			FString WorkspaceRoot = Workspace.Root;
-			FPaths::NormalizeDirectoryName(WorkspaceRoot);
-			for (const auto& Share : Workspace.Shares)
-			{
-				FString SharePath(Share.Path);
-				FPaths::NormalizeDirectoryName(SharePath);
-				if (FPaths::IsSamePath(IsRelative ? SharePath : FPaths::Combine(WorkspaceRoot, SharePath), WorkspaceSharePath))
-				{
-					TAnsiStringBuilder<256> SharePathBuilder;
-					SharePathBuilder.Append("/ws/");
-					SharePathBuilder.Append(Workspace.Id);
-					SharePathBuilder.Append("/");
-					SharePathBuilder.Append(Share.Id);
-					return SharePathBuilder.ToString();
-				}
-			}
-		}
-	}
-	return {};
-}
-
-static FString GetAsSubPath(const FString& WorkspaceRootPath, const FString& WorkspaceSharePath)
-{
-	FString TestRoot = FPaths::GetPath(WorkspaceSharePath);
-	while (TestRoot.Len() >= WorkspaceRootPath.Len())
-	{
-		if (FPaths::IsSamePath(TestRoot, WorkspaceRootPath))
-		{
-			return WorkspaceSharePath.Mid(TestRoot.Len() + 1);
-		}
-		TestRoot = FPaths::GetPath(TestRoot);
-	}
-	return {};
-}
-
-static FAnsiString CreateWorkspaceShare(FStorageServerConnection& QueryConnection, const FString& WorkspaceSharePath, const FStorageServerConnection::Workspaces& Workspaces)
-{
-	const bool IsRelative = FPaths::IsRelative(WorkspaceSharePath);
-	for (const auto& Workspace : Workspaces.Workspaces)
-	{
-		if (Workspace.AllowShareCreationFromHttp)
-		{
-			FString WorkspaceRoot = Workspace.Root;
-			FPaths::NormalizeDirectoryName(WorkspaceRoot);
-			if (!Workspace.Root.IsEmpty())
-			{
-				FString SharePath = IsRelative ? WorkspaceSharePath : GetAsSubPath(WorkspaceRoot, WorkspaceSharePath);
-				FPaths::NormalizeDirectoryName(SharePath);
-				if (!SharePath.IsEmpty())
-				{
-					TIoStatusOr<FString> CreateResult = QueryConnection.CreateShare(Workspace.Id, SharePath, "");
-					if (CreateResult.IsOk())
-					{
-						TAnsiStringBuilder<256> SharePathBuilder;
-						SharePathBuilder.Append("/ws/");
-						SharePathBuilder.Append(Workspace.Id);
-						SharePathBuilder.Append("/");
-						SharePathBuilder.Append(CreateResult.ValueOrDie());
-						return SharePathBuilder.ToString();
-					}
-				}
-			}
-		}
-	}
-	return {};
-}
-
 FAnsiString FStorageServerPlatformFile::MakeBaseURI()
 {
 	TAnsiStringBuilder<256> BaseURIBuilder;
-	if (UE::IsUsingZenPakFileStreaming() && !WorkspaceSharePath.IsEmpty())
-	{
-		FPaths::NormalizeDirectoryName(WorkspaceSharePath);
-		FStorageServerConnection QueryConnection;
-		if (QueryConnection.Initialize(HostAddrs, HostPort, "/ws"))
-		{
-			TIoStatusOr<FStorageServerConnection::Workspaces> WorkspacesResponse = QueryConnection.GetWorkspaces();
-			if (WorkspacesResponse.IsOk())
-			{
-				const FStorageServerConnection::Workspaces& Workspaces = WorkspacesResponse.ValueOrDie();
-
-				if (FAnsiString ExistingShareURI = FindWorkspaceSharePath(WorkspaceSharePath, Workspaces); !ExistingShareURI.IsEmpty())
-				{
-					return ExistingShareURI;
-				}
-				if (FAnsiString NewShareURI = CreateWorkspaceShare(QueryConnection, WorkspaceSharePath, Workspaces); !NewShareURI.IsEmpty())
-				{
-					return NewShareURI;
-				}
-				UE_LOG(LogStorageServerPlatformFile, Error, TEXT("Failed to to resolve or create workspace share path '%s' from %s"), *WorkspaceSharePath, *FString(QueryConnection.GetHostAddr()));
-			}
-			else
-			{
-				UE_LOG(LogStorageServerPlatformFile, Error, TEXT("Failed to get list of workspaces from %s. Status: %s"), *FString(QueryConnection.GetHostAddr()), *WorkspacesResponse.Status().ToString());
-			}
-		}
-		else
-		{
-			UE_LOG(LogStorageServerPlatformFile, Error, TEXT("Failed to connect to %s to get list of workspace shares"), *FString(QueryConnection.GetHostAddr()));
-		}
-	}
 	if (!BaseURI.IsEmpty())
 	{
 		BaseURIBuilder.Append(BaseURI);
@@ -1059,8 +954,6 @@ bool FStorageServerPlatformFile::ShouldBeUsed(IPlatformFile* Inner, const TCHAR*
 	if (FParse::Param(FCommandLine::Get(), TEXT("SkipZenStore")))
 		return false;
 
-	bool bPreferFileSystem = false;
-
 #if WITH_COTF
 	UE::Cook::ICookOnTheFlyModule& CookOnTheFlyModule = FModuleManager::LoadModuleChecked<UE::Cook::ICookOnTheFlyModule>(TEXT("CookOnTheFly"));
 	TSharedPtr<UE::Cook::ICookOnTheFlyServerConnection> DefaultConnection = CookOnTheFlyModule.GetDefaultServerConnection();
@@ -1082,13 +975,6 @@ bool FStorageServerPlatformFile::ShouldBeUsed(IPlatformFile* Inner, const TCHAR*
 			if (ProjectStoreObject->TryGetObjectField(TEXT("zenserver"), ZenServerObjectPtr) && (ZenServerObjectPtr != nullptr))
 			{
 				const TSharedPtr<FJsonObject>& ZenServerObject = *ZenServerObjectPtr;
-
-				FString FilesystemOperatingMode;
-				if (ZenServerObject->TryGetStringField(TEXT("operatingmode"), FilesystemOperatingMode) && (FilesystemOperatingMode == "Filesystem"))
-				{
-					bPreferFileSystem = true;
-				}
-
 #if PLATFORM_DESKTOP || PLATFORM_ANDROID
 				FString HostName;
 				if (ZenServerObject->TryGetStringField(TEXT("hostname"), HostName) && !HostName.IsEmpty())
@@ -1112,7 +998,7 @@ bool FStorageServerPlatformFile::ShouldBeUsed(IPlatformFile* Inner, const TCHAR*
 									// As this is the fastest connection when on USB-C, set this as the first to test
 									// TODO: what about when using USB2?? Should we detect device type?
 									HostAddrs.Insert(MacHostName + TEXT(".local"), 0);
-
+								
 									// Some macs drop the ".local", so try that as well, but as a last resort
 									HostAddrs.Add(MacHostName);
 								}
@@ -1153,13 +1039,7 @@ bool FStorageServerPlatformFile::ShouldBeUsed(IPlatformFile* Inner, const TCHAR*
 	{
 		UE_LOG(LogStorageServerPlatformFile, Display, TEXT("Using connection settings from command line: -ZenStorePort='%d'"), HostPort);
 	}
-
-	if (!bPreferFileSystem || UE::IsUsingZenPakFileStreaming())
-	{
-		return HostAddrs.Num() > 0;
-	}
-
-	return false;
+	return HostAddrs.Num() > 0;
 }
 
 bool FStorageServerPlatformFile::Initialize(IPlatformFile* Inner, const TCHAR* CmdLine)
@@ -1213,10 +1093,6 @@ bool FStorageServerPlatformFile::Initialize(IPlatformFile* Inner, const TCHAR* C
 					{
 						BaseURI.Empty();
 					}
-					if (!ZenServerObject->TryGetStringField(TEXT("workspacesharepath"), WorkspaceSharePath))
-					{
-						WorkspaceSharePath.Empty();
-					}
 					UE_LOG(LogStorageServerPlatformFile, Display, TEXT("Using settings from ue.projectstore: ServerProject='%s' and ServerPlatform='%s'"), *ServerProject, *ServerPlatform);
 				}
 
@@ -1247,10 +1123,6 @@ bool FStorageServerPlatformFile::Initialize(IPlatformFile* Inner, const TCHAR* C
 		{
 			UE_LOG(LogStorageServerPlatformFile, Display, TEXT("Using settings from command line: -ZenStoreBaseURI='%s'"), *BaseURI);
 		}
-		if (FParse::Value(CmdLine, TEXT("-ZenWorkspaceSharePath="), WorkspaceSharePath))
-		{
-			UE_LOG(LogStorageServerPlatformFile, Display, TEXT("Using settings from command line: -ZenWorkspaceSharePath='%s'"), *WorkspaceSharePath);
-		}
 
 		if (UE::IsUsingZenPakFileStreaming())
 		{
diff --git a/Engine/Source/Runtime/StorageServerClient/Private/StorageServerPlatformFile.h b/Engine/Source/Runtime/StorageServerClient/Private/StorageServerPlatformFile.h
--- a/Engine/Source/Runtime/StorageServerClient/Private/StorageServerPlatformFile.h
+++ b/Engine/Source/Runtime/StorageServerClient/Private/StorageServerPlatformFile.h
@@ -161,7 +161,6 @@ private:
 	FString BaseURI;
 	FString AbsProjectDir;
 	FString AbsEngineDir;
-	FString WorkspaceSharePath;
 	TDirectoryTree<FString> RemapDirectoriesTree;
 	mutable TArray<FString> HostAddrs;
 	mutable uint16 HostPort = 8558;
diff --git a/Engine/Source/Runtime/StorageServerClientDebug/Private/StorageServerConnectionDebug.cpp b/Engine/Source/Runtime/StorageServerClientDebug/Private/StorageServerConnectionDebug.cpp
--- a/Engine/Source/Runtime/StorageServerClientDebug/Private/StorageServerConnectionDebug.cpp
+++ b/Engine/Source/Runtime/StorageServerClientDebug/Private/StorageServerConnectionDebug.cpp
@@ -85,7 +85,7 @@ bool FStorageServerConnectionDebug::OnTick(float)
 		if (GZenShowStats && GEngine)
 		{
 			FString ZenConnectionDebugMsg;
-			ZenConnectionDebugMsg = FString::Printf(TEXT("ZenServer %s from %s [%.2fMbps]"), UE::IsUsingZenPakFileStreaming() ? TEXT("pak streaming") : TEXT("streaming"), * HostAddress, Throughput);
+			ZenConnectionDebugMsg = FString::Printf(TEXT("ZenServer streaming from %s [%.2fMbps]"), *HostAddress, Throughput);
 			GEngine->AddOnScreenDebugMessage((uint64)this, 86400.0f, FColor::White, ZenConnectionDebugMsg, false);
 		}
 		
