- Hopefully final retry-on-exception fix

diff --git a/Engine/Source/Programs/AutomationTool/Apple/ExportAppleBuild.Automation.cs b/Engine/Source/Programs/AutomationTool/Apple/ExportAppleBuild.Automation.cs
--- a/Engine/Source/Programs/AutomationTool/Apple/ExportAppleBuild.Automation.cs
+++ b/Engine/Source/Programs/AutomationTool/Apple/ExportAppleBuild.Automation.cs
@@ -220,7 +220,7 @@ class ExportAppleBuild : BuildCommand
 
 	string GetAppId(string BundleId)
 	{
-		AppResponse Response = HttpClient.GetAsync<AppResponse>($"apps?filter[bundleId]={BundleId}&fields[apps]=name", CancellationToken).Result;
+		AppResponse Response = GetAsync<AppResponse>($"apps?filter[bundleId]={BundleId}&fields[apps]=name", CancellationToken).Result;
 		return Response.data[0].id;
 	}
 
@@ -296,6 +296,8 @@ class ExportAppleBuild : BuildCommand
 				{
 					Thread.Sleep(30 * 1000);
 				}
+				Attempts++;
+
 				try
 				{
 					string AppId = GetAppId(BundleId);
@@ -413,9 +415,19 @@ class ExportAppleBuild : BuildCommand
 					// if we get here, there is either no newer version, or we have successfully re-versioned the .app
 					return true;
 				}
-				catch (HttpRequestException Ex)
+				catch (AggregateException Ex)
 				{
-					Logger.LogWarning("Got an exception while checking for existing version numbers: '{Msg}'. Will try again...", Ex.Message);
+					Ex.Handle(x => 
+						{
+							if (x is HttpRequestException)
+							{
+								Logger.LogWarning("Got an exception while checking for existing version numbers: '{Msg}'. Will try again...", x.Message);
+								return true;
+							}
+							return false;
+						}
+					);
+
 				}
 			}
 
@@ -472,9 +484,6 @@ class ExportAppleBuild : BuildCommand
 		using (HttpClient = CreateClient(Token))
 		{
 			// wait 30 second between queries for the version
-
-			string AppId = GetAppId(BundleId);
-
 			int Delay = 30;
 			int Tries = 15 * 60 / Delay;
 
@@ -483,11 +492,22 @@ class ExportAppleBuild : BuildCommand
 			{
 				try
 				{
+					string AppId = GetAppId(BundleId);
+
 					BuildId = GetBuildId(AppId, PreReleaseVersion, BundleVersion);
 				}
-				catch (HttpRequestException Ex)
+				catch (AggregateException Ex)
 				{
-					Logger.LogWarning("Got an exception while checking for looking for uploaded version: '{Msg}'. Will try again...", Ex.Message);
+					Ex.Handle(x => 
+						{
+							if (x is HttpRequestException)
+							{
+								Logger.LogWarning("Got an exception while checking for looking for uploaded version: '{Msg}'. Will try again...", Ex.Message);
+								return true;
+							}
+							return false;
+						}
+					);
 				}
 
 				if (BuildId == null)
@@ -519,21 +539,45 @@ class ExportAppleBuild : BuildCommand
 
 			Logger.LogInformation("Found a build with version: {BuildVersion} [id is {Id}]", BuildVersion, BuildId);
 
-			string LocId = GetLocalizationId(BuildId);	
-
-			bool bUpdateSucceeded;
-			if (LocId == null)
-			{
-				bUpdateSucceeded = CreateWhatsNew(BuildId, Text);
-			}
-			else 
-			{
-				bUpdateSucceeded = SetWhatsNew(LocId, Text);
-			}
-			if (bUpdateSucceeded)
+			int Attempts = 0;
+			while (Attempts < 5)
 			{
-				GetLocalizationId(BuildId);
-				return true;
+				if (Attempts > 0)
+				{
+					Thread.Sleep(30 * 1000);
+				}
+				Attempts++;
+
+				try
+				{
+					string LocId = GetLocalizationId(BuildId);	
+
+					bool bUpdateSucceeded;
+					if (LocId == null)
+					{
+						bUpdateSucceeded = CreateWhatsNew(BuildId, Text);
+					}
+					else 
+					{
+						bUpdateSucceeded = SetWhatsNew(LocId, Text);
+					}
+					if (bUpdateSucceeded)
+					{
+						GetLocalizationId(BuildId);
+						return true;
+					}
+					return false;
+				}
+				catch (AggregateException Ex)
+				{
+					Ex.Handle(x => 
+						{
+							Logger.LogWarning("Got an unknown exception while updating the WhatsNew of the build. Will try again...", Ex.Message);
+							return true;
+						}
+					);
+					return false;
+				}
 			}
 			
 			return false;
