Stage .collection files as DebugNonUFS.

diff --git a/Engine/Source/Programs/AutomationTool/CookedEditor/CookedEditor.Automation.cs b/Engine/Source/Programs/AutomationTool/CookedEditor/CookedEditor.Automation.cs
--- a/Engine/Source/Programs/AutomationTool/CookedEditor/CookedEditor.Automation.cs
+++ b/Engine/Source/Programs/AutomationTool/CookedEditor/CookedEditor.Automation.cs
@@ -80,6 +80,8 @@ public class ModifyStageContext
 	// these files will just be staged
 	public List<FileReference> NonUFSFilesToStage = new List<FileReference>();
 
+	public List<FileReference> DebugNonUFSFilesToStage = new List<FileReference>();
+
 	public bool bStageShaderDirs = true;
 	public bool bStagePlatformBuildDirs = true;
 	public bool bStageExtrasDirs = false;
@@ -159,10 +161,11 @@ public class ModifyStageContext
 		UnUFSFiles(SC);
 
 		// anything we want to be NonUFS make sure is not alreay UFS
-		UFSFilesToStage.RemoveAll(x => NonUFSFilesToStage.Contains(x));
+		UFSFilesToStage.RemoveAll(x => NonUFSFilesToStage.Contains(x) || DebugNonUFSFilesToStage.Contains(x));
 
 		Dictionary<StagedFileReference, FileReference> StagedUFSFiles = MimicStageFiles(SC, UFSFilesToStage);
 		Dictionary<StagedFileReference, FileReference> StagedNonUFSFiles = MimicStageFiles(SC, NonUFSFilesToStage);
+		Dictionary<StagedFileReference, FileReference> StagedDebugNonUFSFiles = MimicStageFiles(SC, DebugNonUFSFilesToStage);
 		Dictionary<StagedFileReference, FileReference> StagedUncookFiles = MimicStageFiles(SC, FilesToUncook);
 
 		// filter out already-cooked assets
@@ -212,12 +215,17 @@ public class ModifyStageContext
 		// stage the processed NonUFSFiles
 		SC.StageFiles(StagedFileType.NonUFS, StagedNonUFSFiles.Values);
 
+		// stage the processed DebugNonUFSFiles
+		SC.StageFiles(StagedFileType.DebugNonUFS, StagedDebugNonUFSFiles.Values);
+
 		// now remove or allow restricted files
 		HandleRestrictedFiles(SC, ref SC.FilesToStage.UFSFiles);
 		HandleRestrictedFiles(SC, ref SC.FilesToStage.NonUFSFiles);
+		HandleRestrictedFiles(SC, ref SC.FilesToStage.NonUFSDebugFiles);
 
 		// remove UFS files if they are also in NonUFS - no need to duplicate
-		SC.FilesToStage.UFSFiles = SC.FilesToStage.UFSFiles.Where(x => !SC.FilesToStage.NonUFSFiles.ContainsKey(x.Key)).ToDictionary(x => x.Key, x => x.Value);
+		SC.FilesToStage.UFSFiles = SC.FilesToStage.UFSFiles.Where(x => !SC.FilesToStage.NonUFSFiles.ContainsKey(x.Key) && !SC.FilesToStage.NonUFSDebugFiles.ContainsKey(x.Key)).ToDictionary(x => x.Key, x => x.Value);
+
 	}
 	#region Private implementation
 	
@@ -250,10 +258,9 @@ public class ModifyStageContext
 
 		ShippedFiles.RemoveWhere(x => x.HasExtension(".ttf") && !x.Name.Contains("LastResort"));
 
-		var RemappedNonUFS = NonUFSFilesToStage.Select(x => DeploymentContext.MakeRelativeStagedReference(SC, x));
-
 		UFSFilesToStage.RemoveAll(x => ShippedFiles.Contains(DeploymentContext.MakeRelativeStagedReference(SC, x)));
 		NonUFSFilesToStage.RemoveAll(x => ShippedFiles.Contains(DeploymentContext.MakeRelativeStagedReference(SC, x)));
+		DebugNonUFSFilesToStage.RemoveAll(x => ShippedFiles.Contains(DeploymentContext.MakeRelativeStagedReference(SC, x)));
 	}
 	private Dictionary<StagedFileReference, FileReference> MimicStageFiles(DeploymentContext SC, List<FileReference> SourceFiles)
 	{
@@ -499,7 +506,7 @@ public class MakeCookedEditor : BuildCommand
 			DirectoryReference ProjectCollectionsDirectory = DirectoryReference.Combine(Context.ProjectDirectory, "Content", "Collections");
 			if (DirectoryReference.Exists(ProjectCollectionsDirectory))
 			{
-				Context.NonUFSFilesToStage.AddRange(DirectoryReference.EnumerateFiles(ProjectCollectionsDirectory, "*.collection", SearchOption.TopDirectoryOnly));
+				Context.DebugNonUFSFilesToStage.AddRange(DirectoryReference.EnumerateFiles(ProjectCollectionsDirectory, "*.collection", SearchOption.TopDirectoryOnly));
 			}
 
 			DirectoryReference ProjectDevelopersDirectory = DirectoryReference.Combine(Context.ProjectDirectory, "Content", "Developers");
@@ -510,7 +517,7 @@ public class MakeCookedEditor : BuildCommand
 					DirectoryReference DeveloperCollectionsDirectory = DirectoryReference.Combine(DeveloperDirectory, "Collections");
 					if (DirectoryReference.Exists(DeveloperCollectionsDirectory))
 					{
-						Context.NonUFSFilesToStage.AddRange(DirectoryReference.EnumerateFiles(DeveloperCollectionsDirectory, "*.collection", SearchOption.TopDirectoryOnly));
+						Context.DebugNonUFSFilesToStage.AddRange(DirectoryReference.EnumerateFiles(DeveloperCollectionsDirectory, "*.collection", SearchOption.TopDirectoryOnly));
 					}
 				}
 			}
