[Gauntlet] Replaces DevicePool's strict reliance on horde as the only means of device reservations. This change adds IDeviceReservationService, which provides the ability to define custom services to reserve, construct, and manage devices.

IDeviceReservationService Summary
========================
Services are disabled by default and can be made available by configurable program state. For example, the HordeDeviceReservationService will become available when -DeviceURL= and -DevicePool= commandline arguments are present.
Each service is able to define whether it can support a particular UnrealDeviceTargetConstraint. This means multiple services can be enabled and each support different types of devices.
For example, HordeDeviceReservationService may provide consoles but a custom service may provide linux servers from a cloud provider.
Device error reporting and reservation release implementations are also left
DevicePool has been made to interface with all available services, but continues to manage the lifetimes and availability of all local devices, whether those be host machine instances or connected devkits.

Other Details/Improvements
========================
- The existing implementation for the Gaunlet's Horde implementation has been moved to the HordeReservationService type, which is an IDeviceReservationService
- Deprecates public interfaces to horde specific reservations from DevicePool
- UnrealDeviceReservation has a new option to avoid releasing all devices when reserving new ones. This is useful when encountering problem devices, as we can selectively discard and re-reserve only that device constraint.

Bug Fixes (applies to legacy flow as well)
========================
RunUnreal will now only dispose of the device pool once.

Reserving a device no longer immediatly "claims" it. Instead it's made "available".
Any consumer of the device pool that relies on reservations would have to ignore a devices claimed status in order to use it.
UnrealDeviceReservation was doing this, which has also been fixed.

diff --git a/Engine/Source/Programs/AutomationTool/Gauntlet/Framework/Devices/Gauntlet.DevicePool.cs b/Engine/Source/Programs/AutomationTool/Gauntlet/Framework/Devices/Gauntlet.DevicePool.cs
--- a/Engine/Source/Programs/AutomationTool/Gauntlet/Framework/Devices/Gauntlet.DevicePool.cs
+++ b/Engine/Source/Programs/AutomationTool/Gauntlet/Framework/Devices/Gauntlet.DevicePool.cs
@@ -13,6 +13,7 @@ using System.Text.Json.Serialization;
 using EpicGames.Core;
 using System.Data;
 using Gauntlet.Utils;
+using System.Reflection;
 
 namespace Gauntlet
 {
@@ -227,7 +228,27 @@ namespace Gauntlet
 		/// <summary>
 		/// Access to our singleton
 		/// </summary>
-		public static DevicePool Instance { get; private set; } = new DevicePool();
+		public static DevicePool Instance
+		{
+			get
+			{
+				if (_Instance == null || _Instance.bDisposed)
+				{
+					if (_Instance != null)
+					{
+						Log.Info("DevicePool has been disposed. Reinitilizing with a new instance.");
+					}
+
+					_Instance = new DevicePool();
+				}
+
+				return _Instance;
+			}
+			private set
+			{
+				_Instance = value;
+			}
+		}
 
 		public static bool SkipInstall;
 
@@ -235,12 +256,12 @@ namespace Gauntlet
 
 		public static bool DeviceReservationBlock;
 
-		private static bool? IsInstallStep;
+		public static bool? IsInstallStep;
 
 		/// <summary>
-		/// Device reservation service URL
+		/// Active pool instance
 		/// </summary>
-		public string DeviceURL;
+		private static DevicePool _Instance;
 
 		/// <summary>
 		/// Object used for locking access to internal data
@@ -248,29 +269,32 @@ namespace Gauntlet
 		private object LockObject = new object();
 
 		/// <summary>
-		/// List of all provisioned devices that are available for reservations
+		/// List of all provisioned devices that are can be claimed
 		/// </summary>
 		private List<ITargetDevice> AvailableDevices = new List<ITargetDevice>();
 
 		/// <summary>
-		/// List of all devices that have been reserved
+		/// List of all provisioned devices that have been claimed
 		/// </summary>
-		private List<ITargetDevice> ReservedDevices = new List<ITargetDevice>();
+		private List<ITargetDevice> ClaimedDevices = new List<ITargetDevice>();
 
 		/// <summary>
-		/// List of platforms we've had devices for
+		/// List of all provisioned devices that were not provided by a reservation service
 		/// </summary>
-		private HashSet<UnrealTargetPlatform?> UsedPlatforms = new HashSet<UnrealTargetPlatform?>();
+		private List<ITargetDevice> LocalDevices => AvailableDevices
+			.Union(ClaimedDevices)
+			.Except(ReservationServices.SelectMany(Service => Service.ReservedDevices))
+			.ToList();
 
 		/// <summary>
-		/// Device to reservation lookup
+		/// List of all enabled reservation services
 		/// </summary>
-		private Dictionary<ITargetDevice, DeviceReservationAutoRenew> ServiceReservations = new Dictionary<ITargetDevice, DeviceReservationAutoRenew>();
+		private List<IDeviceReservationService> ReservationServices = new List<IDeviceReservationService>();
 
 		/// <summary>
-		/// Target device info, private for reservation use
+		/// List of platforms we've had devices for
 		/// </summary>
-		private Dictionary<ITargetDevice, DeviceDefinition> ServiceDeviceInfo = new Dictionary<ITargetDevice, DeviceDefinition>();
+		private HashSet<UnrealTargetPlatform?> UsedPlatforms = new HashSet<UnrealTargetPlatform?>();
 
 		/// <summary>
 		/// List of device definitions that can be provisioned on demand
@@ -282,12 +306,6 @@ namespace Gauntlet
 		/// </summary>
 		private List<DeviceDefinition> FailedProvisions = new List<DeviceDefinition>();
 
-		/// <summary>
-		/// Local connection state for any device. We use this to avoid disconnecting devices that the
-		/// user may already be connected to when running tests
-		/// </summary>
-		private Dictionary<ITargetDevice, bool> InitialConnectState = new Dictionary<ITargetDevice, bool>();
-
 		/// <summary>
 		/// Device constraints for performance profiles, etc
 		/// </summary>
@@ -298,23 +316,16 @@ namespace Gauntlet
 		/// </summary>
 		private string LocalTempDir;
 
-		/// <summary>
-		/// The maximum number of problem devices to report to device backend
-		/// This mitigates issues on builders, such as hung local processes, incorrectly reporting problem devices
-		/// </summary>
-		private int MaxDeviceErrorReports = 3;
-
 		/// <summary>
 		/// Protected constructor - code should use DevicePool.Instance
 		/// </summary>
 		protected DevicePool()
 		{
-			// create two local devices by default?
-			AddLocalDevices(2);
-
-			if(Instance == null)
+			lock (LockObject)
 			{
-				Instance = this;
+				// Create two local devices by default for ease of running a client and server
+				AddLocalDevices(2);
+				InitializeReservationServices();
 			}
 
 			SkipInstall = Globals.Params.ParseParams("SkipInstall", "SkipCopy", "SkipDeploy");
@@ -323,6 +334,7 @@ namespace Gauntlet
 		}
 
 		#region IDisposable Support
+		bool bDisposed = false;
 		~DevicePool()
 		{
 			Dispose(false);
@@ -349,54 +361,39 @@ namespace Gauntlet
 		/// <param name="disposing"></param>
 		protected virtual void Dispose(bool disposing)
 		{
+			if (bDisposed)
+			{
+				return;
+			}
+
 			lock (LockObject)
 			{
 				if (disposing)
 				{
-					// warn if more than a third of the devices in our pool had issues.
-					int TotalDevices = AvailableDevices.Count + UnprovisionedDevices.Count;
-
-					if (TotalDevices > 0)
+					// Dispose all local devices
+					foreach (ITargetDevice LocalDevice in LocalDevices)
 					{
-						float Failures = FailedProvisions.Count / TotalDevices;
-
-						if (Failures >= 0.33)
-						{
-							FailedProvisions.ForEach(D => Log.Warning("Device {DeviceName} could not be added", D.Name));
-						}
+						DeviceConfigurationCache.Instance.RevertDeviceConfiguration(LocalDevice);
+						LocalDevice.Dispose();
 					}
 
-					// release any outstanding reservations
-					ReleaseReservations();
-
-					// warn if anything was left registered
-					ReservedDevices.ForEach(D =>
+					// Dispose of any services. These dispose their associated devices
+					foreach (IDeviceReservationService Service in ReservationServices)
 					{
-						Log.Warning("Device {DeviceName} was not unregistered prior to DevicePool.Dispose!. Forcibly disposing.", D.Name);
-						D.Dispose();
-					});
-
-					ReleaseDevices(ReservedDevices);
-
-					// Cleanup our available devices
-					foreach (ITargetDevice Device in AvailableDevices)
-					{
-						if (Device.IsConnected && InitialConnectState[Device] == false)
-						{
-							Device.Disconnect();
-						}
-
-						Device.Dispose();
+						Service.Dispose();
 					}
 
-					// Additional devices cleanup
+					// Cleanup things like platform sdk daemons
 					CleanupDevices(UsedPlatforms);
 
 					AvailableDevices.Clear();
+					ClaimedDevices.Clear();
+					ReservationServices.Clear();
 					UnprovisionedDevices.Clear();
-					ReservedDevices.Clear();
 				}
 			}
+
+			bDisposed = true;
 		}
 		#endregion
 
@@ -421,14 +418,13 @@ namespace Gauntlet
 		{
 			lock (LockObject)
 			{
-				return AvailableDevices.Union(ReservedDevices).Where(D => Validate == null ? Constraint.Check(D) : Validate(D)).Count() +
+				return AvailableDevices.Union(ClaimedDevices).Where(D => Validate == null ? Constraint.Check(D) : Validate(D)).Count() +
 					UnprovisionedDevices.Where(D => Constraint.Check(D)).Count();
 			}
 		}
 
 		public UnrealDeviceTargetConstraint GetConstraint(ITargetDevice Device)
 		{
-
 			if (!Constraints.ContainsKey(Device))
 			{
 				throw new AutomationException("Device pool has no contstaint for {0} (device was likely released)", Device);
@@ -440,7 +436,7 @@ namespace Gauntlet
 		public void SetLocalOptions(string InLocalTemp, bool InUniqueTemps = false, string InDeviceURL = "")
 		{
 			LocalTempDir = InLocalTemp;
-			DeviceURL = InDeviceURL;
+			Legacy_DeviceURL = InDeviceURL;
 		}
 
 		public void AddLocalDevices(int MaxCount)
@@ -596,9 +592,14 @@ namespace Gauntlet
 							string DeviceAddress = DeviceRef;
 							string Model = string.Empty;
 
-							// when using device service, skip adding local non-desktop devices to pool
-							bool IsDesktop = UnrealBuildTool.Utils.GetPlatformsInClass(UnrealPlatformClass.Desktop).Contains(DevicePlatform);
-							if (!IsDesktop && DeviceRef.Equals("default", StringComparison.OrdinalIgnoreCase) && !String.IsNullOrEmpty(DeviceURL))
+							// When using device services, skip adding non-desktop local devices to pool if any of the services can support that platform
+							bool IsDesktop = DevicePlatform.IsInGroup(UnrealPlatformGroup.Desktop);
+							bool ReservationsEnabled = ReservationServices.Count > 0;
+							bool ServicesCanSupportThisPlatform = ReservationsEnabled 
+								&& ReservationServices.Where(Service => Service.CanSupportDeviceConstraint(new UnrealDeviceTargetConstraint(DevicePlatform))).Any();
+							bool DeviceServiceEnabled = ServicesCanSupportThisPlatform;
+
+							if (!IsDesktop && DeviceRef.Equals("default", StringComparison.OrdinalIgnoreCase) && DeviceServiceEnabled)
 							{
 								continue;
 							}
@@ -661,9 +662,8 @@ namespace Gauntlet
 					throw new Exception("Device already registered!");
 				}
 
-				InitialConnectState[Device] = Device.IsConnected;
 				Constraints[Device] = Constraint ?? new UnrealDeviceTargetConstraint(Device.Platform.Value);
-
+				UsedPlatforms.Add(Device.Platform);
 				AvailableDevices.Add(Device);
 
 				if (Log.IsVerbose)
@@ -762,17 +762,17 @@ namespace Gauntlet
 		}
 
 		/// <summary>
-		/// Reserve all devices in the provided list. Once reserved a device will not be seen by any code that
+		/// Claim all devices in the provided list. Once reserved a device will not be seen by any code that
 		/// calls EnumerateDevices
 		/// </summary>
 		/// <param name="DeviceList"></param>
 		/// <returns></returns>
-		public bool ReserveDevices(IEnumerable<ITargetDevice> DeviceList)
+		public bool ClaimDevices(IEnumerable<ITargetDevice> DeviceList)
 		{
 			lock (LockObject)
 			{
 				// can reserve if not reserved...
-				if (ReservedDevices.Intersect(DeviceList).Count() > 0)
+				if (ClaimedDevices.Intersect(DeviceList).Count() > 0)
 				{
 					return false;
 				}
@@ -780,130 +780,65 @@ namespace Gauntlet
 				// remove these devices from the available list
 				AvailableDevices = AvailableDevices.Where(D => DeviceList.Contains(D) == false).ToList();
 
-				ReservedDevices.AddRange(DeviceList);
+				ClaimedDevices.AddRange(DeviceList);
 				DeviceList.All(D => UsedPlatforms.Add(D.Platform));
 			}
 			return true;
 		}
 
-		/// <summary>
-		/// Reserve devices from service
-		/// </summary>
-		public bool ReserveDevicesFromService(string DeviceURL, Dictionary<UnrealDeviceTargetConstraint, int> DeviceTypes)
+		public bool ReserveDevicesFromService(Dictionary<UnrealDeviceTargetConstraint, int> DeviceTypes)
 		{
-			if (string.IsNullOrEmpty(DeviceURL) || !Uri.TryCreate(DeviceURL, UriKind.Absolute, out Uri ReservationServerUri))
-			{
-				return false;
-			}
-
-			Dictionary<UnrealTargetPlatform, string> DeviceMap = UnrealTargetPlatform.GetValidPlatformNames()
-				.ToDictionary(Platform => UnrealTargetPlatform.Parse(Platform), Platform => Platform);
-
-			string PoolID = Globals.DevicePoolId;
-
-			// First pass, validate all of the constraints are valid/possible
-			foreach (KeyValuePair<UnrealDeviceTargetConstraint, int> Entry in DeviceTypes)
+			// Flatten the required devices into an enumerable
+			List<UnrealDeviceTargetConstraint> RequiredDeviceConstraints = DeviceTypes
+				.SelectMany(KVP => Enumerable.Repeat(KVP.Key, KVP.Value))
+				.ToList();
+
+			// Project each constraint into a supported service
+			Dictionary<IDeviceReservationService, List<UnrealDeviceTargetConstraint>> ServiceMapping = new();
+			foreach (UnrealDeviceTargetConstraint Constraint in RequiredDeviceConstraints)
 			{
-
-				if (Entry.Key.Platform == null || !DeviceMap.ContainsKey(Entry.Key.Platform.Value))
-				{
-					// If an unsupported device, we can't reserve it
-					Log.Info(KnownLogEvents.Gauntlet_DeviceEvent, "Unable to reserve service device of type: {Type}", Entry.Key.Platform);
-					return false;
-				}
-
-				if (!string.IsNullOrEmpty(Entry.Key.Model))
+				foreach (IDeviceReservationService Service in ReservationServices)
 				{
-					// if specific device model, we can't currently reserve it from (legacy) service
-					if (DeviceURL.ToLower().Contains("deviceservice.epicgames.net"))
+					if (Service.CanSupportDeviceConstraint(Constraint))
 					{
-						Log.Info(KnownLogEvents.Gauntlet_DeviceEvent, "Unable to reserve service device of model: {Model} on legacy service", Entry.Key.Model);
-						return false;
-					}
-				}
-
-				if(!string.IsNullOrEmpty(Entry.Key.DeviceName) && Entry.Value > 1)
-				{
-					Log.Info(KnownLogEvents.Gauntlet_DeviceEvent, "Attempted to make a reservation for multiple devices using the same device name {DeviceName}. This is not supported.", Entry.Key.DeviceName);
-					return false;
-				}
-			}
-
-			// Second pass, reserve devices
-			List<ITargetDevice> ScopeReservedDevices = new List<ITargetDevice>();
-			foreach (KeyValuePair<UnrealDeviceTargetConstraint, int> Entry in DeviceTypes)
-			{
-				string[] DeviceRequirements = Enumerable.Repeat(Entry.Key.FormatWithIdentifier(), Entry.Value).ToArray();
-
-				Reservation NewReservation = Reservation.Create(ReservationServerUri, DeviceRequirements, TimeSpan.FromMinutes(10), PoolID: PoolID, DeviceName: Entry.Key.DeviceName);
-				DeviceReservationAutoRenew DeviceReservation = new DeviceReservationAutoRenew(DeviceURL, NewReservation);
-
-				if (DeviceReservation == null || DeviceReservation.Devices.Count != Entry.Value)
-				{
-					Log.Warning(KnownLogEvents.Gauntlet_DeviceEvent, "Unable to make device registration with constraint {Constraint}", Entry.Key);
-					return false;
-				}
-
-				if (DeviceReservationBlock = DeviceReservation.InstallRequired != null)
-				{
-					// InstallRequired is true only for the first reservation attempt of the current step. Cache that value to avoid conflicting value on retry attempt.
-					IsInstallStep = IsInstallStep != null? IsInstallStep : DeviceReservation.InstallRequired == true;
-					SkipInstall = DeviceReservation.InstallRequired == false && IsInstallStep == false;
-					FullClean = !SkipInstall;
-				}
-
-				// Add target devices from reservation
-				foreach (Device Device in DeviceReservation.Devices)
-				{
-					DeviceDefinition Def = new DeviceDefinition();
-					Def.Address = Device.IPOrHostName;
-					Def.Name = Device.Name;
-					Def.Platform = DeviceMap.FirstOrDefault(Entry => Entry.Value == Device.Type.Replace("-DevKit", "", StringComparison.OrdinalIgnoreCase)).Key;
-					Def.DeviceData = Device.DeviceData;
-					Def.Model = Device.Model;
-
-					EPerfSpec PerfSpec = EPerfSpec.Unspecified;
-					if (!string.IsNullOrEmpty(Device.PerfSpec) && !Enum.TryParse(Device.PerfSpec, true, out PerfSpec))
-					{
-						throw new AutomationException("Unable to convert perfspec '{0}' into an EPerfSpec", Device.PerfSpec);
-					}
-					Def.PerfSpec = PerfSpec;
-
-					ITargetDevice TargetDevice = CreateAndRegisterDeviceFromDefinition(Def, Entry.Key);
-
-					// If a device from service can't be added, fail reservation and cleanup devices
-					if (TargetDevice == null)
-					{
-						// If some devices from reservation have been created, release them which will also dispose of reservation
-						if (ScopeReservedDevices.Count > 0)
+						if (ServiceMapping.ContainsKey(Service))
 						{
-							ReleaseDevices(ScopeReservedDevices);
+							ServiceMapping[Service].Add(Constraint);
 						}
 						else
 						{
-							// otherwise, no devices have been creation so just cancel this reservation
-							DeviceReservation.Dispose();
+							ServiceMapping.Add(Service, new List<UnrealDeviceTargetConstraint> { Constraint });
 						}
-
-						Log.Info(KnownLogEvents.Gauntlet_DeviceEvent, "Unable to make device registration: device registration failed for {Platform}:{DeviceName}", Def.Platform, Def.Name);
-						return false;
 					}
+				}
+			}
 
-					ScopeReservedDevices.Add(TargetDevice);
-					ReservedDevices.Add(TargetDevice);
-					UsedPlatforms.Add(TargetDevice.Platform);
+			// Determine if any constraints do not have a service that can supply an appropriate device
+			List<UnrealDeviceTargetConstraint> DevicesWithService = ServiceMapping.Values.SelectMany(Constraint => Constraint).ToList();
+			List<UnrealDeviceTargetConstraint> DevicesLackingService = RequiredDeviceConstraints.Except(DevicesWithService).ToList();
 
-					ServiceDeviceInfo[TargetDevice] = Def;
-					ServiceReservations[TargetDevice] = DeviceReservation;
-				}
+			if (DevicesLackingService.Any())
+			{
+				string Message = "No enabled reservation services were capable of supporting the following constraints:\n";
+				Message += string.Join("\n\t", DevicesLackingService);
+				Message += "\nThe following reservation services were available:\n";
+				Message += string.Join("\n\t", ReservationServices.Select(Service => Service.GetType().Name));
+				Message += "\nEnsure your environment is configured to enable any necessary services.";
+				Log.Warning(Message);
+
+				return false;
 			}
 
-			if (ScopeReservedDevices.Count == 0)
+			foreach (var ReservationPair in ServiceMapping)
 			{
-				Log.Info(KnownLogEvents.Gauntlet_DeviceEvent, "Unable to reserve all devices from service.");
+				if (!ReservationPair.Key.ReserveDevicesFromService(ReservationPair.Value))
+				{
+					Log.Info("Failed to reserve all devices from service. See above log for details");
+					return false;
+				}
 			}
 
-			return ScopeReservedDevices.Count > 0;
+			return true;
 		}
 
 		/// <summary>
@@ -911,13 +846,14 @@ namespace Gauntlet
 		/// </summary>
 		public void ReportDeviceError(ITargetDevice Device, string ErrorMessage)
 		{
-			DeviceDefinition Def = null;
-			if (!ServiceDeviceInfo.TryGetValue(Device, out Def))
+			if (TryGetDevicesReservationService(Device, out IDeviceReservationService Service))
 			{
-				return;
+				Service.ReportDeviceError(Device.Name, ErrorMessage);
+			}
+			else
+			{
+				Log.Verbose("{Device} was not reserved by a service! Ignoring error report.");
 			}
-
-			ReportDeviceError(Def.Name, ErrorMessage);
 		}
 
 		/// <summary>
@@ -926,54 +862,36 @@ namespace Gauntlet
 		/// <param name="DeviceList"></param>
 		public void ReleaseDevices(IEnumerable<ITargetDevice> DeviceList)
 		{
-			lock (LockObject)
+			if (DeviceList == null || !DeviceList.Any())
 			{
-				// remove all these devices from our reserved list
-				ReservedDevices = ReservedDevices.Where(D => DeviceList.Contains(D) == false).ToList();
+				return;
+			}
 
-				// disconnect them
+			lock (LockObject)
+			{
 				List<ITargetDevice> KeepList = new List<ITargetDevice>();
 				foreach (ITargetDevice Device in DeviceList)
 				{
-					DeviceConfigurationCache.Instance.RevertDeviceConfiguration(Device);
-
-					if (Device.IsConnected && InitialConnectState[Device] == false)
+					if (LocalDevices.Contains(Device))
 					{
-						Device.Disconnect();
+						DeviceConfigurationCache.Instance.RevertDeviceConfiguration(Device);
+						KeepList.Add(Device);
 					}
-
-					// if device is from reservation service, unregister it
-					DeviceReservationAutoRenew Reservation;
-					if (ServiceReservations.TryGetValue(Device, out Reservation))
+					else if (TryGetDevicesReservationService(Device, out IDeviceReservationService Service))
 					{
-						bool DisposeReservation = ServiceReservations.Count(Entry => Entry.Value == Reservation) == 1;
-
-						// remove and dispose of device
-						// @todo: add support for reservation modification on server (partial device release)
-						ServiceReservations.Remove(Device);
-						ServiceDeviceInfo.Remove(Device);
-						AvailableDevices.Remove(Device);
-
-						InitialConnectState.Remove(Device);
-						Constraints.Remove(Device);
-
-						Device.Dispose();
-
-						// if last device in reservation, dispose the reservation
-						if (DisposeReservation)
-						{
-							Reservation.Dispose();
-						}
+						Service.ReleaseDevices([Device]);
 					}
 					else
 					{
-						KeepList.Add(Device);
+						Log.Info("Attempted to release an unregistered device. Ensure the device is registered with DevicePool.Instance.RegisterDevice first.");
 					}
-
 				}
 
-				// Put all devices back into the available list at the front - if another device is requested we want to
-				// try and use one we just put a build on
+				// Remove any provisioned devices
+				AvailableDevices = AvailableDevices.Where(Device => !DeviceList.Contains(Device)).ToList();
+				ClaimedDevices = ClaimedDevices.Where(Device => !DeviceList.Contains(Device)).ToList();
+
+				// But keep local devices available
 				KeepList.ForEach(D => { AvailableDevices.Remove(D); AvailableDevices.Insert(0, D); });
 			}
 		}
@@ -997,7 +915,6 @@ namespace Gauntlet
 
 				Func<ITargetDevice, bool> Validate = (ITargetDevice Device) =>
 				{
-
 					if (!Constraint.Check(Device))
 					{
 						return false;
@@ -1032,21 +949,15 @@ namespace Gauntlet
 			var TooFewTotalDevices = RequiredDevices.Where(KP => TotalDeviceTypes[KP.Key] < RequiredDevices[KP.Key]).Select(KP => KP.Key);
 			var TooFewCurrentDevices = RequiredDevices.Where(KP => AvailableDeviceTypes[KP.Key] < RequiredDevices[KP.Key]).Select(KP => KP.Key);
 
-
 			var Devices = TooFewTotalDevices.Concat(TooFewCurrentDevices);
 
 			// Request devices from the service if we need them
-			if (UseServiceDevices && !String.IsNullOrEmpty(DeviceURL) && (TooFewTotalDevices.Count() > 0 || TooFewCurrentDevices.Count() > 0))
+			if (UseServiceDevices && (TooFewTotalDevices.Count() > 0 || TooFewCurrentDevices.Count() > 0))
 			{
-
 				Dictionary<UnrealDeviceTargetConstraint, int> DeviceCounts = new Dictionary<UnrealDeviceTargetConstraint, int>();
-
 				Devices.ToList().ForEach(Platform => DeviceCounts[Platform] = RequiredDevices[Platform]);
 
-				Log.Info(KnownLogEvents.Gauntlet_DeviceEvent, "Requesting devices from service at {URL}", DeviceURL);
-
-				// Acquire necessary devices from service
-				if (!ReserveDevicesFromService(DeviceURL, DeviceCounts))
+				if (!ReserveDevicesFromService(DeviceCounts))
 				{
 					return false;
 				}
@@ -1078,13 +989,17 @@ namespace Gauntlet
 		/// </summary>
 		/// <param name="Def"></param>
 		/// <returns></returns>
-		protected ITargetDevice CreateAndRegisterDeviceFromDefinition(DeviceDefinition Def, UnrealDeviceTargetConstraint Constraint)
+		public ITargetDevice CreateAndRegisterDeviceFromDefinition(DeviceDefinition Def, UnrealDeviceTargetConstraint Constraint, 
+			IDeviceReservationService Service = null, IDeviceFactory Factory = null)
 		{
 			ITargetDevice NewDevice = null;
 
-			IDeviceFactory Factory = InterfaceHelpers.FindImplementations<IDeviceFactory>()
+			if (Factory == null)
+			{
+				Factory = InterfaceHelpers.FindImplementations<IDeviceFactory>()
 					.Where(F => F.CanSupportPlatform(Def.Platform))
 					.FirstOrDefault();
+			}
 
 			if (Factory == null)
 			{
@@ -1130,7 +1045,10 @@ namespace Gauntlet
 				{
 					if (!string.IsNullOrEmpty(Message))
 					{
-						ReportDeviceError(Def.Name, Message);
+						if (Service != null)
+						{
+							Service.ReportDeviceError(Def.Name, Message);
+						}
 					}
 					Log.Info("\nSkipping device.");
 					return null;
@@ -1149,7 +1067,10 @@ namespace Gauntlet
 				string WarningMessage = $"Failed to create device {Def.Name}. {Ex.Message}\n{Ex.StackTrace}";
 				if (Ex is DeviceException)
 				{
-					ReportDeviceError(Def.Name, WarningMessage);
+					if (Service != null)
+					{
+						Service.ReportDeviceError(Def.Name, WarningMessage);
+					}
 				}
 				else
 				{
@@ -1160,6 +1081,27 @@ namespace Gauntlet
 			return NewDevice;
 		}
 
+		private void InitializeReservationServices()
+		{
+			// In order to avoid constructing an instance of the service to determine if it's enabled on any number of types we have to do a bit of generics reflection...
+			foreach (Type ServiceType in InterfaceHelpers.FindTypes<IDeviceReservationService>(true, true))
+			{
+				MethodInfo Function = typeof(DevicePool).GetMethod("IsServiceEnabled", BindingFlags.Instance | BindingFlags.NonPublic);
+				MethodInfo Generic = Function.MakeGenericMethod(ServiceType);
+				object ReturnValue = Generic.Invoke(this, null);
+				bool IsEnabled = (bool)ReturnValue;
+				if (IsEnabled)
+				{
+					ReservationServices.Add(Activator.CreateInstance(ServiceType) as IDeviceReservationService);
+				}
+			}
+		}
+
+		private bool IsServiceEnabled<T>() where T : IDeviceReservationService
+		{
+			return T.Enabled;
+		}
+
 		/// <summary>
 		/// Verifies if the provided TargetDevice meets requirements such as firmware, login status, settings, etc.
 		/// </summary>
@@ -1220,42 +1162,21 @@ namespace Gauntlet
 			return bSucceeded;
 		}
 
-		/// <summary>
-		/// Reporting device issue using the device name provided by service
-		/// </summary>
-		private void ReportDeviceError(string ServiceDeviceName, string ErrorMessage)
-		{
-			if (MaxDeviceErrorReports == 0)
-			{
-				Log.Info(KnownLogEvents.Gauntlet_DeviceEvent, "Maximum device errors reported to backend, {DeviceName} : {Message} ignored", ServiceDeviceName, ErrorMessage);
-				return;
-			}
-
-			MaxDeviceErrorReports--;
-
-			Reservation.ReportDeviceError(DeviceURL, ServiceDeviceName, ErrorMessage);
-		}
-
 		/// <summary>
 		/// Explicitly release all device reservations
 		/// </summary>
 		private void ReleaseReservations()
 		{
-			ServiceDeviceInfo.Clear();
-
-			foreach (ITargetDevice Device in ServiceReservations.Keys)
+			foreach (IDeviceReservationService ReservationService in ReservationServices)
 			{
-				AvailableDevices.Remove(Device);
-				ReservedDevices.Remove(Device);
-				Device.Dispose();
-			}
+				foreach (ITargetDevice Device in ReservationService.ReservedDevices)
+				{
+					AvailableDevices.Remove(Device);
+					ClaimedDevices.Remove(Device);
+				}
 
-			foreach (DeviceReservationAutoRenew Reservation in ServiceReservations.Values.Distinct())
-			{
-				Reservation.Dispose();
+				ReservationService.ReleaseDevices(ReservationService.ReservedDevices);
 			}
-
-			ServiceReservations.Clear();
 		}
 
 		private void CleanupDevices(IEnumerable<UnrealTargetPlatform?> Platforms)
@@ -1274,6 +1195,22 @@ namespace Gauntlet
 			}
 		}
 
+		private bool TryGetDevicesReservationService(ITargetDevice Device, out IDeviceReservationService OutService)
+		{
+			OutService = null;
+
+			foreach (IDeviceReservationService Service in ReservationServices)
+			{
+				if (Service.ReservedDevices.Contains(Device))
+				{
+					OutService = Service;
+					return true;
+				}
+			}
+
+			return false;
+		}
+
 		/// <summary>
 		/// Construct a path to hold cache files and make sure it's properly cleaned
 		/// </summary>
@@ -1349,5 +1286,33 @@ namespace Gauntlet
 
 			return ClientCache;
 		}
+
+		#region Legacy Implementation
+	
+		[Obsolete("Will be removed in a future release")]
+		public string DeviceURL
+		{
+			get { return Legacy_DeviceURL; }
+			set { Legacy_DeviceURL = value; }
+		}
+
+		/// <summary>
+		/// Device reservation service URL
+		/// </summary>
+		private string Legacy_DeviceURL;
+
+		[Obsolete("Misnomer, use ClaimDevices instead. Will be removed in a future release")]
+		public void ReserveDevices(IEnumerable<ITargetDevice> Devices)
+		{
+			ClaimDevices(Devices);
+		}
+
+		[Obsolete("Will be removed in a future release. Use the single parameter overload instead.")]
+		public bool ReserveDevicesFromService(string DeviceURL, Dictionary<UnrealDeviceTargetConstraint, int> DeviceTypes)
+		{
+			return ReserveDevicesFromService(DeviceTypes);
+		}
+
+		#endregion
 	}
 }
\ No newline at end of file
diff --git a/Engine/Source/Programs/AutomationTool/Gauntlet/Framework/Devices/Gauntlet.DeviceReservationService.cs b/Engine/Source/Programs/AutomationTool/Gauntlet/Framework/Devices/Gauntlet.DeviceReservationService.cs
new file mode 100644
--- /dev/null
+++ b/Engine/Source/Programs/AutomationTool/Gauntlet/Framework/Devices/Gauntlet.DeviceReservationService.cs
@@ -0,0 +1,325 @@
+// Copyright Epic Games, Inc. All Rights Reserved.
+
+using AutomationTool;
+using AutomationTool.DeviceReservation;
+using EpicGames.Core;
+using System;
+using System.Collections.Generic;
+using System.Data;
+using System.Linq;
+using UnrealBuildTool;
+
+namespace Gauntlet
+{
+	/// <summary>
+	/// An IDeviceReservationService provides a mechanism to obtain devices from an external service
+	/// </summary>
+	public interface IDeviceReservationService : IDisposable
+	{
+		static virtual bool Enabled { get; }
+		List<ITargetDevice> ReservedDevices { get; }
+		bool CanSupportDeviceConstraint(UnrealDeviceTargetConstraint Constraint);
+		bool ReserveDevicesFromService(IEnumerable<UnrealDeviceTargetConstraint> RequestedDevices);
+		void ReleaseDevices(IEnumerable<ITargetDevice> Devices);
+		void ReportDeviceError(string DeviceName, string ErrorMessage);
+	}
+
+	/// <summary>
+	/// Horde Service. Enabled by providing the base url via -DeviceURL= and pool via -DevicePool=
+	/// </summary>
+	public class HordeDeviceReservationService : IDeviceReservationService
+	{
+		/// <summary>
+		/// Whether or not this reservation service is enabled
+		/// </summary>
+		public static bool Enabled 
+			=> CommandUtils.ParseParam(Globals.Params.AllArguments, "DeviceURL") 
+			&& CommandUtils.ParseParam(Globals.Params.AllArguments, "DevicePool");
+
+		/// <summary>
+		/// List of devices reserved from this service
+		/// </summary>
+		public List<ITargetDevice> ReservedDevices { get; protected set; } = new List<ITargetDevice>();
+
+		/// <summary>
+		/// The base URL of the horde service to request devices from 
+		/// </summary>
+		[AutoParam("")]
+		public string DeviceURL { get; protected set; }
+
+		/// <summary>
+		/// Name of the horde device pool to request devices from 
+		/// </summary>
+		[AutoParamWithNames(Default: "", "DevicePool")]
+		public string DevicePoolID { get; protected set; }
+
+		/// <summary>
+		/// Endpoint for reservations
+		/// </summary>
+		private Uri ReservationServerUri;
+
+		/// <summary>
+		/// Device to reservation lookup
+		/// </summary>
+		private Dictionary<ITargetDevice, DeviceReservationAutoRenew> ServiceReservations = new Dictionary<ITargetDevice, DeviceReservationAutoRenew>();
+
+		/// <summary>
+		/// Target device info, private for reservation use
+		/// </summary>
+		private Dictionary<ITargetDevice, DeviceDefinition> ServiceDeviceInfo = new Dictionary<ITargetDevice, DeviceDefinition>();
+
+		private Dictionary<ITargetDevice, bool> InitialConnectionState = new Dictionary<ITargetDevice, bool>();
+
+		private bool? IsInstallStep
+		{
+			get
+			{
+				return DevicePool.IsInstallStep;
+			}
+			set
+			{
+				DevicePool.IsInstallStep = value;
+			}
+		}
+
+		public HordeDeviceReservationService()
+		{
+			AutoParam.ApplyParamsAndDefaults(this, Globals.Params.AllArguments);
+
+			if (!Uri.TryCreate(DeviceURL, UriKind.Absolute, out ReservationServerUri))
+			{
+				throw new AutomationException("Failed to resolve \"{0}\" as a valid URI", DeviceURL);
+			}
+		}
+
+		#region IDisposable
+		private bool bDisposed;
+		~HordeDeviceReservationService()
+		{
+			Dispose(false);
+		}
+
+		public void Dispose()
+		{
+			Dispose(true);
+		}
+
+		public void Dispose(bool bDisposing)
+		{
+			if (bDisposed)
+			{
+				return;
+			}
+
+			if (bDisposing)
+			{
+				ReleaseDevices(ReservedDevices);
+			}
+
+			bDisposed = true;
+			GC.SuppressFinalize(this);
+		}
+		#endregion
+
+		public virtual bool CanSupportDeviceConstraint(UnrealDeviceTargetConstraint Constraint)
+		{
+			// By default, do not support desktops. Can be overridden in a subtype
+			UnrealTargetPlatform[] SupportedPlatforms = UnrealTargetPlatform.GetValidPlatforms()
+				.Where(Platform => !Platform.IsInGroup(UnrealPlatformGroup.Desktop))
+				.ToArray();
+
+			if (Constraint.Platform == null || !SupportedPlatforms.Contains(Constraint.Platform.Value))
+			{
+				// If an unsupported device, we can't reserve it
+				Log.Info(KnownLogEvents.Gauntlet_DeviceEvent, "Unable to reserve service device of type: {Type}", Constraint.Platform);
+				return false;
+			}
+
+			if (!string.IsNullOrEmpty(Constraint.Model))
+			{
+				// if specific device model, we can't currently reserve it from (legacy) service
+				if (DeviceURL.ToLower().Contains("deviceservice.epicgames.net"))
+				{
+					Log.Info(KnownLogEvents.Gauntlet_DeviceEvent, "Unable to reserve service device of model: {Model} on legacy service", Constraint.Model);
+					return false;
+				}
+			}
+
+			return true;
+		}
+
+		public virtual bool ReserveDevicesFromService(IEnumerable<UnrealDeviceTargetConstraint> RequestedConstraints)
+		{
+			// Ensure no duplicate requests of an explicit device
+			HashSet<string> DeviceNames = new HashSet<string>();
+			foreach (UnrealDeviceTargetConstraint Constraint in RequestedConstraints)
+			{
+				if (!string.IsNullOrEmpty(Constraint.DeviceName))
+				{
+					if (DeviceNames.Contains(Constraint.DeviceName))
+					{
+						Log.Info(KnownLogEvents.Gauntlet_DeviceEvent, "Attempted to make a reservation for multiple devices using the same device name {DeviceName}. This is not supported.", Constraint.DeviceName);
+						return false;
+					}
+
+					DeviceNames.Add(Constraint.DeviceName);
+				}
+			}
+
+			List<ITargetDevice> ScopeReservedDevices = new List<ITargetDevice>();
+			foreach (UnrealDeviceTargetConstraint Constraint in RequestedConstraints)
+			{
+				Reservation NewReservation = Reservation.Create(ReservationServerUri, [Constraint.FormatWithIdentifier()], TimeSpan.FromMinutes(10), RetryMax: 0, PoolID: DevicePoolID, DeviceName: Constraint.DeviceName);
+				DeviceReservationAutoRenew DeviceReservation = new DeviceReservationAutoRenew(DeviceURL, NewReservation);
+
+				if (DeviceReservation == null || DeviceReservation.Devices.Count != 1)
+				{
+					Log.Warning(KnownLogEvents.Gauntlet_DeviceEvent, "Unable to make device registration with constraint {Constraint}", Constraint);
+					return false;
+				}
+
+				if (DevicePool.DeviceReservationBlock = DeviceReservation.InstallRequired != null)
+				{
+					// InstallRequired is true only for the first reservation attempt of the current step. Cache that value to avoid conflicting value on retry attempt.
+					IsInstallStep = IsInstallStep != null ? IsInstallStep : DeviceReservation.InstallRequired == true;
+					DevicePool.SkipInstall = DeviceReservation.InstallRequired == false && IsInstallStep == false;
+					DevicePool.FullClean = !DevicePool.SkipInstall;
+				}
+
+				// Construct a definition from the reservation
+				Device Device = DeviceReservation.Devices[0];
+				DeviceDefinition DeviceDefinition = new DeviceDefinition()
+				{
+					Address = Device.IPOrHostName,
+					Name = Device.Name,
+					Platform = UnrealTargetPlatform.Parse(UnrealTargetPlatform.GetValidPlatformNames().FirstOrDefault(Entry => Entry == Device.Type.Replace("-DevKit", "", StringComparison.OrdinalIgnoreCase))),
+					DeviceData = Device.DeviceData,
+					Model = Device.Model
+				};
+
+				EPerfSpec PerfSpec = EPerfSpec.Unspecified;
+				if (!string.IsNullOrEmpty(Device.PerfSpec) && !Enum.TryParse(Device.PerfSpec, true, out PerfSpec))
+				{
+					throw new AutomationException("Unable to convert perfspec '{0}' into an EPerfSpec", Device.PerfSpec);
+				}
+				DeviceDefinition.PerfSpec = PerfSpec;
+
+				ITargetDevice TargetDevice = DevicePool.Instance.CreateAndRegisterDeviceFromDefinition(DeviceDefinition, Constraint, this);
+
+				// If a device from service can't be added, fail reservation and cleanup devices
+				if (TargetDevice == null)
+				{
+					// If some devices from reservation have been created, release them which will also dispose of reservation
+					if (ScopeReservedDevices.Count > 0)
+					{
+						ReleaseDevices(ScopeReservedDevices);
+					}
+
+					// Cancel this reservation
+					DeviceReservation.Dispose();
+
+					Log.Info(KnownLogEvents.Gauntlet_DeviceEvent, "Unable to make device registration: device registration failed for {Platform}:{DeviceName}", DeviceDefinition.Platform, DeviceDefinition.Name);
+					return false;
+				}
+				else
+				{
+					ScopeReservedDevices.Add(TargetDevice);
+					ReservedDevices.Add(TargetDevice);
+					ServiceDeviceInfo.Add(TargetDevice, DeviceDefinition);
+					ServiceReservations.Add(TargetDevice, DeviceReservation);
+					InitialConnectionState.Add(TargetDevice, TargetDevice.IsConnected);
+				}
+			}
+
+			if (ScopeReservedDevices.Count == RequestedConstraints.Count())
+			{
+				return true;
+			}
+
+			Log.Info(KnownLogEvents.Gauntlet_DeviceEvent, "Unable to reserve all devices from service.");
+			return false;
+		}
+
+		/// <summary>
+		/// Release all devices in the provided list from our reserved list
+		/// </summary>
+		/// <param name="DeviceList"></param>
+		public void ReleaseDevices(IEnumerable<ITargetDevice> DeviceList)
+		{
+			if (!DeviceList.Any())
+			{
+				return;
+			}
+
+			// Remove all these devices from our reserved list
+			ReservedDevices = ReservedDevices.Except(DeviceList).ToList();
+
+			List<ITargetDevice> ThrowDevices = new List<ITargetDevice>();
+			foreach (ITargetDevice Device in DeviceList)
+			{
+				// Reset any state if necessary
+				DeviceConfigurationCache.Instance.RevertDeviceConfiguration(Device);
+
+				if (Device.IsConnected && !InitialConnectionState[Device])
+				{
+					Device.Disconnect();
+				}
+
+				// Unregister device
+				if (ServiceReservations.TryGetValue(Device, out DeviceReservationAutoRenew Reservation))
+				{
+					bool DisposeReservation = ServiceReservations.Count(Entry => Entry.Value == Reservation) == 1;
+
+					// remove and dispose of device
+					// @todo: add support for reservation modification on server (partial device release)
+					ServiceReservations.Remove(Device);
+					ServiceDeviceInfo.Remove(Device);
+					InitialConnectionState.Remove(Device);
+
+					Device.Dispose();
+					Reservation.Dispose();
+				}
+				else
+				{
+					ThrowDevices.Add(Device);
+				}
+			}
+
+			if (ThrowDevices.Any())
+			{
+				// If a user explicitly calls a service's release function with an incorrect device, throw this exception
+				string ExceptionMessage = "Attempted to release the following devices from a service that did not reserve them!";
+				ExceptionMessage += string.Join("\n\t", ThrowDevices.Select(Device => Device.Name));
+				ExceptionMessage += "\nUse DevicePool.Instance.ReleaseDevices to avoid mistakingly releasing devices reserved from a different service";
+				throw new AutomationException(ExceptionMessage);
+			}
+		}
+
+		/// <summary>
+		/// Report target device issue to service with given error message
+		/// </summary>
+		public virtual void ReportDeviceError(string DeviceName, string ErrorMessage)
+		{
+			// TargetDevice name is not always DeviceData name... need to try to resolve target device names
+			ITargetDevice MatchingDevice = null;
+			foreach (ITargetDevice Device in ReservedDevices)
+			{
+				if (Device.Name.Equals(DeviceName, StringComparison.OrdinalIgnoreCase))
+				{
+					MatchingDevice = Device;
+					break;
+				}
+			}
+
+			if (MatchingDevice == null)
+			{
+				// No Target device, assume this is DeviceData name
+				Reservation.ReportDeviceError(DeviceURL, DeviceName, ErrorMessage);
+			}
+			else
+			{
+				Reservation.ReportDeviceError(DeviceURL, ServiceDeviceInfo[MatchingDevice].Name, ErrorMessage);
+			}
+		}
+	}
+}
\ No newline at end of file
diff --git a/Engine/Source/Programs/AutomationTool/Gauntlet/Framework/Devices/Gauntlet.UnrealDeviceReservation.cs b/Engine/Source/Programs/AutomationTool/Gauntlet/Framework/Devices/Gauntlet.UnrealDeviceReservation.cs
--- a/Engine/Source/Programs/AutomationTool/Gauntlet/Framework/Devices/Gauntlet.UnrealDeviceReservation.cs
+++ b/Engine/Source/Programs/AutomationTool/Gauntlet/Framework/Devices/Gauntlet.UnrealDeviceReservation.cs
@@ -12,27 +12,46 @@ namespace Gauntlet
 	public class UnrealDeviceReservation
 	{
 		protected List<ProblemDevice> ProblemDevices { get; private set; } = new List<ProblemDevice>();
-		public List<ITargetDevice> ReservedDevices { get; protected set; }
+		public List<ITargetDevice> ReservedDevices { get; protected set; } = new List<ITargetDevice>();
 
-		public bool TryReserveDevices(Dictionary<UnrealDeviceTargetConstraint, int> RequiredDeviceTypes, int ExpectedNumberOfDevices)
+		public bool TryReserveDevices(Dictionary<UnrealDeviceTargetConstraint, int> RequiredDeviceTypes, int ExpectedNumberOfDevices, bool bAllowPartialReservations = false)
 		{
-			List<ITargetDevice> AcquiredDevices = new List<ITargetDevice>();
-			List<ITargetDevice> SkippedDevices = new List<ITargetDevice>();
-
-			ReleaseDevices();
+			// First, determine if this is a partial reservation. In the event a device is considered a problem device, don't discard all devices.
+			// Instead, only discard the devices that have a problem.
+			if (bAllowPartialReservations && ReservedDevices.Count > 0)
+			{
+				RequiredDeviceTypes = GetPartialReservationTypes(RequiredDeviceTypes);
+				if (RequiredDeviceTypes.Count == 0)
+				{
+					// We've already reserved all the requested devices
+					return true;
+				}
 
-			ProblemDevices.Clear();
+				int NewExpectedCount = 0;
+				foreach (var Pair in RequiredDeviceTypes)
+				{
+					NewExpectedCount += Pair.Value;
+				}
+				ExpectedNumberOfDevices = NewExpectedCount;
+			}
+			else
+			{
+				ReleaseDevices();
+				ProblemDevices.Clear();
+			}
 
-			// check whether pool can accommodate devices
+			// Ensure the device pool can support the request. This will reserve devices from services if it needs to
 			if (!DevicePool.Instance.CheckAvailableDevices(RequiredDeviceTypes, ProblemDevices))
 			{
 				return false;
 			}
 
-			// nothing acquired yet...
-			AcquiredDevices.Clear();
+			List<ITargetDevice> AcquiredDevices = new List<ITargetDevice>();
+			List<ITargetDevice> SkippedDevices = new List<ITargetDevice>();
 
-			// for each platform, enumerate and select from the available devices
+			// The pool now contains enough devices to support this reservation. 
+			// For each platform, enumerate and select from the available devices as long as they match the predicate.
+			// This will discard any devices that are considered unavailable or marked as a problem
 			foreach (var PlatformReqKP in RequiredDeviceTypes)
 			{
 				UnrealDeviceTargetConstraint Constraint = PlatformReqKP.Key;
@@ -75,70 +94,48 @@ namespace Gauntlet
 				});
 			}
 
-			// If we got enough devices, go to step2 where we provision and try to connect them
-			if (AcquiredDevices.Count == ExpectedNumberOfDevices)
-			{
-				// actually acquire them
-				DevicePool.Instance.ReserveDevices(AcquiredDevices);
+			// Release any devices that were skipped
+			DevicePool.Instance.ReleaseDevices(SkippedDevices);
 
-				Log.Info("Selected devices {Devices} for client(s). Prepping", string.Join(", ", AcquiredDevices));
+			// Obtain connections to any required devices.
+			EstablishDeviceConnections(AcquiredDevices);
 
-				foreach (ITargetDevice Device in AcquiredDevices)
-				{
-					if (Device.IsOn == false)
-					{
-						Log.Info("Powering on {0}", Device);
-						Device.PowerOn();
-					}
-					else if (Globals.Params.ParseParam("reboot"))
-					{
-						Log.Info("Rebooting {0}", Device);
-						Device.Reboot();
-						// Force a re-login after a reboot to be sure it is ready.
-						if (Globals.Params.ParseParam("VerifyLogin") && Device is IOnlineServiceLogin DeviceLogin)
-						{
-							Log.Info("Revalidate device login after reboot...");
-							DeviceLogin.VerifyLogin();
-						}
-					}
+			// Mark any devices we can't connect to as problems and release them.
+			// Could be something grabbed them before us, could be that they are unresponsive in some way
+			List<ITargetDevice> LostDevices = AcquiredDevices.Where(Device => !Device.IsConnected).ToList();
+			if (LostDevices.Count > 0)
+			{
+				LostDevices.ForEach(Device => MarkProblemDevice(Device, "Agent lost connection to device."));
+				AcquiredDevices = AcquiredDevices.Except(LostDevices).ToList();
+				ReleaseProblemDevices();
+			}
 
-					if (Device.IsConnected == false)
+			// If we failed to acquire every device we needed, but support partial reservations, hold onto any devices we did manage to acquire
+			if (AcquiredDevices.Count < ExpectedNumberOfDevices && AcquiredDevices.Count > 0)
+			{
+				if (bAllowPartialReservations)
+				{
+					if (ClaimDevices(AcquiredDevices))
 					{
-						Log.Verbose("Connecting to {0}", Device);
-						Device.Connect();
+						ReservedDevices.AddRange(AcquiredDevices);
 					}
 				}
-
-				// Step 3: Verify we actually connected to them
-				var LostDevices = AcquiredDevices.Where(D => !D.IsConnected);
-
-				if (LostDevices.Count() > 0)
+				else
 				{
-					// mark these as problems. Could be something grabbed them before us, could be that they are
-					// unresponsive in some way
-					LostDevices.ToList().ForEach(D => MarkProblemDevice(D, "Agent lost connection to device."));
-					AcquiredDevices.ToList().ForEach(D => D.Disconnect());
-
 					DevicePool.Instance.ReleaseDevices(AcquiredDevices);
-					AcquiredDevices.Clear();
 				}
-			}
 
-			if (AcquiredDevices.Count() != ExpectedNumberOfDevices)
-			{
-				Log.Info("Failed to resolve all devices. Releasing the ones we have ");
-				DevicePool.Instance.ReleaseDevices(AcquiredDevices);
-				ReservedDevices = new List<ITargetDevice>();
+				return false;
 			}
-			else
+			
+
+			if(!ClaimDevices(AcquiredDevices))
 			{
-				ReservedDevices = AcquiredDevices;
+				return false;
 			}
 
-			// release devices that were skipped
-			DevicePool.Instance.ReleaseDevices(SkippedDevices);
-
-			return ReservedDevices.Count() == ExpectedNumberOfDevices;
+			ReservedDevices.AddRange(AcquiredDevices);
+			return true;
 		}
 
 		public void ReleaseDevices()
@@ -200,5 +197,88 @@ namespace Gauntlet
 				ProblemDevices.Add(new ProblemDevice(Device.Name, Device.Platform.Value));
 			}
 		}
+
+		/// <summary>
+		/// Converts a set of required devices into only what is necessary by analyzing the list of already reserved devices
+		/// </summary>
+		/// <param name="DeviceTypes">Request list</param>
+		/// <returns>The difference in constraints/values between the request types and the currently reserved devices</returns>
+		private Dictionary<UnrealDeviceTargetConstraint, int> GetPartialReservationTypes(Dictionary<UnrealDeviceTargetConstraint, int> DeviceTypes)
+		{
+			Dictionary<UnrealDeviceTargetConstraint, int> CurrentlyAvailableTypes = new();
+			foreach (ITargetDevice Device in ReservedDevices)
+			{
+				UnrealDeviceTargetConstraint Constraint = DevicePool.Instance.GetConstraint(Device);
+				if (CurrentlyAvailableTypes.ContainsKey(Constraint))
+				{
+					++CurrentlyAvailableTypes[Constraint];
+				}
+				else
+				{
+					CurrentlyAvailableTypes.Add(Constraint, 1);
+				}
+			}
+
+			Dictionary<UnrealDeviceTargetConstraint, int> NewSet = new();
+			foreach (var Pair in DeviceTypes)
+			{
+				if (CurrentlyAvailableTypes.ContainsKey(Pair.Key))
+				{
+					// If we have more devices already reserved than what is requested we don't need to reserve devices
+					if (CurrentlyAvailableTypes[Pair.Key] < Pair.Value)
+					{
+						// We have at least 1 device with this constraint, but not enough to fufill the request. Request only what is needed for this constraint
+						NewSet.Add(Pair.Key, Pair.Value - CurrentlyAvailableTypes[Pair.Key]);
+					}
+				}
+				else
+				{
+					// No existing constraint, we don't have these devices
+					NewSet.Add(Pair.Key, Pair.Value);
+				}
+			}
+
+			return NewSet;
+		}
+
+		private void EstablishDeviceConnections(IEnumerable<ITargetDevice> Devices)
+		{
+			foreach (ITargetDevice Device in Devices)
+			{
+				if (Device.IsOn == false)
+				{
+					Log.Info("Powering on {0}", Device);
+					Device.PowerOn();
+				}
+				else if (Globals.Params.ParseParam("reboot"))
+				{
+					Log.Info("Rebooting {0}", Device);
+					Device.Reboot();
+					// Force a re-login after a reboot to be sure it is ready.
+					if (Globals.Params.ParseParam("VerifyLogin") && Device is IOnlineServiceLogin DeviceLogin)
+					{
+						Log.Info("Revalidate device login after reboot...");
+						DeviceLogin.VerifyLogin();
+					}
+				}
+
+				if (Device.IsConnected == false)
+				{
+					Log.Verbose("Connecting to {0}", Device);
+					Device.Connect();
+				}
+			}
+		}
+
+		private bool ClaimDevices(IEnumerable<ITargetDevice> Devices)
+		{
+			bool bSuccess = DevicePool.Instance.ClaimDevices(Devices);
+			if (!bSuccess)
+			{
+				Log.Warning("Attempted to claim the following devices which are already marked as claimed:\n{Devices}",
+						string.Join("\n\t", Devices.Select(Device => Device.Name)));
+			}
+			return bSuccess;
+		}
 	}
 }
\ No newline at end of file
diff --git a/Engine/Source/Programs/AutomationTool/Gauntlet/Framework/Gauntlet.Utils.cs b/Engine/Source/Programs/AutomationTool/Gauntlet/Framework/Gauntlet.Utils.cs
--- a/Engine/Source/Programs/AutomationTool/Gauntlet/Framework/Gauntlet.Utils.cs
+++ b/Engine/Source/Programs/AutomationTool/Gauntlet/Framework/Gauntlet.Utils.cs
@@ -825,21 +825,37 @@ namespace Gauntlet
 
 		public static class InterfaceHelpers
 		{
-
-			public static IEnumerable<InterfaceType> FindImplementations<InterfaceType>(bool bIncludeCompiledScripts = false)
+			public static IEnumerable<Type> FindTypes<InterfaceType>(bool bIncludeCompiledScripts = false, bool bConcreteTypesOnly = false)
 				where InterfaceType : class
 			{
-				HashSet<Type> AllTypes = new HashSet<Type>(Assembly.GetExecutingAssembly().GetTypes().Where(T => typeof(InterfaceType).IsAssignableFrom(T)));
+				HashSet<Type> AllTypes = new HashSet<Type>(Assembly.GetExecutingAssembly().GetTypes().Where(Type => typeof(InterfaceType).IsAssignableFrom(Type)));
+				if (bConcreteTypesOnly)
+				{
+					AllTypes = AllTypes.Where(Type => !Type.IsAbstract && !Type.IsInterface).ToHashSet();
+				}
 
 				if (bIncludeCompiledScripts)
 				{
 					foreach (Assembly Assembly in ScriptManager.AllScriptAssemblies)
 					{
-						List<Type> AssemblyTypes = Assembly.GetTypes().Where(T => typeof(InterfaceType).IsAssignableFrom(T)).ToList();
+						HashSet<Type> AssemblyTypes = new HashSet<Type>(Assembly.GetTypes().Where(T => typeof(InterfaceType).IsAssignableFrom(T)));
+						if (bConcreteTypesOnly)
+						{
+							AssemblyTypes = AllTypes.Where(Type => !Type.IsAbstract && !Type.IsInterface).ToHashSet();
+						}
+
 						AllTypes.UnionWith(AssemblyTypes);
 					}
 				}
 
+				return AllTypes;
+			}
+
+			public static IEnumerable<InterfaceType> FindImplementations<InterfaceType>(bool bIncludeCompiledScripts = false)
+				where InterfaceType : class
+			{
+				HashSet<Type> AllTypes = FindTypes<InterfaceType>(bIncludeCompiledScripts).ToHashSet();
+
 				List<InterfaceType> ConstructedTypes = new List<InterfaceType>();
 
 				foreach (Type FoundType in AllTypes)
diff --git a/Engine/Source/Programs/AutomationTool/Gauntlet/Unreal/Base/Gauntlet.UnrealSession.cs b/Engine/Source/Programs/AutomationTool/Gauntlet/Unreal/Base/Gauntlet.UnrealSession.cs
--- a/Engine/Source/Programs/AutomationTool/Gauntlet/Unreal/Base/Gauntlet.UnrealSession.cs
+++ b/Engine/Source/Programs/AutomationTool/Gauntlet/Unreal/Base/Gauntlet.UnrealSession.cs
@@ -727,9 +727,11 @@ namespace Gauntlet
 				return true;
 			}
 
-			// figure out how many of each device we need
+			// Figure out how many of each device we need
 			Dictionary<UnrealDeviceTargetConstraint, int> RequiredDeviceTypes = new Dictionary<UnrealDeviceTargetConstraint, int>();
-			IEnumerable<UnrealSessionRole> RolesThatRequireDevice = SessionRoles.Where(R => !R.IsNullRole());
+			IEnumerable<UnrealSessionRole> RolesThatRequireDevice = SessionRoles
+				.Where(Role => !Role.IsNullRole())
+				.Where(RoleNeedsDevice);
 
 			// Get a count of the number of devices required for each platform
 			foreach (UnrealSessionRole Role in RolesThatRequireDevice)
@@ -746,7 +748,7 @@ namespace Gauntlet
 
 			try
 			{
-				return UnrealDeviceReservation.TryReserveDevices(RequiredDeviceTypes, RolesThatRequireDevice.Count());
+				return UnrealDeviceReservation.TryReserveDevices(RequiredDeviceTypes, RolesThatRequireDevice.Count(), true);
 			}
 			catch (Exception ex)
 			{
diff --git a/Engine/Source/Programs/AutomationTool/Gauntlet/Unreal/RunUnreal.cs b/Engine/Source/Programs/AutomationTool/Gauntlet/Unreal/RunUnreal.cs
--- a/Engine/Source/Programs/AutomationTool/Gauntlet/Unreal/RunUnreal.cs
+++ b/Engine/Source/Programs/AutomationTool/Gauntlet/Unreal/RunUnreal.cs
@@ -320,9 +320,6 @@ namespace Gauntlet
 			// Tear down accounts
 			TeardownAccounts();
 
-			// dispose now, not during shutdown gc, because this runs commands...
-			DevicePool.Instance.Dispose();
-
 			// If only running a single unreal test node, we'll direct the summary artifact link to the specific test
 			string ArtifactPath = Globals.LogDir;
 			if(AllTestNodes.Count == 1)
