[Gauntlet] Fixes assembly get type overwrite error and throws exception rather than logging error when failing to find reservation services

diff --git a/Engine/Source/Programs/AutomationTool/Gauntlet/Framework/Devices/Gauntlet.DevicePool.cs b/Engine/Source/Programs/AutomationTool/Gauntlet/Framework/Devices/Gauntlet.DevicePool.cs
--- a/Engine/Source/Programs/AutomationTool/Gauntlet/Framework/Devices/Gauntlet.DevicePool.cs
+++ b/Engine/Source/Programs/AutomationTool/Gauntlet/Framework/Devices/Gauntlet.DevicePool.cs
@@ -824,9 +824,7 @@ namespace Gauntlet
 				Message += "\nThe following reservation services were available:\n";
 				Message += string.Join("\n\t", ReservationServices.Select(Service => Service.GetType().Name));
 				Message += "\nEnsure your environment is configured to enable any necessary services.";
-				Log.Warning(Message);
-
-				return false;
+				throw new AutomationException(Message);
 			}
 
 			foreach (var ReservationPair in ServiceMapping)
diff --git a/Engine/Source/Programs/AutomationTool/Gauntlet/Framework/Gauntlet.Utils.cs b/Engine/Source/Programs/AutomationTool/Gauntlet/Framework/Gauntlet.Utils.cs
--- a/Engine/Source/Programs/AutomationTool/Gauntlet/Framework/Gauntlet.Utils.cs
+++ b/Engine/Source/Programs/AutomationTool/Gauntlet/Framework/Gauntlet.Utils.cs
@@ -841,7 +841,7 @@ namespace Gauntlet
 						HashSet<Type> AssemblyTypes = new HashSet<Type>(Assembly.GetTypes().Where(T => typeof(InterfaceType).IsAssignableFrom(T)));
 						if (bConcreteTypesOnly)
 						{
-							AssemblyTypes = AllTypes.Where(Type => !Type.IsAbstract && !Type.IsInterface).ToHashSet();
+							AssemblyTypes = AssemblyTypes.Where(Type => !Type.IsAbstract && !Type.IsInterface).ToHashSet();
 						}
 
 						AllTypes.UnionWith(AssemblyTypes);
