- Fixed a log to have right params

diff --git a/Engine/Source/Programs/AutomationTool/Apple/ExportAppleBuild.Automation.cs b/Engine/Source/Programs/AutomationTool/Apple/ExportAppleBuild.Automation.cs
--- a/Engine/Source/Programs/AutomationTool/Apple/ExportAppleBuild.Automation.cs
+++ b/Engine/Source/Programs/AutomationTool/Apple/ExportAppleBuild.Automation.cs
@@ -572,7 +572,7 @@ class ExportAppleBuild : BuildCommand
 				{
 					Ex.Handle(x => 
 						{
-							Logger.LogWarning("Got an unknown exception while updating the WhatsNew of the build. Will try again...", Ex.Message);
+							Logger.LogWarning("Got an unknown exception ('{Exception}') while updating the WhatsNew of the build. Will try again...", Ex.Message);
 							return true;
 						}
 					);
