[Backout] - CL40471098
#rnx
Original CL Desc
-----------------------------------------------------------------
[Gauntlet] Replaces DevicePool's strict reliance on horde as the only means of device reservations. This change adds IDeviceReservationService, which provides the ability to define custom services to reserve, construct, and manage devices.

IDeviceReservationService Summary
========================
Services are disabled by default and can be made available by configurable program state. For example, the HordeDeviceReservationService will become available when -DeviceURL= and -DevicePool= commandline arguments are present.
Each service is able to define whether it can support a particular UnrealDeviceTargetConstraint. This means multiple services can be enabled and each support different types of devices.
For example, HordeDeviceReservationService may provide consoles but a custom service may provide linux servers from a cloud provider.
Device error reporting and reservation release implementations are also left
DevicePool has been made to interface with all available services, but continues to manage the lifetimes and availability of all local devices, whether those be host machine instances or connected devkits.

Other Details/Improvements
========================
- The existing implementation for the Gaunlet's Horde implementation has been moved to the HordeReservationService type, which is an IDeviceReservationService
- Deprecates public interfaces to horde specific reservations from DevicePool
- UnrealDeviceReservation has a new option to avoid releasing all devices when reserving new ones. This is useful when encountering problem devices, as we can selectively discard and re-reserve only that device constraint.

Bug Fixes (applies to legacy flow as well)
========================
RunUnreal will now only dispose of the device pool once.

Reserving a device no longer immediatly "claims" it. Instead it's made "available".
Any consumer of the device pool that relies on reservations would have to ignore a devices claimed status in order to use it.
UnrealDeviceReservation was doing this, which has also been fixed.

diff --git a/Engine/Source/Programs/AutomationTool/Gauntlet/Framework/Devices/Gauntlet.DevicePool.cs b/Engine/Source/Programs/AutomationTool/Gauntlet/Framework/Devices/Gauntlet.DevicePool.cs
--- a/Engine/Source/Programs/AutomationTool/Gauntlet/Framework/Devices/Gauntlet.DevicePool.cs
+++ b/Engine/Source/Programs/AutomationTool/Gauntlet/Framework/Devices/Gauntlet.DevicePool.cs
@@ -13,7 +13,6 @@ using System.Text.Json.Serialization;
 using EpicGames.Core;
 using System.Data;
 using Gauntlet.Utils;
-using System.Reflection;
 
 namespace Gauntlet
 {
@@ -228,27 +227,7 @@ namespace Gauntlet
 		/// <summary>
 		/// Access to our singleton
 		/// </summary>
-		public static DevicePool Instance
-		{
-			get
-			{
-				if (_Instance == null || _Instance.bDisposed)
-				{
-					if (_Instance != null)
-					{
-						Log.Info("DevicePool has been disposed. Reinitilizing with a new instance.");
-					}
-
-					_Instance = new DevicePool();
-				}
-
-				return _Instance;
-			}
-			private set
-			{
-				_Instance = value;
-			}
-		}
+		public static DevicePool Instance { get; private set; } = new DevicePool();
 
 		public static bool SkipInstall;
 
@@ -256,12 +235,12 @@ namespace Gauntlet
 
 		public static bool DeviceReservationBlock;
 
-		public static bool? IsInstallStep;
+		private static bool? IsInstallStep;
 
 		/// <summary>
-		/// Active pool instance
+		/// Device reservation service URL
 		/// </summary>
-		private static DevicePool _Instance;
+		public string DeviceURL;
 
 		/// <summary>
 		/// Object used for locking access to internal data
@@ -269,32 +248,29 @@ namespace Gauntlet
 		private object LockObject = new object();
 
 		/// <summary>
-		/// List of all provisioned devices that are can be claimed
+		/// List of all provisioned devices that are available for reservations
 		/// </summary>
 		private List<ITargetDevice> AvailableDevices = new List<ITargetDevice>();
 
 		/// <summary>
-		/// List of all provisioned devices that have been claimed
+		/// List of all devices that have been reserved
 		/// </summary>
-		private List<ITargetDevice> ClaimedDevices = new List<ITargetDevice>();
+		private List<ITargetDevice> ReservedDevices = new List<ITargetDevice>();
 
 		/// <summary>
-		/// List of all provisioned devices that were not provided by a reservation service
+		/// List of platforms we've had devices for
 		/// </summary>
-		private List<ITargetDevice> LocalDevices => AvailableDevices
-			.Union(ClaimedDevices)
-			.Except(ReservationServices.SelectMany(Service => Service.ReservedDevices))
-			.ToList();
+		private HashSet<UnrealTargetPlatform?> UsedPlatforms = new HashSet<UnrealTargetPlatform?>();
 
 		/// <summary>
-		/// List of all enabled reservation services
+		/// Device to reservation lookup
 		/// </summary>
-		private List<IDeviceReservationService> ReservationServices = new List<IDeviceReservationService>();
+		private Dictionary<ITargetDevice, DeviceReservationAutoRenew> ServiceReservations = new Dictionary<ITargetDevice, DeviceReservationAutoRenew>();
 
 		/// <summary>
-		/// List of platforms we've had devices for
+		/// Target device info, private for reservation use
 		/// </summary>
-		private HashSet<UnrealTargetPlatform?> UsedPlatforms = new HashSet<UnrealTargetPlatform?>();
+		private Dictionary<ITargetDevice, DeviceDefinition> ServiceDeviceInfo = new Dictionary<ITargetDevice, DeviceDefinition>();
 
 		/// <summary>
 		/// List of device definitions that can be provisioned on demand
@@ -306,6 +282,12 @@ namespace Gauntlet
 		/// </summary>
 		private List<DeviceDefinition> FailedProvisions = new List<DeviceDefinition>();
 
+		/// <summary>
+		/// Local connection state for any device. We use this to avoid disconnecting devices that the
+		/// user may already be connected to when running tests
+		/// </summary>
+		private Dictionary<ITargetDevice, bool> InitialConnectState = new Dictionary<ITargetDevice, bool>();
+
 		/// <summary>
 		/// Device constraints for performance profiles, etc
 		/// </summary>
@@ -316,16 +298,23 @@ namespace Gauntlet
 		/// </summary>
 		private string LocalTempDir;
 
+		/// <summary>
+		/// The maximum number of problem devices to report to device backend
+		/// This mitigates issues on builders, such as hung local processes, incorrectly reporting problem devices
+		/// </summary>
+		private int MaxDeviceErrorReports = 3;
+
 		/// <summary>
 		/// Protected constructor - code should use DevicePool.Instance
 		/// </summary>
 		protected DevicePool()
 		{
-			lock (LockObject)
+			// create two local devices by default?
+			AddLocalDevices(2);
+
+			if(Instance == null)
 			{
-				// Create two local devices by default for ease of running a client and server
-				AddLocalDevices(2);
-				InitializeReservationServices();
+				Instance = this;
 			}
 
 			SkipInstall = Globals.Params.ParseParams("SkipInstall", "SkipCopy", "SkipDeploy");
@@ -334,7 +323,6 @@ namespace Gauntlet
 		}
 
 		#region IDisposable Support
-		bool bDisposed = false;
 		~DevicePool()
 		{
 			Dispose(false);
@@ -361,39 +349,54 @@ namespace Gauntlet
 		/// <param name="disposing"></param>
 		protected virtual void Dispose(bool disposing)
 		{
-			if (bDisposed)
-			{
-				return;
-			}
-
 			lock (LockObject)
 			{
 				if (disposing)
 				{
-					// Dispose all local devices
-					foreach (ITargetDevice LocalDevice in LocalDevices)
+					// warn if more than a third of the devices in our pool had issues.
+					int TotalDevices = AvailableDevices.Count + UnprovisionedDevices.Count;
+
+					if (TotalDevices > 0)
 					{
-						DeviceConfigurationCache.Instance.RevertDeviceConfiguration(LocalDevice);
-						LocalDevice.Dispose();
+						float Failures = FailedProvisions.Count / TotalDevices;
+
+						if (Failures >= 0.33)
+						{
+							FailedProvisions.ForEach(D => Log.Warning("Device {DeviceName} could not be added", D.Name));
+						}
 					}
 
-					// Dispose of any services. These dispose their associated devices
-					foreach (IDeviceReservationService Service in ReservationServices)
+					// release any outstanding reservations
+					ReleaseReservations();
+
+					// warn if anything was left registered
+					ReservedDevices.ForEach(D =>
 					{
-						Service.Dispose();
+						Log.Warning("Device {DeviceName} was not unregistered prior to DevicePool.Dispose!. Forcibly disposing.", D.Name);
+						D.Dispose();
+					});
+
+					ReleaseDevices(ReservedDevices);
+
+					// Cleanup our available devices
+					foreach (ITargetDevice Device in AvailableDevices)
+					{
+						if (Device.IsConnected && InitialConnectState[Device] == false)
+						{
+							Device.Disconnect();
+						}
+
+						Device.Dispose();
 					}
 
-					// Cleanup things like platform sdk daemons
+					// Additional devices cleanup
 					CleanupDevices(UsedPlatforms);
 
 					AvailableDevices.Clear();
-					ClaimedDevices.Clear();
-					ReservationServices.Clear();
 					UnprovisionedDevices.Clear();
+					ReservedDevices.Clear();
 				}
 			}
-
-			bDisposed = true;
 		}
 		#endregion
 
@@ -418,13 +421,14 @@ namespace Gauntlet
 		{
 			lock (LockObject)
 			{
-				return AvailableDevices.Union(ClaimedDevices).Where(D => Validate == null ? Constraint.Check(D) : Validate(D)).Count() +
+				return AvailableDevices.Union(ReservedDevices).Where(D => Validate == null ? Constraint.Check(D) : Validate(D)).Count() +
 					UnprovisionedDevices.Where(D => Constraint.Check(D)).Count();
 			}
 		}
 
 		public UnrealDeviceTargetConstraint GetConstraint(ITargetDevice Device)
 		{
+
 			if (!Constraints.ContainsKey(Device))
 			{
 				throw new AutomationException("Device pool has no contstaint for {0} (device was likely released)", Device);
@@ -436,7 +440,7 @@ namespace Gauntlet
 		public void SetLocalOptions(string InLocalTemp, bool InUniqueTemps = false, string InDeviceURL = "")
 		{
 			LocalTempDir = InLocalTemp;
-			Legacy_DeviceURL = InDeviceURL;
+			DeviceURL = InDeviceURL;
 		}
 
 		public void AddLocalDevices(int MaxCount)
@@ -592,14 +596,9 @@ namespace Gauntlet
 							string DeviceAddress = DeviceRef;
 							string Model = string.Empty;
 
-							// When using device services, skip adding non-desktop local devices to pool if any of the services can support that platform
-							bool IsDesktop = DevicePlatform.IsInGroup(UnrealPlatformGroup.Desktop);
-							bool ReservationsEnabled = ReservationServices.Count > 0;
-							bool ServicesCanSupportThisPlatform = ReservationsEnabled 
-								&& ReservationServices.Where(Service => Service.CanSupportDeviceConstraint(new UnrealDeviceTargetConstraint(DevicePlatform))).Any();
-							bool DeviceServiceEnabled = ServicesCanSupportThisPlatform;
-
-							if (!IsDesktop && DeviceRef.Equals("default", StringComparison.OrdinalIgnoreCase) && DeviceServiceEnabled)
+							// when using device service, skip adding local non-desktop devices to pool
+							bool IsDesktop = UnrealBuildTool.Utils.GetPlatformsInClass(UnrealPlatformClass.Desktop).Contains(DevicePlatform);
+							if (!IsDesktop && DeviceRef.Equals("default", StringComparison.OrdinalIgnoreCase) && !String.IsNullOrEmpty(DeviceURL))
 							{
 								continue;
 							}
@@ -662,8 +661,9 @@ namespace Gauntlet
 					throw new Exception("Device already registered!");
 				}
 
+				InitialConnectState[Device] = Device.IsConnected;
 				Constraints[Device] = Constraint ?? new UnrealDeviceTargetConstraint(Device.Platform.Value);
-				UsedPlatforms.Add(Device.Platform);
+
 				AvailableDevices.Add(Device);
 
 				if (Log.IsVerbose)
@@ -762,17 +762,17 @@ namespace Gauntlet
 		}
 
 		/// <summary>
-		/// Claim all devices in the provided list. Once reserved a device will not be seen by any code that
+		/// Reserve all devices in the provided list. Once reserved a device will not be seen by any code that
 		/// calls EnumerateDevices
 		/// </summary>
 		/// <param name="DeviceList"></param>
 		/// <returns></returns>
-		public bool ClaimDevices(IEnumerable<ITargetDevice> DeviceList)
+		public bool ReserveDevices(IEnumerable<ITargetDevice> DeviceList)
 		{
 			lock (LockObject)
 			{
 				// can reserve if not reserved...
-				if (ClaimedDevices.Intersect(DeviceList).Count() > 0)
+				if (ReservedDevices.Intersect(DeviceList).Count() > 0)
 				{
 					return false;
 				}
@@ -780,63 +780,130 @@ namespace Gauntlet
 				// remove these devices from the available list
 				AvailableDevices = AvailableDevices.Where(D => DeviceList.Contains(D) == false).ToList();
 
-				ClaimedDevices.AddRange(DeviceList);
+				ReservedDevices.AddRange(DeviceList);
 				DeviceList.All(D => UsedPlatforms.Add(D.Platform));
 			}
 			return true;
 		}
 
-		public bool ReserveDevicesFromService(Dictionary<UnrealDeviceTargetConstraint, int> DeviceTypes)
+		/// <summary>
+		/// Reserve devices from service
+		/// </summary>
+		public bool ReserveDevicesFromService(string DeviceURL, Dictionary<UnrealDeviceTargetConstraint, int> DeviceTypes)
 		{
-			// Flatten the required devices into an enumerable
-			List<UnrealDeviceTargetConstraint> RequiredDeviceConstraints = DeviceTypes
-				.SelectMany(KVP => Enumerable.Repeat(KVP.Key, KVP.Value))
-				.ToList();
-
-			// Project each constraint into a supported service
-			Dictionary<IDeviceReservationService, List<UnrealDeviceTargetConstraint>> ServiceMapping = new();
-			foreach (UnrealDeviceTargetConstraint Constraint in RequiredDeviceConstraints)
+			if (string.IsNullOrEmpty(DeviceURL) || !Uri.TryCreate(DeviceURL, UriKind.Absolute, out Uri ReservationServerUri))
+			{
+				return false;
+			}
+
+			Dictionary<UnrealTargetPlatform, string> DeviceMap = UnrealTargetPlatform.GetValidPlatformNames()
+				.ToDictionary(Platform => UnrealTargetPlatform.Parse(Platform), Platform => Platform);
+
+			string PoolID = Globals.DevicePoolId;
+
+			// First pass, validate all of the constraints are valid/possible
+			foreach (KeyValuePair<UnrealDeviceTargetConstraint, int> Entry in DeviceTypes)
 			{
-				foreach (IDeviceReservationService Service in ReservationServices)
+
+				if (Entry.Key.Platform == null || !DeviceMap.ContainsKey(Entry.Key.Platform.Value))
 				{
-					if (Service.CanSupportDeviceConstraint(Constraint))
+					// If an unsupported device, we can't reserve it
+					Log.Info(KnownLogEvents.Gauntlet_DeviceEvent, "Unable to reserve service device of type: {Type}", Entry.Key.Platform);
+					return false;
+				}
+
+				if (!string.IsNullOrEmpty(Entry.Key.Model))
+				{
+					// if specific device model, we can't currently reserve it from (legacy) service
+					if (DeviceURL.ToLower().Contains("deviceservice.epicgames.net"))
 					{
-						if (ServiceMapping.ContainsKey(Service))
+						Log.Info(KnownLogEvents.Gauntlet_DeviceEvent, "Unable to reserve service device of model: {Model} on legacy service", Entry.Key.Model);
+						return false;
+					}
+				}
+
+				if(!string.IsNullOrEmpty(Entry.Key.DeviceName) && Entry.Value > 1)
+				{
+					Log.Info(KnownLogEvents.Gauntlet_DeviceEvent, "Attempted to make a reservation for multiple devices using the same device name {DeviceName}. This is not supported.", Entry.Key.DeviceName);
+					return false;
+				}
+			}
+
+			// Second pass, reserve devices
+			List<ITargetDevice> ScopeReservedDevices = new List<ITargetDevice>();
+			foreach (KeyValuePair<UnrealDeviceTargetConstraint, int> Entry in DeviceTypes)
+			{
+				string[] DeviceRequirements = Enumerable.Repeat(Entry.Key.FormatWithIdentifier(), Entry.Value).ToArray();
+
+				Reservation NewReservation = Reservation.Create(ReservationServerUri, DeviceRequirements, TimeSpan.FromMinutes(10), PoolID: PoolID, DeviceName: Entry.Key.DeviceName);
+				DeviceReservationAutoRenew DeviceReservation = new DeviceReservationAutoRenew(DeviceURL, NewReservation);
+
+				if (DeviceReservation == null || DeviceReservation.Devices.Count != Entry.Value)
+				{
+					Log.Warning(KnownLogEvents.Gauntlet_DeviceEvent, "Unable to make device registration with constraint {Constraint}", Entry.Key);
+					return false;
+				}
+
+				if (DeviceReservationBlock = DeviceReservation.InstallRequired != null)
+				{
+					// InstallRequired is true only for the first reservation attempt of the current step. Cache that value to avoid conflicting value on retry attempt.
+					IsInstallStep = IsInstallStep != null? IsInstallStep : DeviceReservation.InstallRequired == true;
+					SkipInstall = DeviceReservation.InstallRequired == false && IsInstallStep == false;
+					FullClean = !SkipInstall;
+				}
+
+				// Add target devices from reservation
+				foreach (Device Device in DeviceReservation.Devices)
+				{
+					DeviceDefinition Def = new DeviceDefinition();
+					Def.Address = Device.IPOrHostName;
+					Def.Name = Device.Name;
+					Def.Platform = DeviceMap.FirstOrDefault(Entry => Entry.Value == Device.Type.Replace("-DevKit", "", StringComparison.OrdinalIgnoreCase)).Key;
+					Def.DeviceData = Device.DeviceData;
+					Def.Model = Device.Model;
+
+					EPerfSpec PerfSpec = EPerfSpec.Unspecified;
+					if (!string.IsNullOrEmpty(Device.PerfSpec) && !Enum.TryParse(Device.PerfSpec, true, out PerfSpec))
+					{
+						throw new AutomationException("Unable to convert perfspec '{0}' into an EPerfSpec", Device.PerfSpec);
+					}
+					Def.PerfSpec = PerfSpec;
+
+					ITargetDevice TargetDevice = CreateAndRegisterDeviceFromDefinition(Def, Entry.Key);
+
+					// If a device from service can't be added, fail reservation and cleanup devices
+					if (TargetDevice == null)
+					{
+						// If some devices from reservation have been created, release them which will also dispose of reservation
+						if (ScopeReservedDevices.Count > 0)
 						{
-							ServiceMapping[Service].Add(Constraint);
+							ReleaseDevices(ScopeReservedDevices);
 						}
 						else
 						{
-							ServiceMapping.Add(Service, new List<UnrealDeviceTargetConstraint> { Constraint });
+							// otherwise, no devices have been creation so just cancel this reservation
+							DeviceReservation.Dispose();
 						}
+
+						Log.Info(KnownLogEvents.Gauntlet_DeviceEvent, "Unable to make device registration: device registration failed for {Platform}:{DeviceName}", Def.Platform, Def.Name);
+						return false;
 					}
-				}
-			}
 
-			// Determine if any constraints do not have a service that can supply an appropriate device
-			List<UnrealDeviceTargetConstraint> DevicesWithService = ServiceMapping.Values.SelectMany(Constraint => Constraint).ToList();
-			List<UnrealDeviceTargetConstraint> DevicesLackingService = RequiredDeviceConstraints.Except(DevicesWithService).ToList();
+					ScopeReservedDevices.Add(TargetDevice);
+					ReservedDevices.Add(TargetDevice);
+					UsedPlatforms.Add(TargetDevice.Platform);
 
-			if (DevicesLackingService.Any())
-			{
-				string Message = "No enabled reservation services were capable of supporting the following constraints:\n";
-				Message += string.Join("\n\t", DevicesLackingService);
-				Message += "\nThe following reservation services were available:\n";
-				Message += string.Join("\n\t", ReservationServices.Select(Service => Service.GetType().Name));
-				Message += "\nEnsure your environment is configured to enable any necessary services.";
-				throw new AutomationException(Message);
+					ServiceDeviceInfo[TargetDevice] = Def;
+					ServiceReservations[TargetDevice] = DeviceReservation;
+				}
 			}
 
-			foreach (var ReservationPair in ServiceMapping)
+			if (ScopeReservedDevices.Count == 0)
 			{
-				if (!ReservationPair.Key.ReserveDevicesFromService(ReservationPair.Value))
-				{
-					Log.Info("Failed to reserve all devices from service. See above log for details");
-					return false;
-				}
+				Log.Info(KnownLogEvents.Gauntlet_DeviceEvent, "Unable to reserve all devices from service.");
 			}
 
-			return true;
+			return ScopeReservedDevices.Count > 0;
 		}
 
 		/// <summary>
@@ -844,14 +911,13 @@ namespace Gauntlet
 		/// </summary>
 		public void ReportDeviceError(ITargetDevice Device, string ErrorMessage)
 		{
-			if (TryGetDevicesReservationService(Device, out IDeviceReservationService Service))
-			{
-				Service.ReportDeviceError(Device.Name, ErrorMessage);
-			}
-			else
+			DeviceDefinition Def = null;
+			if (!ServiceDeviceInfo.TryGetValue(Device, out Def))
 			{
-				Log.Verbose("{Device} was not reserved by a service! Ignoring error report.");
+				return;
 			}
+
+			ReportDeviceError(Def.Name, ErrorMessage);
 		}
 
 		/// <summary>
@@ -860,36 +926,54 @@ namespace Gauntlet
 		/// <param name="DeviceList"></param>
 		public void ReleaseDevices(IEnumerable<ITargetDevice> DeviceList)
 		{
-			if (DeviceList == null || !DeviceList.Any())
-			{
-				return;
-			}
-
 			lock (LockObject)
 			{
+				// remove all these devices from our reserved list
+				ReservedDevices = ReservedDevices.Where(D => DeviceList.Contains(D) == false).ToList();
+
+				// disconnect them
 				List<ITargetDevice> KeepList = new List<ITargetDevice>();
 				foreach (ITargetDevice Device in DeviceList)
 				{
-					if (LocalDevices.Contains(Device))
+					DeviceConfigurationCache.Instance.RevertDeviceConfiguration(Device);
+
+					if (Device.IsConnected && InitialConnectState[Device] == false)
 					{
-						DeviceConfigurationCache.Instance.RevertDeviceConfiguration(Device);
-						KeepList.Add(Device);
+						Device.Disconnect();
 					}
-					else if (TryGetDevicesReservationService(Device, out IDeviceReservationService Service))
+
+					// if device is from reservation service, unregister it
+					DeviceReservationAutoRenew Reservation;
+					if (ServiceReservations.TryGetValue(Device, out Reservation))
 					{
-						Service.ReleaseDevices([Device]);
+						bool DisposeReservation = ServiceReservations.Count(Entry => Entry.Value == Reservation) == 1;
+
+						// remove and dispose of device
+						// @todo: add support for reservation modification on server (partial device release)
+						ServiceReservations.Remove(Device);
+						ServiceDeviceInfo.Remove(Device);
+						AvailableDevices.Remove(Device);
+
+						InitialConnectState.Remove(Device);
+						Constraints.Remove(Device);
+
+						Device.Dispose();
+
+						// if last device in reservation, dispose the reservation
+						if (DisposeReservation)
+						{
+							Reservation.Dispose();
+						}
 					}
 					else
 					{
-						Log.Info("Attempted to release an unregistered device. Ensure the device is registered with DevicePool.Instance.RegisterDevice first.");
+						KeepList.Add(Device);
 					}
-				}
 
-				// Remove any provisioned devices
-				AvailableDevices = AvailableDevices.Where(Device => !DeviceList.Contains(Device)).ToList();
-				ClaimedDevices = ClaimedDevices.Where(Device => !DeviceList.Contains(Device)).ToList();
+				}
 
-				// But keep local devices available
+				// Put all devices back into the available list at the front - if another device is requested we want to
+				// try and use one we just put a build on
 				KeepList.ForEach(D => { AvailableDevices.Remove(D); AvailableDevices.Insert(0, D); });
 			}
 		}
@@ -913,6 +997,7 @@ namespace Gauntlet
 
 				Func<ITargetDevice, bool> Validate = (ITargetDevice Device) =>
 				{
+
 					if (!Constraint.Check(Device))
 					{
 						return false;
@@ -947,15 +1032,21 @@ namespace Gauntlet
 			var TooFewTotalDevices = RequiredDevices.Where(KP => TotalDeviceTypes[KP.Key] < RequiredDevices[KP.Key]).Select(KP => KP.Key);
 			var TooFewCurrentDevices = RequiredDevices.Where(KP => AvailableDeviceTypes[KP.Key] < RequiredDevices[KP.Key]).Select(KP => KP.Key);
 
+
 			var Devices = TooFewTotalDevices.Concat(TooFewCurrentDevices);
 
 			// Request devices from the service if we need them
-			if (UseServiceDevices && (TooFewTotalDevices.Count() > 0 || TooFewCurrentDevices.Count() > 0))
+			if (UseServiceDevices && !String.IsNullOrEmpty(DeviceURL) && (TooFewTotalDevices.Count() > 0 || TooFewCurrentDevices.Count() > 0))
 			{
+
 				Dictionary<UnrealDeviceTargetConstraint, int> DeviceCounts = new Dictionary<UnrealDeviceTargetConstraint, int>();
+
 				Devices.ToList().ForEach(Platform => DeviceCounts[Platform] = RequiredDevices[Platform]);
 
-				if (!ReserveDevicesFromService(DeviceCounts))
+				Log.Info(KnownLogEvents.Gauntlet_DeviceEvent, "Requesting devices from service at {URL}", DeviceURL);
+
+				// Acquire necessary devices from service
+				if (!ReserveDevicesFromService(DeviceURL, DeviceCounts))
 				{
 					return false;
 				}
@@ -987,17 +1078,13 @@ namespace Gauntlet
 		/// </summary>
 		/// <param name="Def"></param>
 		/// <returns></returns>
-		public ITargetDevice CreateAndRegisterDeviceFromDefinition(DeviceDefinition Def, UnrealDeviceTargetConstraint Constraint, 
-			IDeviceReservationService Service = null, IDeviceFactory Factory = null)
+		protected ITargetDevice CreateAndRegisterDeviceFromDefinition(DeviceDefinition Def, UnrealDeviceTargetConstraint Constraint)
 		{
 			ITargetDevice NewDevice = null;
 
-			if (Factory == null)
-			{
-				Factory = InterfaceHelpers.FindImplementations<IDeviceFactory>()
+			IDeviceFactory Factory = InterfaceHelpers.FindImplementations<IDeviceFactory>()
 					.Where(F => F.CanSupportPlatform(Def.Platform))
 					.FirstOrDefault();
-			}
 
 			if (Factory == null)
 			{
@@ -1043,10 +1130,7 @@ namespace Gauntlet
 				{
 					if (!string.IsNullOrEmpty(Message))
 					{
-						if (Service != null)
-						{
-							Service.ReportDeviceError(Def.Name, Message);
-						}
+						ReportDeviceError(Def.Name, Message);
 					}
 					Log.Info("\nSkipping device.");
 					return null;
@@ -1065,10 +1149,7 @@ namespace Gauntlet
 				string WarningMessage = $"Failed to create device {Def.Name}. {Ex.Message}\n{Ex.StackTrace}";
 				if (Ex is DeviceException)
 				{
-					if (Service != null)
-					{
-						Service.ReportDeviceError(Def.Name, WarningMessage);
-					}
+					ReportDeviceError(Def.Name, WarningMessage);
 				}
 				else
 				{
@@ -1079,27 +1160,6 @@ namespace Gauntlet
 			return NewDevice;
 		}
 
-		private void InitializeReservationServices()
-		{
-			// In order to avoid constructing an instance of the service to determine if it's enabled on any number of types we have to do a bit of generics reflection...
-			foreach (Type ServiceType in InterfaceHelpers.FindTypes<IDeviceReservationService>(true, true))
-			{
-				MethodInfo Function = typeof(DevicePool).GetMethod("IsServiceEnabled", BindingFlags.Instance | BindingFlags.NonPublic);
-				MethodInfo Generic = Function.MakeGenericMethod(ServiceType);
-				object ReturnValue = Generic.Invoke(this, null);
-				bool IsEnabled = (bool)ReturnValue;
-				if (IsEnabled)
-				{
-					ReservationServices.Add(Activator.CreateInstance(ServiceType) as IDeviceReservationService);
-				}
-			}
-		}
-
-		private bool IsServiceEnabled<T>() where T : IDeviceReservationService
-		{
-			return T.Enabled;
-		}
-
 		/// <summary>
 		/// Verifies if the provided TargetDevice meets requirements such as firmware, login status, settings, etc.
 		/// </summary>
@@ -1160,21 +1220,42 @@ namespace Gauntlet
 			return bSucceeded;
 		}
 
+		/// <summary>
+		/// Reporting device issue using the device name provided by service
+		/// </summary>
+		private void ReportDeviceError(string ServiceDeviceName, string ErrorMessage)
+		{
+			if (MaxDeviceErrorReports == 0)
+			{
+				Log.Info(KnownLogEvents.Gauntlet_DeviceEvent, "Maximum device errors reported to backend, {DeviceName} : {Message} ignored", ServiceDeviceName, ErrorMessage);
+				return;
+			}
+
+			MaxDeviceErrorReports--;
+
+			Reservation.ReportDeviceError(DeviceURL, ServiceDeviceName, ErrorMessage);
+		}
+
 		/// <summary>
 		/// Explicitly release all device reservations
 		/// </summary>
 		private void ReleaseReservations()
 		{
-			foreach (IDeviceReservationService ReservationService in ReservationServices)
+			ServiceDeviceInfo.Clear();
+
+			foreach (ITargetDevice Device in ServiceReservations.Keys)
 			{
-				foreach (ITargetDevice Device in ReservationService.ReservedDevices)
-				{
-					AvailableDevices.Remove(Device);
-					ClaimedDevices.Remove(Device);
-				}
+				AvailableDevices.Remove(Device);
+				ReservedDevices.Remove(Device);
+				Device.Dispose();
+			}
 
-				ReservationService.ReleaseDevices(ReservationService.ReservedDevices);
+			foreach (DeviceReservationAutoRenew Reservation in ServiceReservations.Values.Distinct())
+			{
+				Reservation.Dispose();
 			}
+
+			ServiceReservations.Clear();
 		}
 
 		private void CleanupDevices(IEnumerable<UnrealTargetPlatform?> Platforms)
@@ -1193,22 +1274,6 @@ namespace Gauntlet
 			}
 		}
 
-		private bool TryGetDevicesReservationService(ITargetDevice Device, out IDeviceReservationService OutService)
-		{
-			OutService = null;
-
-			foreach (IDeviceReservationService Service in ReservationServices)
-			{
-				if (Service.ReservedDevices.Contains(Device))
-				{
-					OutService = Service;
-					return true;
-				}
-			}
-
-			return false;
-		}
-
 		/// <summary>
 		/// Construct a path to hold cache files and make sure it's properly cleaned
 		/// </summary>
@@ -1284,33 +1349,5 @@ namespace Gauntlet
 
 			return ClientCache;
 		}
-
-		#region Legacy Implementation
-	
-		[Obsolete("Will be removed in a future release")]
-		public string DeviceURL
-		{
-			get { return Legacy_DeviceURL; }
-			set { Legacy_DeviceURL = value; }
-		}
-
-		/// <summary>
-		/// Device reservation service URL
-		/// </summary>
-		private string Legacy_DeviceURL;
-
-		[Obsolete("Misnomer, use ClaimDevices instead. Will be removed in a future release")]
-		public void ReserveDevices(IEnumerable<ITargetDevice> Devices)
-		{
-			ClaimDevices(Devices);
-		}
-
-		[Obsolete("Will be removed in a future release. Use the single parameter overload instead.")]
-		public bool ReserveDevicesFromService(string DeviceURL, Dictionary<UnrealDeviceTargetConstraint, int> DeviceTypes)
-		{
-			return ReserveDevicesFromService(DeviceTypes);
-		}
-
-		#endregion
 	}
 }
\ No newline at end of file
diff --git a/Engine/Source/Programs/AutomationTool/Gauntlet/Framework/Devices/Gauntlet.DeviceReservationService.cs b/Engine/Source/Programs/AutomationTool/Gauntlet/Framework/Devices/Gauntlet.DeviceReservationService.cs
deleted file mode 100644
diff --git a/Engine/Source/Programs/AutomationTool/Gauntlet/Framework/Devices/Gauntlet.UnrealDeviceReservation.cs b/Engine/Source/Programs/AutomationTool/Gauntlet/Framework/Devices/Gauntlet.UnrealDeviceReservation.cs
--- a/Engine/Source/Programs/AutomationTool/Gauntlet/Framework/Devices/Gauntlet.UnrealDeviceReservation.cs
+++ b/Engine/Source/Programs/AutomationTool/Gauntlet/Framework/Devices/Gauntlet.UnrealDeviceReservation.cs
@@ -12,46 +12,27 @@ namespace Gauntlet
 	public class UnrealDeviceReservation
 	{
 		protected List<ProblemDevice> ProblemDevices { get; private set; } = new List<ProblemDevice>();
-		public List<ITargetDevice> ReservedDevices { get; protected set; } = new List<ITargetDevice>();
+		public List<ITargetDevice> ReservedDevices { get; protected set; }
 
-		public bool TryReserveDevices(Dictionary<UnrealDeviceTargetConstraint, int> RequiredDeviceTypes, int ExpectedNumberOfDevices, bool bAllowPartialReservations = false)
+		public bool TryReserveDevices(Dictionary<UnrealDeviceTargetConstraint, int> RequiredDeviceTypes, int ExpectedNumberOfDevices)
 		{
-			// First, determine if this is a partial reservation. In the event a device is considered a problem device, don't discard all devices.
-			// Instead, only discard the devices that have a problem.
-			if (bAllowPartialReservations && ReservedDevices.Count > 0)
-			{
-				RequiredDeviceTypes = GetPartialReservationTypes(RequiredDeviceTypes);
-				if (RequiredDeviceTypes.Count == 0)
-				{
-					// We've already reserved all the requested devices
-					return true;
-				}
+			List<ITargetDevice> AcquiredDevices = new List<ITargetDevice>();
+			List<ITargetDevice> SkippedDevices = new List<ITargetDevice>();
 
-				int NewExpectedCount = 0;
-				foreach (var Pair in RequiredDeviceTypes)
-				{
-					NewExpectedCount += Pair.Value;
-				}
-				ExpectedNumberOfDevices = NewExpectedCount;
-			}
-			else
-			{
-				ReleaseDevices();
-				ProblemDevices.Clear();
-			}
+			ReleaseDevices();
+
+			ProblemDevices.Clear();
 
-			// Ensure the device pool can support the request. This will reserve devices from services if it needs to
+			// check whether pool can accommodate devices
 			if (!DevicePool.Instance.CheckAvailableDevices(RequiredDeviceTypes, ProblemDevices))
 			{
 				return false;
 			}
 
-			List<ITargetDevice> AcquiredDevices = new List<ITargetDevice>();
-			List<ITargetDevice> SkippedDevices = new List<ITargetDevice>();
+			// nothing acquired yet...
+			AcquiredDevices.Clear();
 
-			// The pool now contains enough devices to support this reservation. 
-			// For each platform, enumerate and select from the available devices as long as they match the predicate.
-			// This will discard any devices that are considered unavailable or marked as a problem
+			// for each platform, enumerate and select from the available devices
 			foreach (var PlatformReqKP in RequiredDeviceTypes)
 			{
 				UnrealDeviceTargetConstraint Constraint = PlatformReqKP.Key;
@@ -94,48 +75,70 @@ namespace Gauntlet
 				});
 			}
 
-			// Release any devices that were skipped
-			DevicePool.Instance.ReleaseDevices(SkippedDevices);
-
-			// Obtain connections to any required devices.
-			EstablishDeviceConnections(AcquiredDevices);
-
-			// Mark any devices we can't connect to as problems and release them.
-			// Could be something grabbed them before us, could be that they are unresponsive in some way
-			List<ITargetDevice> LostDevices = AcquiredDevices.Where(Device => !Device.IsConnected).ToList();
-			if (LostDevices.Count > 0)
+			// If we got enough devices, go to step2 where we provision and try to connect them
+			if (AcquiredDevices.Count == ExpectedNumberOfDevices)
 			{
-				LostDevices.ForEach(Device => MarkProblemDevice(Device, "Agent lost connection to device."));
-				AcquiredDevices = AcquiredDevices.Except(LostDevices).ToList();
-				ReleaseProblemDevices();
-			}
+				// actually acquire them
+				DevicePool.Instance.ReserveDevices(AcquiredDevices);
 
-			// If we failed to acquire every device we needed, but support partial reservations, hold onto any devices we did manage to acquire
-			if (AcquiredDevices.Count < ExpectedNumberOfDevices && AcquiredDevices.Count > 0)
-			{
-				if (bAllowPartialReservations)
+				Log.Info("Selected devices {Devices} for client(s). Prepping", string.Join(", ", AcquiredDevices));
+
+				foreach (ITargetDevice Device in AcquiredDevices)
 				{
-					if (ClaimDevices(AcquiredDevices))
+					if (Device.IsOn == false)
+					{
+						Log.Info("Powering on {0}", Device);
+						Device.PowerOn();
+					}
+					else if (Globals.Params.ParseParam("reboot"))
+					{
+						Log.Info("Rebooting {0}", Device);
+						Device.Reboot();
+						// Force a re-login after a reboot to be sure it is ready.
+						if (Globals.Params.ParseParam("VerifyLogin") && Device is IOnlineServiceLogin DeviceLogin)
+						{
+							Log.Info("Revalidate device login after reboot...");
+							DeviceLogin.VerifyLogin();
+						}
+					}
+
+					if (Device.IsConnected == false)
 					{
-						ReservedDevices.AddRange(AcquiredDevices);
+						Log.Verbose("Connecting to {0}", Device);
+						Device.Connect();
 					}
 				}
-				else
+
+				// Step 3: Verify we actually connected to them
+				var LostDevices = AcquiredDevices.Where(D => !D.IsConnected);
+
+				if (LostDevices.Count() > 0)
 				{
+					// mark these as problems. Could be something grabbed them before us, could be that they are
+					// unresponsive in some way
+					LostDevices.ToList().ForEach(D => MarkProblemDevice(D, "Agent lost connection to device."));
+					AcquiredDevices.ToList().ForEach(D => D.Disconnect());
+
 					DevicePool.Instance.ReleaseDevices(AcquiredDevices);
+					AcquiredDevices.Clear();
 				}
-
-				return false;
 			}
-			
 
-			if(!ClaimDevices(AcquiredDevices))
+			if (AcquiredDevices.Count() != ExpectedNumberOfDevices)
 			{
-				return false;
+				Log.Info("Failed to resolve all devices. Releasing the ones we have ");
+				DevicePool.Instance.ReleaseDevices(AcquiredDevices);
+				ReservedDevices = new List<ITargetDevice>();
 			}
+			else
+			{
+				ReservedDevices = AcquiredDevices;
+			}
+
+			// release devices that were skipped
+			DevicePool.Instance.ReleaseDevices(SkippedDevices);
 
-			ReservedDevices.AddRange(AcquiredDevices);
-			return true;
+			return ReservedDevices.Count() == ExpectedNumberOfDevices;
 		}
 
 		public void ReleaseDevices()
@@ -197,88 +200,5 @@ namespace Gauntlet
 				ProblemDevices.Add(new ProblemDevice(Device.Name, Device.Platform.Value));
 			}
 		}
-
-		/// <summary>
-		/// Converts a set of required devices into only what is necessary by analyzing the list of already reserved devices
-		/// </summary>
-		/// <param name="DeviceTypes">Request list</param>
-		/// <returns>The difference in constraints/values between the request types and the currently reserved devices</returns>
-		private Dictionary<UnrealDeviceTargetConstraint, int> GetPartialReservationTypes(Dictionary<UnrealDeviceTargetConstraint, int> DeviceTypes)
-		{
-			Dictionary<UnrealDeviceTargetConstraint, int> CurrentlyAvailableTypes = new();
-			foreach (ITargetDevice Device in ReservedDevices)
-			{
-				UnrealDeviceTargetConstraint Constraint = DevicePool.Instance.GetConstraint(Device);
-				if (CurrentlyAvailableTypes.ContainsKey(Constraint))
-				{
-					++CurrentlyAvailableTypes[Constraint];
-				}
-				else
-				{
-					CurrentlyAvailableTypes.Add(Constraint, 1);
-				}
-			}
-
-			Dictionary<UnrealDeviceTargetConstraint, int> NewSet = new();
-			foreach (var Pair in DeviceTypes)
-			{
-				if (CurrentlyAvailableTypes.ContainsKey(Pair.Key))
-				{
-					// If we have more devices already reserved than what is requested we don't need to reserve devices
-					if (CurrentlyAvailableTypes[Pair.Key] < Pair.Value)
-					{
-						// We have at least 1 device with this constraint, but not enough to fufill the request. Request only what is needed for this constraint
-						NewSet.Add(Pair.Key, Pair.Value - CurrentlyAvailableTypes[Pair.Key]);
-					}
-				}
-				else
-				{
-					// No existing constraint, we don't have these devices
-					NewSet.Add(Pair.Key, Pair.Value);
-				}
-			}
-
-			return NewSet;
-		}
-
-		private void EstablishDeviceConnections(IEnumerable<ITargetDevice> Devices)
-		{
-			foreach (ITargetDevice Device in Devices)
-			{
-				if (Device.IsOn == false)
-				{
-					Log.Info("Powering on {0}", Device);
-					Device.PowerOn();
-				}
-				else if (Globals.Params.ParseParam("reboot"))
-				{
-					Log.Info("Rebooting {0}", Device);
-					Device.Reboot();
-					// Force a re-login after a reboot to be sure it is ready.
-					if (Globals.Params.ParseParam("VerifyLogin") && Device is IOnlineServiceLogin DeviceLogin)
-					{
-						Log.Info("Revalidate device login after reboot...");
-						DeviceLogin.VerifyLogin();
-					}
-				}
-
-				if (Device.IsConnected == false)
-				{
-					Log.Verbose("Connecting to {0}", Device);
-					Device.Connect();
-				}
-			}
-		}
-
-		private bool ClaimDevices(IEnumerable<ITargetDevice> Devices)
-		{
-			bool bSuccess = DevicePool.Instance.ClaimDevices(Devices);
-			if (!bSuccess)
-			{
-				Log.Warning("Attempted to claim the following devices which are already marked as claimed:\n{Devices}",
-						string.Join("\n\t", Devices.Select(Device => Device.Name)));
-			}
-			return bSuccess;
-		}
 	}
 }
\ No newline at end of file
diff --git a/Engine/Source/Programs/AutomationTool/Gauntlet/Framework/Gauntlet.Utils.cs b/Engine/Source/Programs/AutomationTool/Gauntlet/Framework/Gauntlet.Utils.cs
--- a/Engine/Source/Programs/AutomationTool/Gauntlet/Framework/Gauntlet.Utils.cs
+++ b/Engine/Source/Programs/AutomationTool/Gauntlet/Framework/Gauntlet.Utils.cs
@@ -825,37 +825,21 @@ namespace Gauntlet
 
 		public static class InterfaceHelpers
 		{
-			public static IEnumerable<Type> FindTypes<InterfaceType>(bool bIncludeCompiledScripts = false, bool bConcreteTypesOnly = false)
+
+			public static IEnumerable<InterfaceType> FindImplementations<InterfaceType>(bool bIncludeCompiledScripts = false)
 				where InterfaceType : class
 			{
-				HashSet<Type> AllTypes = new HashSet<Type>(Assembly.GetExecutingAssembly().GetTypes().Where(Type => typeof(InterfaceType).IsAssignableFrom(Type)));
-				if (bConcreteTypesOnly)
-				{
-					AllTypes = AllTypes.Where(Type => !Type.IsAbstract && !Type.IsInterface).ToHashSet();
-				}
+				HashSet<Type> AllTypes = new HashSet<Type>(Assembly.GetExecutingAssembly().GetTypes().Where(T => typeof(InterfaceType).IsAssignableFrom(T)));
 
 				if (bIncludeCompiledScripts)
 				{
 					foreach (Assembly Assembly in ScriptManager.AllScriptAssemblies)
 					{
-						HashSet<Type> AssemblyTypes = new HashSet<Type>(Assembly.GetTypes().Where(T => typeof(InterfaceType).IsAssignableFrom(T)));
-						if (bConcreteTypesOnly)
-						{
-							AssemblyTypes = AssemblyTypes.Where(Type => !Type.IsAbstract && !Type.IsInterface).ToHashSet();
-						}
-
+						List<Type> AssemblyTypes = Assembly.GetTypes().Where(T => typeof(InterfaceType).IsAssignableFrom(T)).ToList();
 						AllTypes.UnionWith(AssemblyTypes);
 					}
 				}
 
-				return AllTypes;
-			}
-
-			public static IEnumerable<InterfaceType> FindImplementations<InterfaceType>(bool bIncludeCompiledScripts = false)
-				where InterfaceType : class
-			{
-				HashSet<Type> AllTypes = FindTypes<InterfaceType>(bIncludeCompiledScripts).ToHashSet();
-
 				List<InterfaceType> ConstructedTypes = new List<InterfaceType>();
 
 				foreach (Type FoundType in AllTypes)
diff --git a/Engine/Source/Programs/AutomationTool/Gauntlet/Unreal/Base/Gauntlet.UnrealSession.cs b/Engine/Source/Programs/AutomationTool/Gauntlet/Unreal/Base/Gauntlet.UnrealSession.cs
--- a/Engine/Source/Programs/AutomationTool/Gauntlet/Unreal/Base/Gauntlet.UnrealSession.cs
+++ b/Engine/Source/Programs/AutomationTool/Gauntlet/Unreal/Base/Gauntlet.UnrealSession.cs
@@ -727,11 +727,9 @@ namespace Gauntlet
 				return true;
 			}
 
-			// Figure out how many of each device we need
+			// figure out how many of each device we need
 			Dictionary<UnrealDeviceTargetConstraint, int> RequiredDeviceTypes = new Dictionary<UnrealDeviceTargetConstraint, int>();
-			IEnumerable<UnrealSessionRole> RolesThatRequireDevice = SessionRoles
-				.Where(Role => !Role.IsNullRole())
-				.Where(RoleNeedsDevice);
+			IEnumerable<UnrealSessionRole> RolesThatRequireDevice = SessionRoles.Where(R => !R.IsNullRole());
 
 			// Get a count of the number of devices required for each platform
 			foreach (UnrealSessionRole Role in RolesThatRequireDevice)
@@ -748,7 +746,7 @@ namespace Gauntlet
 
 			try
 			{
-				return UnrealDeviceReservation.TryReserveDevices(RequiredDeviceTypes, RolesThatRequireDevice.Count(), true);
+				return UnrealDeviceReservation.TryReserveDevices(RequiredDeviceTypes, RolesThatRequireDevice.Count());
 			}
 			catch (Exception ex)
 			{
diff --git a/Engine/Source/Programs/AutomationTool/Gauntlet/Unreal/RunUnreal.cs b/Engine/Source/Programs/AutomationTool/Gauntlet/Unreal/RunUnreal.cs
--- a/Engine/Source/Programs/AutomationTool/Gauntlet/Unreal/RunUnreal.cs
+++ b/Engine/Source/Programs/AutomationTool/Gauntlet/Unreal/RunUnreal.cs
@@ -320,6 +320,9 @@ namespace Gauntlet
 			// Tear down accounts
 			TeardownAccounts();
 
+			// dispose now, not during shutdown gc, because this runs commands...
+			DevicePool.Instance.Dispose();
+
 			// If only running a single unreal test node, we'll direct the summary artifact link to the specific test
 			string ArtifactPath = Globals.LogDir;
 			if(AllTestNodes.Count == 1)
