Add missing ILowLevelTestsBuildFactory and ILowLevelTestsReporting implementations for iOS

diff --git a/Engine/Source/Programs/AutomationTool/LowLevelTests/RunLowLevelTests.cs b/Engine/Source/Programs/AutomationTool/LowLevelTests/RunLowLevelTests.cs
--- a/Engine/Source/Programs/AutomationTool/LowLevelTests/RunLowLevelTests.cs
+++ b/Engine/Source/Programs/AutomationTool/LowLevelTests/RunLowLevelTests.cs
@@ -697,7 +697,15 @@ namespace LowLevelTests
 		IBuild CreateBuild(UnrealTargetPlatform InPlatform, UnrealTargetConfiguration InConfiguration, string InTestApp, string InBuildPath, bool bSkipStage);
 		protected static string GetExecutable(UnrealTargetPlatform InPlatform, UnrealTargetConfiguration InConfiguration, string InTestApp, string InBuildPath, string FileRegEx)
 		{
-			IEnumerable<string> Executables = DirectoryUtils.FindMatchingFiles(InBuildPath, FileRegEx, -1).Select(FileInfo => FileInfo.FullName);
+			IEnumerable<string> Executables = null;
+			if (InPlatform == UnrealTargetPlatform.IOS)
+			{
+				Executables = DirectoryUtils.FindMatchingDirectories(InBuildPath, FileRegEx, -1).Select(FileInfo => FileInfo.FullName);
+			}
+			else
+			{
+				Executables = DirectoryUtils.FindMatchingFiles(InBuildPath, FileRegEx, -1).Select(FileInfo => FileInfo.FullName);
+			}
 			string ParentDirPath;
 			string BuildExecutableName;
 			foreach (string Executable in Executables)
@@ -859,6 +867,43 @@ namespace LowLevelTests
 		}
 	}
 
+	public class IOSLowLevelTestsBuildFactory : ILowLevelTestsBuildFactory
+	{
+		public bool CanSupportPlatform(UnrealTargetPlatform InPlatform)
+		{
+			return InPlatform.IsInGroup(UnrealPlatformGroup.IOS);
+		}
+
+		public IBuild CreateBuild(UnrealTargetPlatform InPlatform, UnrealTargetConfiguration InConfiguration, string InTestApp, string InBuildPath, bool bSkipStage)
+		{
+			string ExecutablePath = ILowLevelTestsBuildFactory.GetExecutable(InPlatform, InConfiguration, InTestApp, InBuildPath, GetExecutableRegex(InPlatform));
+			return new IOSBuild(InConfiguration, InTestApp, Path.Combine(InBuildPath, ExecutablePath), BuildFlags.Packaged | BuildFlags.CanReplaceCommandLine);
+		}
+
+		public string GetExecutableRegex(UnrealTargetPlatform InPlatform)
+		{
+			return @"[A-Za-z0-9_]+(Tests)?(?:-[A-Za-z0-9_]+)?(?:-[A-Za-z0-9_]+)?.app";
+		}
+	}
+
+	public class IOSLowLevelTestsReporting : ILowLevelTestsReporting
+	{
+		public bool CanSupportPlatform(UnrealTargetPlatform InPlatform)
+		{
+			return InPlatform.IsInGroup(UnrealPlatformGroup.IOS);
+		}
+
+		public string GetTargetReportPath(UnrealTargetPlatform InPlatform, string InTestApp, string InBuildPath)
+		{
+			return string.Format("{0}LLTResults.out", InPlatform.ToString());
+		}
+
+		public string CopyDeviceReportTo(IAppInstall InAppInstall, UnrealTargetPlatform InPlatform, string InTestApp, string InBuildPath, string InTargetDirectory)
+		{
+			throw new NotSupportedException("Reports currently not supported for IOS.");
+		}
+	}
+
 	public class DesktopLowLevelTestsReporting : ILowLevelTestsReporting
 	{
 		public bool CanSupportPlatform(UnrealTargetPlatform InPlatform)
