snapshot import does not depend on -p4

#rnx

diff --git a/Engine/Source/Programs/AutomationTool/Scripts/CookCommand.Automation.cs b/Engine/Source/Programs/AutomationTool/Scripts/CookCommand.Automation.cs
--- a/Engine/Source/Programs/AutomationTool/Scripts/CookCommand.Automation.cs
+++ b/Engine/Source/Programs/AutomationTool/Scripts/CookCommand.Automation.cs
@@ -349,6 +349,23 @@ namespace AutomationScripts
 
 		private static void ImportSnapshots(ProjectParams Params, string PlatformsToCook)
 		{
+			string Branch;
+			int Changelist;
+			if (UnrealBuildTool.BuildVersion.TryRead(UnrealBuildTool.BuildVersion.GetDefaultFileName(), out UnrealBuildTool.BuildVersion CurrentBuildVersion) && CurrentBuildVersion.Changelist != 0)
+			{
+				Branch = CurrentBuildVersion.BranchName.Replace('+','/');
+				Changelist = CurrentBuildVersion.Changelist;
+			}
+			else if (P4Enabled)
+			{
+				Branch = P4Env.Branch;
+				Changelist = P4Env.Changelist;
+			}
+			else
+			{
+				throw new AutomationException($"Cannot get current branch or changelist for snapshot import - could not read from {UnrealBuildTool.BuildVersion.GetDefaultFileName()} and -p4 was not specified on the command line");
+			}
+
 			FileReference ZenLaunchExe = new FileReference(Path.Combine(Unreal.EngineDirectory.FullName, "Binaries", HostPlatform.Current.HostEditorPlatform.ToString(), String.Format("ZenLaunch{0}", RuntimePlatform.ExeExtension)));
 			StringBuilder ZenLaunchCommandline = new StringBuilder();
 			ZenLaunchCommandline.AppendFormat("{0} -SponsorProcessID={1}", CommandUtils.MakePathSafeToUseWithCommandLine(Params.RawProjectPath.FullName), Environment.ProcessId);
@@ -392,7 +409,7 @@ namespace AutomationScripts
 			{
 				foreach (string Runtime in Runtimes)
 				{
-					EpicGames.ProjectStore.Build Snapshot = BuildIndex.GetBuild(CommandUtils.P4Env.Branch, Params.ShortProjectName, Platform, Runtime, CommandUtils.P4Env.Changelist, true);
+					EpicGames.ProjectStore.Build Snapshot = BuildIndex.GetBuild(Branch, Params.ShortProjectName, Platform, Runtime, Changelist, true);
 					Snapshot.Import(Params.RawProjectPath, PlatformCookedDirectory, false, false);
 				}
 			}
