Fix access level of TestNodeEvents in UnrealTestNode

#rnx

diff --git a/Engine/Source/Programs/AutomationTool/Gauntlet/Unreal/Automation/UE.Automation.cs b/Engine/Source/Programs/AutomationTool/Gauntlet/Unreal/Automation/UE.Automation.cs
--- a/Engine/Source/Programs/AutomationTool/Gauntlet/Unreal/Automation/UE.Automation.cs
+++ b/Engine/Source/Programs/AutomationTool/Gauntlet/Unreal/Automation/UE.Automation.cs
@@ -686,10 +686,10 @@ namespace UE
 							switch (Entry.Level)
 							{
 								case UnrealLog.LogLevel.Error:
-									TestNodeEvents.Add(new UnrealTestEvent(EventSeverity.Error, Entry.Message, Enumerable.Empty<string>()));
+									AddTestEvent(new UnrealTestEvent(EventSeverity.Error, Entry.Message, Enumerable.Empty<string>()));
 									break;
 								case UnrealLog.LogLevel.Warning:
-									TestNodeEvents.Add(new UnrealTestEvent(EventSeverity.Warning, Entry.Message, Enumerable.Empty<string>()));
+									AddTestEvent(new UnrealTestEvent(EventSeverity.Warning, Entry.Message, Enumerable.Empty<string>()));
 									break;
 							}
 						}
diff --git a/Engine/Source/Programs/AutomationTool/Gauntlet/Unreal/Base/Gauntlet.UnrealTestNode.cs b/Engine/Source/Programs/AutomationTool/Gauntlet/Unreal/Base/Gauntlet.UnrealTestNode.cs
--- a/Engine/Source/Programs/AutomationTool/Gauntlet/Unreal/Base/Gauntlet.UnrealTestNode.cs
+++ b/Engine/Source/Programs/AutomationTool/Gauntlet/Unreal/Base/Gauntlet.UnrealTestNode.cs
@@ -130,7 +130,7 @@ namespace Gauntlet
 		public virtual void ReportError(string Message, params object[] Args)
 		{
 			Message = string.Format(Message, Args);
-			TestNodeEvents.Add(new UnrealTestEvent(EventSeverity.Error, Message, Enumerable.Empty<string>()));
+			AddTestEvent(new UnrealTestEvent(EventSeverity.Error, Message, Enumerable.Empty<string>()));
 			if (!LogWarningsAndErrorsAfterSummary) { Log.Error(KnownLogEvents.Gauntlet_TestEvent, Message); }
 			if (GetTestStatus() == TestStatus.Complete && GetTestResult() == TestResult.Passed)
 			{
@@ -146,7 +146,7 @@ namespace Gauntlet
 		public virtual void ReportWarning(string Message, params object[] Args)
 		{
 			Message = string.Format(Message, Args);
-			TestNodeEvents.Add(new UnrealTestEvent(EventSeverity.Warning, Message, Enumerable.Empty<string>()));
+			AddTestEvent(new UnrealTestEvent(EventSeverity.Warning, Message, Enumerable.Empty<string>()));
 			if (!LogWarningsAndErrorsAfterSummary) { Log.Warning(KnownLogEvents.Gauntlet_TestEvent, Message); }
 		}
 
@@ -230,11 +230,12 @@ namespace Gauntlet
 		/// <summary>
 		/// Collection of events thrown by gauntlet itself during the test run.
 		/// </summary>
-		protected List<UnrealTestEvent> TestNodeEvents { get; private set; } = new List<UnrealTestEvent>();
+		protected IReadOnlyList<UnrealTestEvent> TestNodeEvents => _TestNodeEvents;
+		private List<UnrealTestEvent> _TestNodeEvents { get; set; } = new List<UnrealTestEvent>();
 
 		public override void AddTestEvent(UnrealTestEvent InEvent)
 		{
-			TestNodeEvents.Add(InEvent);
+			_TestNodeEvents.Add(InEvent);
 		}
 
 		/// <summary>
@@ -1167,27 +1168,27 @@ namespace Gauntlet
 				}
 				if (DistinctResults.Contains(UnrealProcessResult.InitializationFailure))
 				{
-					TestNodeEvents.Add(new UnrealTestEvent(EventSeverity.Fatal, "Engine Initialization Failed", new List<string> { "Engine failed to initialize in one or more roles. This is likely a bad build." }));
+					AddTestEvent(new UnrealTestEvent(EventSeverity.Fatal, "Engine Initialization Failed", new List<string> { "Engine failed to initialize in one or more roles. This is likely a bad build." }));
 				}
 				if (DistinctResults.Contains(UnrealProcessResult.LoginFailed))
 				{
-					TestNodeEvents.Add(new UnrealTestEvent(EventSeverity.Fatal, "Login Unsuccessful", new List<string> { "User account never successfully finished logging in." }));
+					AddTestEvent(new UnrealTestEvent(EventSeverity.Fatal, "Login Unsuccessful", new List<string> { "User account never successfully finished logging in." }));
 				}
 				if (DistinctResults.Contains(UnrealProcessResult.EncounteredFatalError))
 				{
-					TestNodeEvents.Add(new UnrealTestEvent(EventSeverity.Fatal, "Fatal Error Encountered", new List<string> { "Test encountered a fatal error. Check ClientOutput.log and ServerOutput.log for details." }));
+					AddTestEvent(new UnrealTestEvent(EventSeverity.Fatal, "Fatal Error Encountered", new List<string> { "Test encountered a fatal error. Check ClientOutput.log and ServerOutput.log for details." }));
 				}
 				if (DistinctResults.Contains(UnrealProcessResult.EncounteredEnsure))
 				{
-					TestNodeEvents.Add(new UnrealTestEvent(EventSeverity.Error, "Ensure Found", new List<string> { "Test encountered one or more ensures. See above for details." }));
+					AddTestEvent(new UnrealTestEvent(EventSeverity.Error, "Ensure Found", new List<string> { "Test encountered one or more ensures. See above for details." }));
 				}
 				if (DistinctResults.Contains(UnrealProcessResult.TestFailure))
 				{
-					TestNodeEvents.Add(new UnrealTestEvent(EventSeverity.Error, "Test Failure Encountered", new List<string> { "Test encountered a failure. See above for more info." }));
+					AddTestEvent(new UnrealTestEvent(EventSeverity.Error, "Test Failure Encountered", new List<string> { "Test encountered a failure. See above for more info." }));
 				}
 				if (DistinctResults.Contains(UnrealProcessResult.TimeOut))
 				{
-					TestNodeEvents.Add(new UnrealTestEvent(EventSeverity.Error, "Test Timed Out", new List<string> { "Test terminated due to timeout. Check ClientOutput.log and ServerOutput.log for more info." }));
+					AddTestEvent(new UnrealTestEvent(EventSeverity.Error, "Test Timed Out", new List<string> { "Test terminated due to timeout. Check ClientOutput.log and ServerOutput.log for more info." }));
 				}
 			}
 		}
