Cloud Builds - Added metadata for cloud artifacts that points back to the Horde url for the build produced.
Also removed support for seperating metadata with : as that causes a lot of confusion when values can also contain :
Now we simply seperate on = as before but only the first = is assumed to be the seperator with the rest being part of the value.
Removed the explicit validate part call and used the new -validate option instead - same thing but faster. Still temporary to help us detect issues in the uploads.
Lastly we know replace . with - in any of the components used to build the bucket name as the . is used to seperate components and thus is reserved.

diff --git a/Engine/Build/Graph/Tasks/BuildAndTestProject.xml b/Engine/Build/Graph/Tasks/BuildAndTestProject.xml
--- a/Engine/Build/Graph/Tasks/BuildAndTestProject.xml
+++ b/Engine/Build/Graph/Tasks/BuildAndTestProject.xml
@@ -1,5 +1,6 @@
 <?xml version='1.0' ?>
 <BuildGraph xmlns="http://www.epicgames.com/BuildGraph" xmlns:xsi="http://www.w3.org/2001/XMLSchema-instance" xsi:schemaLocation="http://www.epicgames.com/BuildGraph ../Schema.xsd" >
+	<EnvVar Name="UE_HORDE_URL"/>
 	<EnvVar Name="UE_HORDE_JOBID"/>
 	<EnvVar Name="UE_HORDE_STEPID"/>
 	<EnvVar Name="UE_HORDE_TEMPLATEID"/>
@@ -761,7 +762,11 @@
 				<Node Name="$(CloudPlatformStagedArtifactNodeName)" Requires="$(PlatformStageNodeName)">
 					<ForEach Name="PackageConfiguration" Values="$(PackageConfigurations)" Separator="+">
 						<Tag Files="$(ProjectOutputDirectory)/$(StagedPlatformFolder)/..." With="#$(PlatformStageNodeName) Files"/>
-						<CreateCloudArtifact Name="$(BuildNamePath)" Type="staged-build" BaseDir="$(ProjectOutputDirectory)/$(StagedPlatformFolder)" Metadata="Configuration=$(PackageConfiguration);BuildName=$(BuildNamePath);dashboard-category=$(PackageConfiguration) Builds" Project="$(TargetProject)" Platform="$(StagedPlatformFolder)" HttpVersion="$(UE-CloudPublishHttpVersion)" Files="#$(PlatformStageNodeName) Files"/>
+
+						<Property Name="BuildUrl" Value=""/>
+						<Property Name="BuildUrl" Value="$(UE_HORDE_URL)job/$(UE_HORDE_JOBID)?step=$(UE_HORDE_STEPID)" If="$(Horde)"/>
+
+						<CreateCloudArtifact Name="$(BuildNamePath)" Type="staged-build" BaseDir="$(ProjectOutputDirectory)/$(StagedPlatformFolder)" Metadata="Configuration=$(PackageConfiguration);BuildName=$(BuildNamePath);dashboard-category=$(PackageConfiguration) Builds;BuildUrl=$(BuildUrl)" Project="$(TargetProject)" Platform="$(StagedPlatformFolder)" HttpVersion="$(UE-CloudPublishHttpVersion)" Files="#$(PlatformStageNodeName) Files"/>
 					</ForEach>
 				</Node>
 			</Do>
@@ -817,7 +822,10 @@
 						<Command Name="BuildCookRun" Arguments="$(BCRArgs) -configuration=$(PackageConfiguration) $(StageAndPackageArguments) -skipbuild -skipcook -skippak $(CloudPackageSkipBCRArguments) -archive -archivedirectory=&quot;$(ProjectOutputDirectory)/Packages/$(StagedPlatformFolder)/$(PackageConfiguration)&quot;" />
 						<Tag Files="$(ProjectOutputDirectory)/Packages/$(StagedPlatformFolder)/$(PackageConfiguration)/..." With="#$(PlatformPackageNodeName) $(PackageConfiguration) Files"/>
 
-						<CreateCloudArtifact Name="$(BuildNamePath)" Type="packaged-build" BaseDir="$(ProjectOutputDirectory)/Packages/$(StagedPlatformFolder)/$(PackageConfiguration)" Metadata="Configuration=$(PackageConfiguration);BuildName=$(BuildNamePath);dashboard-category=$(PackageConfiguration) Builds" Project="$(TargetProject)" Platform="$(StagedPlatformFolder)" HttpVersion="$(UE-CloudPublishHttpVersion)" Files="#$(PlatformPackageNodeName) $(PackageConfiguration) Files"/>
+						<Property Name="BuildUrl" Value=""/>
+						<Property Name="BuildUrl" Value="$(UE_HORDE_URL)job/$(UE_HORDE_JOBID)?step=$(UE_HORDE_STEPID)" If="$(Horde)"/>
+
+						<CreateCloudArtifact Name="$(BuildNamePath)" Type="packaged-build" BaseDir="$(ProjectOutputDirectory)/Packages/$(StagedPlatformFolder)/$(PackageConfiguration)" Metadata="Configuration=$(PackageConfiguration);dashboard-category=$(PackageConfiguration) Builds;BuildUrl=$(BuildUrl)" Project="$(TargetProject)" Platform="$(StagedPlatformFolder)" HttpVersion="$(UE-CloudPublishHttpVersion)" Files="#$(PlatformPackageNodeName) $(PackageConfiguration) Files"/>
 					</ForEach>
 				</Node>
 			</Do>
diff --git a/Engine/Source/Programs/AutomationTool/BuildGraph/Tasks/CreateCloudArtifactTask.cs b/Engine/Source/Programs/AutomationTool/BuildGraph/Tasks/CreateCloudArtifactTask.cs
--- a/Engine/Source/Programs/AutomationTool/BuildGraph/Tasks/CreateCloudArtifactTask.cs
+++ b/Engine/Source/Programs/AutomationTool/BuildGraph/Tasks/CreateCloudArtifactTask.cs
@@ -104,7 +104,7 @@ namespace AutomationTool.Tasks
 		public string Files { get; set; } = null!;
 
 		/// <summary>
-		/// Other metadata for the artifact, separated by semicolons.
+		/// Other metadata for the artifact, separated by semicolons. Use key=value for each field.
 		/// </summary>
 		[TaskParameter(Optional = true)]
 		public string? Metadata { get; set; }
@@ -174,28 +174,23 @@ namespace AutomationTool.Tasks
 			List<string> metadataParam = (_parameters.Metadata ?? String.Empty).Split(';', StringSplitOptions.RemoveEmptyEntries | StringSplitOptions.TrimEntries).ToList();
 
 			Dictionary<string, object> metadata = new Dictionary<string, object>();
-			string project = _parameters.Project.GetFileNameWithoutExtension();
-			string platform = _parameters.Platform;
-
+			string project = SanitizeBucketValue(_parameters.Project.GetFileNameWithoutExtension());
+			string platform = SanitizeBucketValue(_parameters.Platform);
 			ProjectProperties properties = ProjectUtils.GetProjectProperties(_parameters.Project);
 			ConfigHierarchy config = properties.EngineConfigs[HostPlatform.Current.HostEditorPlatform];
 			bool foundConfig = config.TryGetValueGeneric("StorageServers", "Cloud", out CloudConfiguration cloudConfig);
 
 			foreach (string line in metadataParam)
 			{
-				string[] s = line.Split(":");
-				if (s.Length > 1 && !String.IsNullOrEmpty(s[0]))
-				{
-					metadata.Add(s[0], s[1]);
-					continue;
-				}
-				s = line.Split("=");
-				if (s.Length > 1 && !String.IsNullOrEmpty(s[0]))
+				int sep = line.IndexOf('=', StringComparison.OrdinalIgnoreCase);
+				if (sep != -1)
 				{
-					metadata.Add(s[0], s[1]);
+					string k = line[..sep];
+					string v = line[(sep + 1)..];
+					metadata.Add(k, v);
 					continue;
 				}
-				Logger.LogWarning("Unable to convert metadata line: \'{Line}\' to key-value. Please use : or = to separate key and value. Ignoring.", line);
+				Logger.LogWarning("Unable to convert metadata line: \'{Line}\' to key-value. Please use = to separate key and value. Ignoring.", line);
 			}
 
 			metadata.Add("type", type);
@@ -218,26 +213,27 @@ namespace AutomationTool.Tasks
 			}
 
 			// Figure out the current change and stream id
-			StreamId streamId;
+			StreamId streamIdTyped;
 			if (!String.IsNullOrEmpty(_parameters.StreamId))
 			{
-				streamId = new StreamId(_parameters.StreamId);
+				streamIdTyped = new StreamId(_parameters.StreamId);
 			}
 			else
 			{
 				string? streamIdEnvVar = Environment.GetEnvironmentVariable("UE_HORDE_STREAMID");
 				if (!String.IsNullOrEmpty(streamIdEnvVar))
 				{
-					streamId = new StreamId(streamIdEnvVar);
+					streamIdTyped = new StreamId(streamIdEnvVar);
 				}
 				else
 				{
 					throw new AutomationException("Missing UE_HORDE_STREAMID environment variable; unable to determine current stream.");
 				}
 			}
+			string streamId = SanitizeBucketValue(streamIdTyped.ToString());
 			metadata.Add("stream", streamId);
 			
-			string baselineBranch = _parameters.BaselineBranch;
+			string baselineBranch =_parameters.BaselineBranch;
 			if (foundConfig && String.IsNullOrEmpty(baselineBranch))
 			{
 				// no explicit baseline branch set and we have found a storage server config, using that value
@@ -249,6 +245,8 @@ namespace AutomationTool.Tasks
 				baselineBranch = streamId.ToString();
 			}
 
+			baselineBranch = SanitizeBucketValue(baselineBranch);
+
 			CommitId commitId;
 			if (!String.IsNullOrEmpty(_parameters.Commit))
 			{
@@ -397,8 +395,6 @@ namespace AutomationTool.Tasks
 					Platform = platform
 				}, allowMultipart: _parameters.AllowMultipart, assumeHttp2: String.Equals(_parameters.HttpVersion, "http2-only", StringComparison.OrdinalIgnoreCase), forceClean: _parameters.Force);
 
-				BuildUploader.ValidatePart(host, ns, bucketId, buildId, partId, accessToken, assumeHttp2: String.Equals(_parameters.HttpVersion, "http2-only", StringComparison.OrdinalIgnoreCase));
-
 				return;
 			}
 
@@ -497,6 +493,12 @@ namespace AutomationTool.Tasks
 			Logger.LogInformation("Uploaded artifact {ArtifactId} in {Time:n1}s ({Rate:n1}mb/s, {RateMbps:n1}mbps)", buildId, timer.Elapsed.TotalSeconds, totalSizeUploaded / (uploadTimeSecs * 1024.0 * 1024.0), (totalSizeUploaded * 8.0) / (uploadTimeSecs * 1024.0 * 1024.0));
 		}
 
+		internal static string SanitizeBucketValue(string s)
+		{
+			// . have special meaning in the bucket names, its used as a separator between components, as such component can not themselves contain .
+			return s.Replace(".", "-", StringComparison.OrdinalIgnoreCase);
+		}
+
 		static HttpClient BuildHttpClient(string host, string ns, string bucketId, CbObjectId buildId, string accessToken)
 		{
 			HttpClient httpClient = new HttpClient();
@@ -793,7 +795,7 @@ namespace AutomationTool.Tasks
 				string metadataString = File.ReadAllText(metadataJsonPath);
 				string http2Options = assumeHttp2 ? "--assume-http2" : "";
 				// pass the access token via environment variable to avoid showing it on the cli
-				string cmdline = $"builds upload --url {host} --namespace {namespaceId} --bucket {bucketId} --build-id {buildId} --build-part-name {partName} --build-part-id {partId} --metadata-path {metadataJsonPath} --access-token-env UE-CloudDataCacheAccessToken --create-build --plain-progress --allow-multipart={allowMultipart} {http2Options} --clean={forceClean} --verbose {baseDir}";
+				string cmdline = $"builds upload --url {host} --namespace {namespaceId} --bucket {bucketId} --build-id {buildId} --build-part-name {partName} --build-part-id {partId} --metadata-path {metadataJsonPath} --access-token-env UE-CloudDataCacheAccessToken --create-build --plain-progress --allow-multipart={allowMultipart} {http2Options} --clean={forceClean} --verify --verbose {baseDir}";
 				Log.Logger.LogInformation("Using metadata '{Metadata}'", metadataString);
 				CommandUtils.RunAndLog(CommandUtils.CmdEnv, zenExe.FullName,cmdline , Options: CommandUtils.ERunOptions.Default, EnvVars: new Dictionary<string, string> {{"UE-CloudDataCacheAccessToken", accessToken}});
 			}
@@ -802,15 +804,6 @@ namespace AutomationTool.Tasks
 				File.Delete(metadataJsonPath);
 			}
 		}
-
-		public static void ValidatePart(string host, string namespaceId, string bucketId, CbObjectId buildId, CbObjectId partId, string accessToken, bool assumeHttp2)
-		{
-			FileInfo zenExe = new FileInfo("Engine/Binaries/Win64/zen.exe");
-			string http2Options = assumeHttp2 ? "--assume-http2" : "";
-			// pass the access token via environment variable to avoid showing it on the cli
-			string cmdline = $"builds validate-part --url {host} --namespace {namespaceId} --bucket {bucketId} --build-id {buildId} --build-part-id {partId} --access-token-env UE-CloudDataCacheAccessToken --plain-progress {http2Options} --verbose";
-			CommandUtils.RunAndLog(CommandUtils.CmdEnv, zenExe.FullName,cmdline , Options: CommandUtils.ERunOptions.Default, EnvVars: new Dictionary<string, string> {{"UE-CloudDataCacheAccessToken", accessToken}});
-		}
 	}
 	namespace CloudDDC
 	{
diff --git a/Engine/Source/Programs/AutomationTool/BuildGraph/Tasks/RetrieveCloudArtifactTask.cs b/Engine/Source/Programs/AutomationTool/BuildGraph/Tasks/RetrieveCloudArtifactTask.cs
--- a/Engine/Source/Programs/AutomationTool/BuildGraph/Tasks/RetrieveCloudArtifactTask.cs
+++ b/Engine/Source/Programs/AutomationTool/BuildGraph/Tasks/RetrieveCloudArtifactTask.cs
@@ -340,7 +340,7 @@ namespace AutomationTool.Tasks
 
 			CbObject queryObject = queryWriter.ToObject();
 			DirectoryReference outputDir = ResolveDirectory(_parameters.OutputDir);
-			string bucketId = $"{project}.{type}.{streamId}.{platform}";
+			string bucketId = $"{CreateCloudArtifactTask.SanitizeBucketValue(project)}.{CreateCloudArtifactTask.SanitizeBucketValue(type.ToString())}.{CreateCloudArtifactTask.SanitizeBucketValue(streamId.ToString())}.{CreateCloudArtifactTask.SanitizeBucketValue(platform)}";
 			using HttpClient httpClient = BuildHttpClient(host, ns, bucketId, accessToken);
 
 			List<CbObjectId> artifactIds = await SearchArtifactsAsync(httpClient, queryObject);
