Adding Horde artifact registration to cloud artifact task

diff --git a/Engine/Source/Programs/AutomationTool/BuildGraph/Tasks/CreateCloudArtifactTask.cs b/Engine/Source/Programs/AutomationTool/BuildGraph/Tasks/CreateCloudArtifactTask.cs
--- a/Engine/Source/Programs/AutomationTool/BuildGraph/Tasks/CreateCloudArtifactTask.cs
+++ b/Engine/Source/Programs/AutomationTool/BuildGraph/Tasks/CreateCloudArtifactTask.cs
@@ -14,11 +14,14 @@ using System.Xml;
 using AutomationTool.Tasks.CloudDDC;
 using EpicGames.Compression;
 using EpicGames.Core;
+using EpicGames.Horde;
 using EpicGames.Horde.Artifacts;
 using EpicGames.Horde.Commits;
 using EpicGames.Horde.Storage;
+using EpicGames.Horde.Storage.Nodes;
 using EpicGames.Horde.Streams;
 using EpicGames.Serialization;
+using Microsoft.Extensions.DependencyInjection;
 using Microsoft.Extensions.Logging;
 using UnrealBuildTool;
 using IoHash = EpicGames.Core.IoHash;
@@ -247,7 +250,8 @@ namespace AutomationTool.Tasks
 
 			baselineBranch = SanitizeBucketValue(baselineBranch);
 
-			CommitId commitId;
+			CommitId commitId = CommitId.FromPerforceChange(0);
+
 			if (!String.IsNullOrEmpty(_parameters.Commit))
 			{
 				commitId = new CommitId(_parameters.Commit);
@@ -377,13 +381,6 @@ namespace AutomationTool.Tasks
 
 			Uri buildUri = new Uri($"{host}/api/v2/builds/{ns}/{bucketId}/{buildId}");
 
-			// we create the artifact descriptor that Horde can then pickup
-			// TODO: upload this to horde
-			ArtifactDescriptor _ = new ArtifactDescriptor
-			{
-				Backend = new ArtifactBackend {Type = ArtifactBackend.BackendTypeEnum.Zen}, Name = name.ToString(), BaseUrl = buildUri
-			};
-
 			if (_parameters.UseZenCli)
 			{
 				BuildUploader.UploadBuild(host, baseDir, ns, bucketId, buildId, partName, partId, metadata, accessToken, new BuildUploader.BuildInfo
@@ -395,6 +392,15 @@ namespace AutomationTool.Tasks
 					Platform = platform
 				}, allowMultipart: _parameters.AllowMultipart, assumeHttp2: String.Equals(_parameters.HttpVersion, "http2-only", StringComparison.OrdinalIgnoreCase), forceClean: _parameters.Force);
 
+				try
+				{
+					await RegisterHordeArtifactAsync(buildUri, name, type, streamIdTyped, commitId, metadata);
+				}
+				catch (Exception ex)
+				{
+					Logger.LogWarning(ex, "Exception registering Horde artifact: {Message}", ex.Message);
+				}
+
 				return;
 			}
 
@@ -491,6 +497,15 @@ namespace AutomationTool.Tasks
 
 			double uploadTimeSecs = timer.Elapsed.TotalSeconds;
 			Logger.LogInformation("Uploaded artifact {ArtifactId} in {Time:n1}s ({Rate:n1}mb/s, {RateMbps:n1}mbps)", buildId, timer.Elapsed.TotalSeconds, totalSizeUploaded / (uploadTimeSecs * 1024.0 * 1024.0), (totalSizeUploaded * 8.0) / (uploadTimeSecs * 1024.0 * 1024.0));
+
+			try
+			{
+				await RegisterHordeArtifactAsync(buildUri, name, type, streamIdTyped, commitId, metadata);
+			}
+			catch (Exception ex)
+			{
+				Logger.LogWarning(ex, "Exception registering Horde artifact: {Message}", ex.Message);
+			}
 		}
 
 		internal static string SanitizeBucketValue(string s)
@@ -728,6 +743,52 @@ namespace AutomationTool.Tasks
 			return contentLength;
 		}
 
+		async Task RegisterHordeArtifactAsync(Uri buildUri, ArtifactName name, ArtifactType type, StreamId streamId, CommitId commitId, Dictionary<string, object> metadata)
+		{		
+			// Add keys used to collate and find builds
+			List<string> keys = new List<string>();			
+
+			// Add keys for the job that's executing
+			string? jobId = Environment.GetEnvironmentVariable("UE_HORDE_JOBID");
+			if (!String.IsNullOrEmpty(jobId))
+			{
+				keys.Add($"job:{jobId}");
+
+				string? stepId = Environment.GetEnvironmentVariable("UE_HORDE_STEPID");
+				if (!String.IsNullOrEmpty(stepId))
+				{
+					keys.Add($"job:{jobId}/step:{stepId}");
+				}
+			}
+
+			List<string> metadataList = new List<string>();
+
+			metadataList.Add("Backend=Zen");
+			metadataList.Add($"ZenBuildUri={buildUri.ToString()}");
+
+			foreach (KeyValuePair<string, object> pair in metadata)
+			{
+				metadataList.Add($"{pair.Key}={pair.Value.ToString()}");
+			}			
+		
+			// TODO: Investigate registering in a zen artifact namespace 
+			IHordeClient hordeClient = CommandUtils.ServiceProvider.GetRequiredService<IHordeClient>();
+			IArtifactBuilder artifact = await hordeClient.Artifacts.CreateAsync(name, type, _parameters.Description, streamId, commitId, keys, metadataList);
+
+			Logger.LogInformation("Registering cloud artifact with Horde, artifact {ArtifactId} '{ArtifactName}' ({ArtifactType}). Namespace: {NamespaceId}, ref {RefName}", artifact.Id, name, type, artifact.NamespaceId, artifact.RefName);
+
+			// TODO: Add manifest and support partial downloads
+			IHashedBlobRef<DirectoryNode> rootRef;
+			await using (IBlobWriter blobWriter = artifact.CreateBlobWriter())
+			{
+				DirectoryReference baseDir = ResolveDirectory("");
+				List<FileReference> files = new List<FileReference>();
+				rootRef = await blobWriter.AddFilesAsync(baseDir, files);
+			}
+
+			await artifact.CompleteAsync(rootRef);
+		}	
+
 		private const long UseTempFileAtBytes = Int32.MaxValue; // use a backing file when the uncompressed file is too large to fit into memory
 
 		/// <inheritdoc/>
