Horde: Use empty commit id instead of CL 0

diff --git a/Engine/Source/Programs/AutomationTool/BuildGraph/Tasks/CreateCloudArtifactTask.cs b/Engine/Source/Programs/AutomationTool/BuildGraph/Tasks/CreateCloudArtifactTask.cs
--- a/Engine/Source/Programs/AutomationTool/BuildGraph/Tasks/CreateCloudArtifactTask.cs
+++ b/Engine/Source/Programs/AutomationTool/BuildGraph/Tasks/CreateCloudArtifactTask.cs
@@ -250,7 +250,7 @@ namespace AutomationTool.Tasks
 
 			baselineBranch = SanitizeBucketValue(baselineBranch);
 
-			CommitId commitId = CommitId.FromPerforceChange(0);
+			CommitId commitId = CommitId.Empty;
 
 			if (!String.IsNullOrEmpty(_parameters.Commit))
 			{
