P4Utils:
Added a perforce utility function that reverts a list of files provided.

diff --git a/Engine/Source/Programs/AutomationTool/AutomationUtils/P4Utils.cs b/Engine/Source/Programs/AutomationTool/AutomationUtils/P4Utils.cs
--- a/Engine/Source/Programs/AutomationTool/AutomationUtils/P4Utils.cs
+++ b/Engine/Source/Programs/AutomationTool/AutomationUtils/P4Utils.cs
@@ -2366,6 +2366,17 @@ namespace AutomationTool
 			LogP4("", "revert " + String.Format("-c {0} ", CL) + CommandLine, AllowSpew: AllowSpew);
 		}
 
+		/// <summary>
+		/// Invokes p4 revert command with a list of files.
+		/// </summary>
+		/// <param name="CL">Changelist where the checked out files should be reverted.</param>
+		/// <param name="Files">List of files to be reverted.</param>
+		/// <param name="AllowSpew"></param>
+		public void Revert(int CL, List<string> Files, bool AllowSpew = true)
+		{
+			BatchedCommand($"revert -c {CL}", Files, AllowSpew: AllowSpew);
+		}
+
 		/// <summary>
 		/// Reverts all unchanged file from the specified changelist.
 		/// </summary>
