Add an AIBotSoak test you can install, run one match, uninstall and do it over x time.
Added optional parameter TestIterations and ReInstallPerPass on AIBotSoak commands.
Adjusted AIBotSoak to be able to run more than one iteration successfully.

diff --git a/Engine/Source/Programs/AutomationTool/Gauntlet/Framework/RpcFramework/Gauntlet.RpcExecutor.cs b/Engine/Source/Programs/AutomationTool/Gauntlet/Framework/RpcFramework/Gauntlet.RpcExecutor.cs
--- a/Engine/Source/Programs/AutomationTool/Gauntlet/Framework/RpcFramework/Gauntlet.RpcExecutor.cs
+++ b/Engine/Source/Programs/AutomationTool/Gauntlet/Framework/RpcFramework/Gauntlet.RpcExecutor.cs
@@ -12,6 +12,7 @@ using System.Net.Sockets;
 using System.Text;
 using System.Text.RegularExpressions;
 using System.Threading;
+using EpicGames.Core;
 using ThreadState = System.Threading.ThreadState;
 
 namespace Gauntlet
@@ -189,7 +190,7 @@ namespace Gauntlet
 			{
 				RequestListener.Prefixes.Add(String.Format("http://{0}/", GetListenAddressAndPort()));
 				Log.Info("Base URL is: " + String.Format("http://{0}/", GetListenAddressAndPort()));
-				RegisterHttpRoute(new RouteMapping("/sendmessage"), HandleIncomingMessage);
+				RegisterHttpRoute(new RouteMapping("/sendmessage"), HandleIncomingMessage, true);
 				RequestListener.Start();
 			}
 			catch (Exception ex)
@@ -682,10 +683,22 @@ namespace Gauntlet
 				HttpListenThread = new Thread(ListenLoopThread);
 				HttpListenThread.Start();
 			}
+			else
+			{
+				Gauntlet.Log.Warning(KnownLogEvents.Gauntlet, "RPC Executor tried to start a new listening thread but there's already an active one in the \"{RPCListeningThreadState}\" state", HttpListenThread.ThreadState);
+				if (HttpListenThread.ThreadState != ThreadState.Running && !RequestListener.IsListening)
+				{
+					Gauntlet.Log.Info("Trying to force start RPC listen thread");
+					bHasReceivedTerminationCall = false;
+					HttpListenThread = new Thread(ListenLoopThread);
+					HttpListenThread.Start();
+				}
+			}
 		}
 
 		public void ShutdownListenThread()
 		{
+			Gauntlet.Log.Verbose("RPC Executor shutting down Listen Thread");
 			if (HttpListenThread != null)
 			{
 				bHasReceivedTerminationCall = true;
diff --git a/Engine/Source/Programs/AutomationTool/Gauntlet/Framework/RpcFramework/Gauntlet.RpcTargetRegistry.cs b/Engine/Source/Programs/AutomationTool/Gauntlet/Framework/RpcFramework/Gauntlet.RpcTargetRegistry.cs
--- a/Engine/Source/Programs/AutomationTool/Gauntlet/Framework/RpcFramework/Gauntlet.RpcTargetRegistry.cs
+++ b/Engine/Source/Programs/AutomationTool/Gauntlet/Framework/RpcFramework/Gauntlet.RpcTargetRegistry.cs
@@ -147,7 +147,7 @@ namespace Gauntlet
 		}
 	}
 
-	public class RpcTargetRegistry
+	public class RpcTargetRegistry : IDisposable
 	{
 		#region RPC Target Accessors
 		/// <summary>
@@ -173,7 +173,6 @@ namespace Gauntlet
 		private int CurrentClientRpcTargetIndex = 0;
 		private int CurrentServerRpcTargetIndex = 0;
 		private int CurrentEditorRpcTargetIndex = 0;
-
 		#endregion
 
 		#region RPC Target Configuration Base Values
@@ -231,7 +230,7 @@ namespace Gauntlet
 
 		~RpcTargetRegistry()
 		{
-			RpcExecutor.Instance.ShutdownListenThread();
+			ReleaseUnmanagedResources();
 		}
 
 		public void AddNewTarget(string InName, string IpAddress = "localhost", int Port = 12321, UnrealTargetRole InType = UnrealTargetRole.Client)
@@ -559,5 +558,17 @@ namespace Gauntlet
 			return ProcessUri;
 		}
 		#endregion //Target Config Helpers
+
+		private void ReleaseUnmanagedResources()
+		{
+			RpcExecutor.Instance?.ShutdownListenThread();
+			RpcTargets.Clear();
+		}
+
+		public void Dispose()
+		{
+			ReleaseUnmanagedResources();
+			GC.SuppressFinalize(this);
+		}
 	}
 }
\ No newline at end of file
