Changing the LaunchMultiServer script to give the output logs meaningful names.

#rnx

diff --git a/Engine/Source/Programs/AutomationTool/Scripts/LaunchMultiServer.Automation.cs b/Engine/Source/Programs/AutomationTool/Scripts/LaunchMultiServer.Automation.cs
--- a/Engine/Source/Programs/AutomationTool/Scripts/LaunchMultiServer.Automation.cs
+++ b/Engine/Source/Programs/AutomationTool/Scripts/LaunchMultiServer.Automation.cs
@@ -330,8 +330,8 @@ namespace AutomationTool
 
 							for (int ClientIndex = 0; ClientIndex < ClientCount; ClientIndex++)
 							{
-								var ProxyClientEndpoint = String.Format("127.0.0.1:{0}", ProxyListeningPort);
-								StartProcessesForServerInfos(UniqueServerInfos, ProcessLaunchType.ProxyClient, ProxyClientEndpoint + " -SaveWinPos="+(ClientIndex + 1));
+								var ProxyClientAdditionalParams = String.Format("127.0.0.1:{0} -log=DstmProxyClient-{1}.log -SaveWinPos={1}", ProxyListeningPort, (ClientIndex+1));
+								StartProcessesForServerInfos(UniqueServerInfos, ProcessLaunchType.ProxyClient, ProxyClientAdditionalParams);
 							}
 						}
 					}
@@ -381,6 +381,10 @@ namespace AutomationTool
 			{
 				CommonArgs += Params.MapToRun + " -server -log ";
 
+				// Here are some examples of adding parameters for debugging purposes (note the trailing space within the quotes)
+				// CommonArgs += "-dpcvars=net.MaxConstructedPartialBunchSizeBytes=131071,dstm.SimpleCoinBots=0 ";
+				// CommonArgs += "-logcmds=\"LogActor Verbose, LogRemoteObject Verbose, LogDstmPushMigration VeryVerbose, LogRepCompares Log, LogNet Verbose\" ";
+
 				if (!ParseParam("nonewconsole"))
 				{
 					CommonArgs += " -newconsole ";
@@ -404,13 +408,14 @@ namespace AutomationTool
 
 						if (LaunchType == ProcessLaunchType.Server)
 						{
-							Args += String.Format("-MultiServerLocalId={0} -MultiServerListenPort={1} -port={2} ", ServerInfo.Definition.LocalId, ServerInfo.Definition.ListenPort, ServerInfo.GameListenPort);
+							Args += String.Format("-MultiServerLocalId={0} -MultiServerListenPort={1} -port={2} -log=DstmMultiServer-{0}.log ", ServerInfo.Definition.LocalId, ServerInfo.Definition.ListenPort, ServerInfo.GameListenPort);
 							Args += String.Format("-ConsoleTitle=\"{0} Server ID {1}\" -SessionName=\"{0} Server ID {1}\" ", Params.ShortProjectName, ServerInfo.Definition.LocalId);
 							Args += Params.ServerCommandline + " ";
 						}
 						else if (LaunchType == ProcessLaunchType.DirectClient)
 						{
 							Args += String.Format("127.0.0.1:{0} ", ServerInfo.GameListenPort) + (Params.Cook ? "" : "-game ");
+							Args += String.Format(" -log=DstmDirectClient-{0}.log", ServerInfo.Definition.LocalId);
 							Args += String.Format(" -SessionName=\"{0} Direct Client | Server ID {1}\" ", Params.ShortProjectName, ServerInfo.Definition.LocalId);
 							Args += Params.ClientCommandline + " ";
 
@@ -459,7 +464,11 @@ namespace AutomationTool
 					if (ServerInfos.Count() > 0)
 					{
 						Args += String.Format("-ConsoleTitle=\"{0} Proxy Server\" -SessionName=\"{0} Proxy Server\" ", Params.ShortProjectName);
-						
+						Args += "-log=DstmProxyServer.log ";
+
+						// Example of adding additional debug parameters for the proxy
+						// Args += "-logcmds=\"LogNetProxy VeryVerbose\" ";
+
 						if (ParseParam("notimeouts"))
 						{
 							Args += "-notimeouts ";
