Replace path separators with forward slash when hashing project name for zen snapshot imports, and use UTF8 instead of ASCII.

[RN] none

diff --git a/Engine/Source/Programs/AutomationTool/BuildGraph/Tasks/ZenImportSnapshotTask.cs b/Engine/Source/Programs/AutomationTool/BuildGraph/Tasks/ZenImportSnapshotTask.cs
--- a/Engine/Source/Programs/AutomationTool/BuildGraph/Tasks/ZenImportSnapshotTask.cs
+++ b/Engine/Source/Programs/AutomationTool/BuildGraph/Tasks/ZenImportSnapshotTask.cs
@@ -358,10 +358,10 @@ namespace AutomationTool.Tasks
 		{
 			using (System.Security.Cryptography.MD5 hasher = System.Security.Cryptography.MD5.Create())
 			{
-				byte[] bytes = System.Text.Encoding.ASCII.GetBytes(_parameters.Project.FullName);
+				byte[] bytes = System.Text.Encoding.UTF8.GetBytes(_parameters.Project.FullName.Replace("\\", "/", StringComparison.InvariantCulture));
 				byte[] hashed = hasher.ComputeHash(bytes);
 
-				string hexString = Convert.ToHexString(hashed);
+				string hexString = Convert.ToHexString(hashed).ToLowerInvariant();
 
 				return _parameters.Project.GetFileNameWithoutAnyExtensions() + "." + hexString.Substring(0, 8);
 			}
diff --git a/Engine/Source/Programs/Shared/EpicGames.ProjectStore/ProjectStoreTypes.cs b/Engine/Source/Programs/Shared/EpicGames.ProjectStore/ProjectStoreTypes.cs
--- a/Engine/Source/Programs/Shared/EpicGames.ProjectStore/ProjectStoreTypes.cs
+++ b/Engine/Source/Programs/Shared/EpicGames.ProjectStore/ProjectStoreTypes.cs
@@ -162,7 +162,7 @@ namespace EpicGames.ProjectStore
 		{
 			using (System.Security.Cryptography.MD5 hasher = System.Security.Cryptography.MD5.Create())
 			{
-				byte[] bytes = System.Text.Encoding.ASCII.GetBytes(projectFile.FullName);
+				byte[] bytes = System.Text.Encoding.UTF8.GetBytes(projectFile.FullName.Replace("\\", "/", StringComparison.InvariantCulture));
 				byte[] hashed = hasher.ComputeHash(bytes);
 
 				string hexString = Convert.ToHexString(hashed).ToLowerInvariant();
