Cloud Artifacts - Support for using different url in uartifact files vs what we connect to to publish cloud artifacts
Also removed the ability to do downloads without zen cli, that C# path is not supported and doesn't work as well as using zen.
Lastly I added a HordeArtifact name that sets the name of the artifact in Horde (which does not use the same conventions as we use for build names) - this should make it easier to find the artifact we want in Horde.

diff --git a/Engine/Build/Graph/Tasks/BuildAndTestProject.xml b/Engine/Build/Graph/Tasks/BuildAndTestProject.xml
--- a/Engine/Build/Graph/Tasks/BuildAndTestProject.xml
+++ b/Engine/Build/Graph/Tasks/BuildAndTestProject.xml
@@ -766,7 +766,7 @@
 						<Property Name="BuildUrl" Value=""/>
 						<Property Name="BuildUrl" Value="$(UE_HORDE_URL)job/$(UE_HORDE_JOBID)?step=$(UE_HORDE_STEPID)" If="$(Horde)"/>
 
-						<CreateCloudArtifact Name="$(BuildNamePath)" Type="staged-build" BaseDir="$(ProjectOutputDirectory)/$(StagedPlatformFolder)" Metadata="Configuration=$(PackageConfiguration);BuildName=$(BuildNamePath);dashboard-category=$(PackageConfiguration) Builds;BuildUrl=$(BuildUrl)" Project="$(TargetProject)" Platform="$(StagedPlatformFolder)" HttpVersion="$(UE-CloudPublishHttpVersion)" Files="#$(PlatformStageNodeName) Files"/>
+						<CreateCloudArtifact Name="$(BuildNamePath)" HordeArtifactName="Cloud-$(ProjectName)-Staged-$(StagedPlatformFolder)-$(PackageConfiguration)" Type="staged-build" BaseDir="$(ProjectOutputDirectory)/$(StagedPlatformFolder)" Metadata="Configuration=$(PackageConfiguration);dashboard-category=Staged Builds - Cloud;BuildUrl=$(BuildUrl)" Project="$(TargetProject)" Platform="$(StagedPlatformFolder)" HttpVersion="$(UE-CloudPublishHttpVersion)" Files="#$(PlatformStageNodeName) Files"/>
 					</ForEach>
 				</Node>
 			</Do>
@@ -825,7 +825,7 @@
 						<Property Name="BuildUrl" Value=""/>
 						<Property Name="BuildUrl" Value="$(UE_HORDE_URL)job/$(UE_HORDE_JOBID)?step=$(UE_HORDE_STEPID)" If="$(Horde)"/>
 
-						<CreateCloudArtifact Name="$(BuildNamePath)" Type="packaged-build" BaseDir="$(ProjectOutputDirectory)/Packages/$(StagedPlatformFolder)/$(PackageConfiguration)" Metadata="Configuration=$(PackageConfiguration);dashboard-category=$(PackageConfiguration) Builds;BuildUrl=$(BuildUrl)" Project="$(TargetProject)" Platform="$(StagedPlatformFolder)" HttpVersion="$(UE-CloudPublishHttpVersion)" Files="#$(PlatformPackageNodeName) $(PackageConfiguration) Files"/>
+						<CreateCloudArtifact Name="$(BuildNamePath)" HordeArtifactName="Cloud-$(ProjectName)-Packaged-$(StagedPlatformFolder)-$(PackageConfiguration)" Type="packaged-build" BaseDir="$(ProjectOutputDirectory)/Packages/$(StagedPlatformFolder)/$(PackageConfiguration)" Metadata="Configuration=$(PackageConfiguration);dashboard-category=Packaged Builds - Cloud;BuildUrl=$(BuildUrl)" Project="$(TargetProject)" Platform="$(StagedPlatformFolder)" HttpVersion="$(UE-CloudPublishHttpVersion)" Files="#$(PlatformPackageNodeName) $(PackageConfiguration) Files"/>
 					</ForEach>
 				</Node>
 			</Do>
diff --git a/Engine/Source/Programs/AutomationTool/BuildGraph/Tasks/CreateCloudArtifactTask.cs b/Engine/Source/Programs/AutomationTool/BuildGraph/Tasks/CreateCloudArtifactTask.cs
--- a/Engine/Source/Programs/AutomationTool/BuildGraph/Tasks/CreateCloudArtifactTask.cs
+++ b/Engine/Source/Programs/AutomationTool/BuildGraph/Tasks/CreateCloudArtifactTask.cs
@@ -1,18 +1,11 @@
 // Copyright Epic Games, Inc. All Rights Reserved.
 
 using System;
-using System.Collections.Concurrent;
 using System.Collections.Generic;
-using System.Diagnostics;
 using System.IO;
 using System.Linq;
-using System.Net.Http;
-using System.Net.Http.Headers;
-using System.Net.Mime;
 using System.Threading.Tasks;
 using System.Xml;
-using AutomationTool.Tasks.CloudDDC;
-using EpicGames.Compression;
 using EpicGames.Core;
 using EpicGames.Horde;
 using EpicGames.Horde.Artifacts;
@@ -24,7 +17,6 @@ using EpicGames.Serialization;
 using Microsoft.Extensions.DependencyInjection;
 using Microsoft.Extensions.Logging;
 using UnrealBuildTool;
-using IoHash = EpicGames.Core.IoHash;
 
 #nullable enable
 
@@ -41,6 +33,12 @@ namespace AutomationTool.Tasks
 		[TaskParameter]
 		public string Name { get; set; } = null!;
 
+		/// <summary>
+		/// Optional override of the name of the artifact as presented in Horde, defaults to name if not set
+		/// </summary>
+		[TaskParameter(Optional = true)]
+		public string HordeArtifactName { get; set; } = null!;
+
 		/// <summary>
 		/// The artifact type. Determines the permissions and expiration policy for the artifact.
 		/// </summary>
@@ -136,12 +134,6 @@ namespace AutomationTool.Tasks
 		[TaskParameter(Optional = true)]
 		public bool AllowMultipart { get; set; } = false;
 
-		/// <summary>
-		/// Use the zen cli to upload the build. Generally faster but experimental.
-		/// </summary>
-		[TaskParameter(Optional = true)]
-		public bool UseZenCli { get; set; } = true;
-
 		/// <summary>
 		/// Set the explicit http version to use. None to use http handshaking.
 		/// </summary>
@@ -174,6 +166,7 @@ namespace AutomationTool.Tasks
 		{
 			ArtifactName name = new ArtifactName(_parameters.Name);
 			ArtifactType type = new ArtifactType(_parameters.Type);
+			ArtifactName hordeArtifactName = String.IsNullOrEmpty(_parameters.HordeArtifactName) ? name : new ArtifactName(_parameters.HordeArtifactName);
 			List<string> metadataParam = (_parameters.Metadata ?? String.Empty).Split(';', StringSplitOptions.RemoveEmptyEntries | StringSplitOptions.TrimEntries).ToList();
 
 			Dictionary<string, object> metadata = new Dictionary<string, object>();
@@ -245,7 +238,7 @@ namespace AutomationTool.Tasks
 			if (String.IsNullOrEmpty(baselineBranch))
 			{
 				Logger.LogWarning("No baseline branch set, falling back to current stream as baseline. This will generate in-optimal de-dup for branches, make sure to set this to your mainline by specifying the property or setting this in your StorageServer section in your engine ini");
-				baselineBranch = streamId.ToString();
+				baselineBranch = streamId;
 			}
 
 			baselineBranch = SanitizeBucketValue(baselineBranch);
@@ -277,10 +270,37 @@ namespace AutomationTool.Tasks
 			}
 
 			string publishHostEnvVar = "UE-CloudPublishHost";
+			string descriptorHostEnvVar = "UE-CloudPublishDescriptorHost";
 			if (!OperatingSystem.IsWindows())
 			{
 				publishHostEnvVar = publishHostEnvVar.Replace("-", "_", StringComparison.OrdinalIgnoreCase);
+				descriptorHostEnvVar = descriptorHostEnvVar.Replace("-", "_", StringComparison.OrdinalIgnoreCase);
+			}
+
+			// the host to specify in files published to users, as the host we connect to in the farm maybe different from what users should connect to
+			string? descriptorHost = null;
+			if (!String.IsNullOrEmpty(Environment.GetEnvironmentVariable(descriptorHostEnvVar)))
+			{
+				descriptorHost = Environment.GetEnvironmentVariable(descriptorHostEnvVar);
 			}
+
+			string? cloudConfigHost = null;
+			if (foundConfig && !String.IsNullOrEmpty(cloudConfig.Host))
+			{
+				cloudConfigHost = cloudConfig.Host;
+				if (cloudConfigHost.Contains(';', StringComparison.OrdinalIgnoreCase))
+				{
+					// if its a list pick the first element
+					cloudConfigHost = cloudConfigHost.Split(";").First();
+				}
+
+				// if not explicitly set we default the user facing host value to the config setting value
+				if (descriptorHost == null)
+				{
+					descriptorHost = cloudConfigHost;
+				}
+			}
+
 			string host;
 			if (!String.IsNullOrEmpty(_parameters.Host))
 			{
@@ -290,20 +310,21 @@ namespace AutomationTool.Tasks
 			{
 				host = Environment.GetEnvironmentVariable(publishHostEnvVar)!;
 			}
-			else if (foundConfig && !String.IsNullOrEmpty(cloudConfig.Host))
+			else if (!String.IsNullOrEmpty(cloudConfigHost))
 			{
-				host = cloudConfig.Host;
-				if (host.Contains(';', StringComparison.OrdinalIgnoreCase))
-				{
-					// if its a list pick the first element
-					host = host.Split(";").First();
-				}
+				host = cloudConfigHost;
 			}
 			else
 			{
 				throw new AutomationException("Missing UE-CloudPublishHost environment variable; unable to determine cloud host. Please specify this environment variable or define it in your StorageServer section in your engine ini");
 			}
 
+			// if we have not yet had the user facing host set then we just assume that is the same as the host we are publishing to
+			if (descriptorHost == null)
+			{
+				descriptorHost = host;
+			}
+
 			string publishNamespaceEnvVar = "UE-CloudPublishNamespace";
 			if (!OperatingSystem.IsWindows())
 			{
@@ -379,128 +400,20 @@ namespace AutomationTool.Tasks
 			// set part name to default by default - should be exposed as a option to allow for uploading separate parts
 			string partName = "default";
 
-			Uri buildUri = new Uri($"{host}/api/v2/builds/{ns}/{bucketId}/{buildId}");
-
-			if (_parameters.UseZenCli)
-			{
-				BuildUploader.UploadBuild(host, baseDir, ns, bucketId, buildId, partName, partId, metadata, accessToken, new BuildUploader.BuildInfo
-				{
-					Name = name.ToString(),
-					ProjectName = project,
-					Branch = streamId.ToString(),
-					BaselineBranch = baselineBranch,
-					Platform = platform
-				}, allowMultipart: _parameters.AllowMultipart, assumeHttp2: String.Equals(_parameters.HttpVersion, "http2-only", StringComparison.OrdinalIgnoreCase), forceClean: _parameters.Force);
-
-				try
-				{
-					await RegisterHordeArtifactAsync(buildUri, name, type, streamIdTyped, commitId, metadata);
-				}
-				catch (Exception ex)
-				{
-					Logger.LogWarning(ex, "Exception registering Horde artifact: {Message}", ex.Message);
-				}
-
-				return;
-			}
-
-			using HttpClient httpClient = BuildHttpClient(host, ns, bucketId, buildId, accessToken);
-			Logger.LogInformation("Creating artifact {ArtifactId} '{ArtifactName}' ({ArtifactType}) from {NumFiles} files ({TotalSize:n1}mb). Namespace: {NamespaceId}, Bucket: {Bucket}", buildId, name, type, files.Count, totalSize / (1024.0 * 1024.0), ns, bucketId);
-
-			// Do the work
-			Stopwatch timer = Stopwatch.StartNew();
-
-			// Compute file hashes
-			Dictionary<FileReference, IoHash> hashCache = await ComputeHashCacheAsync(files);
-
-			double hashTimeSecs = timer.Elapsed.TotalSeconds;
-			Logger.LogInformation("Hash cache built for build '{ArtifactName}' in {Duration:n1}s.", buildId, hashTimeSecs);
-
-			// Create the build
-			uint chunkSizeHint = await UploadBuildAsync(httpClient, project, streamId.ToString(), baselineBranch, platform, name.ToString(), metadata);
-
-			long totalSizeUploaded = 0;
-			CbObjectId buildPartId = CbObjectId.NewObjectId();
-			if (_parameters.Force)
+			BuildUploader.UploadBuild(host, baseDir, ns, bucketId, buildId, partName, partId, metadata, accessToken, new BuildUploader.BuildInfo
 			{
-				// force all parts to be uploaded, unlikely that blobs will be missing after this.
-				// this is faster for blobs that are very likely to be unique and when your connectivity to CloudDDC is very fast
-				Logger.LogInformation("Force uploading all {CountOfParts} parts of artifact '{ArtifactId}'. ", files.Count, buildId);
-
-				long sizeUploaded = 0;
-				TimeSpan blobUploadStart = timer.Elapsed;
-				await Parallel.ForEachAsync(files, new ParallelOptions {MaxDegreeOfParallelism = _parameters.MaxBlobUploadParallelism}, async (file, _) =>
-				{
-					IoHash hashOfFile = hashCache[file];
-					long uploadedBytes = await UploadBlobAsync(httpClient, buildPartId, file, hashOfFile);
-					sizeUploaded += uploadedBytes;
-				});
-
-				totalSizeUploaded += sizeUploaded;
-				double uploadTimeBlobsSecs = (timer.Elapsed - blobUploadStart).TotalSeconds;
-				Logger.LogInformation("Force uploaded blobs for artifact {ArtifactId} in {Time:n1}s ({Rate:n1}mb/s, {RateMbps:n1}mbps, {TotalSize:n1}mb)", buildId, uploadTimeBlobsSecs, sizeUploaded / (uploadTimeBlobsSecs * 1024.0 * 1024.0), (sizeUploaded * 8.0) / (uploadTimeBlobsSecs * 1024.0 * 1024.0), sizeUploaded / (1024.0 * 1024.0));
-			}
-
-			TimeSpan partCreationStart = timer.Elapsed;
-			// Create the build part
-			(List<IoHash> missingParts, IoHash partHash) = await UploadBuildPartObjectAsync(httpClient, name, buildPartId, baseDir, files, hashCache, chunkSizeHint);
-
-			double partCreationSecs = (timer.Elapsed - partCreationStart).TotalSeconds;
-			Logger.LogInformation("Spent {Time:n1}s finding {CountOfMissingParts} missing part(s) for build '{ArtifactName}'.", partCreationSecs, missingParts.Count, name);
-
-			const int MaxAttempts = 5;
-			for (int i = 0; i < MaxAttempts; i++)
-			{
-				// upload missing parts
-				if (missingParts.Count != 0)
-				{
-					Logger.LogInformation("{NumMissingParts} parts missing from build '{ArtifactName}'. Starting upload.", missingParts.Count, name);
-
-					List<KeyValuePair<IoHash, FileReference>> filesToDownload = new();
-
-					foreach (FileReference f in files)
-					{
-						IoHash hashOfFile = hashCache[f];
-
-						if (missingParts.Contains(hashOfFile))
-						{
-							filesToDownload.Add(new KeyValuePair<IoHash, FileReference>(hashOfFile, f));
-						}
-					}
-
-					long sizeUploaded = 0;
-					TimeSpan blobUploadStart = timer.Elapsed;
-					await Parallel.ForEachAsync(filesToDownload, new ParallelOptions {MaxDegreeOfParallelism = _parameters.MaxBlobUploadParallelism}, async (pair, _) =>
-					{
-						long uploadedBytes = await UploadBlobAsync(httpClient, buildPartId, pair.Value, pair.Key);
-						sizeUploaded += uploadedBytes;
-					});
-
-					totalSizeUploaded += sizeUploaded;
-					double uploadTimeBlobsSecs = (timer.Elapsed - blobUploadStart).TotalSeconds;
-					Logger.LogInformation("Uploaded blobs for artifact {ArtifactId} in {Time:n1}s ({Rate:n1}mb/s, {RateMbps:n1}mbps, {TotalSize:n1}mb)", buildId, uploadTimeBlobsSecs, sizeUploaded / (uploadTimeBlobsSecs * 1024.0 * 1024.0), (sizeUploaded * 8.0) / (uploadTimeBlobsSecs * 1024.0 * 1024.0), sizeUploaded / (1024.0 * 1024.0));
-					// verify that all parts are now found
-					missingParts = await FinalizeBuildPartAsync(httpClient, buildPartId, partHash);
-				}
+				Name = name.ToString(),
+				ProjectName = project,
+				Branch = streamId,
+				BaselineBranch = baselineBranch,
+				Platform = platform
+			}, allowMultipart: _parameters.AllowMultipart, assumeHttp2: String.Equals(_parameters.HttpVersion, "http2-only", StringComparison.OrdinalIgnoreCase), forceClean: _parameters.Force);
 
-				if (missingParts.Count == 0)
-				{
-					Logger.LogInformation("Part {PartId} successfully uploaded for build '{ArtifactName}'.", buildPartId, buildId);
-
-					// upload successful
-					break;
-				}
-			}
-
-			// Flag the build as uploaded, making it read only
-			await FinalizeBuildAsync(httpClient);
-
-			double uploadTimeSecs = timer.Elapsed.TotalSeconds;
-			Logger.LogInformation("Uploaded artifact {ArtifactId} in {Time:n1}s ({Rate:n1}mb/s, {RateMbps:n1}mbps)", buildId, timer.Elapsed.TotalSeconds, totalSizeUploaded / (uploadTimeSecs * 1024.0 * 1024.0), (totalSizeUploaded * 8.0) / (uploadTimeSecs * 1024.0 * 1024.0));
+			Uri buildUri = new Uri($"{descriptorHost}/api/v2/builds/{ns}/{bucketId}/{buildId}");
 
 			try
 			{
-				await RegisterHordeArtifactAsync(buildUri, name, type, streamIdTyped, commitId, metadata);
+				await RegisterHordeArtifactAsync(buildUri, hordeArtifactName, type, streamIdTyped, commitId, metadata);
 			}
 			catch (Exception ex)
 			{
@@ -513,240 +426,11 @@ namespace AutomationTool.Tasks
 			// . have special meaning in the bucket names, its used as a separator between components, as such component can not themselves contain .
 			return s.Replace(".", "-", StringComparison.OrdinalIgnoreCase);
 		}
-
-		static HttpClient BuildHttpClient(string host, string ns, string bucketId, CbObjectId buildId, string accessToken)
-		{
-			HttpClient httpClient = new HttpClient();
-			httpClient.BaseAddress = new Uri($"{host}/api/v2/builds/{ns}/{bucketId}/{buildId}/");
-
-			CbObjectId sessionId = CbObjectId.NewObjectId();
-			Logger.LogInformation("Using SessionId {SessionId} please include this in any error reports.", sessionId);
-			string authScheme = "Bearer"; // TODO: for a proper implementation this should be configurable
-			httpClient.DefaultRequestHeaders.Add("Authorization", $"{authScheme} {accessToken}");
-			httpClient.DefaultRequestHeaders.Add("UE-Session", sessionId.ToString());
-			httpClient.DefaultRequestHeaders.Add("Accept", CustomMediaTypeNames.UnrealCompactBinary);
-			httpClient.Timeout = TimeSpan.FromMinutes(5.0); // bump timeout for each request as we sometimes have larger files to upload
-			return httpClient;
-		}
-
-		static async Task<Dictionary<FileReference, IoHash>> ComputeHashCacheAsync(List<FileReference> files)
-		{
-			ConcurrentDictionary<FileReference, IoHash> cache = new();
-
-			await Parallel.ForEachAsync(files, async (file, token) =>
-			{
-				await using Stream s = file.ToFileInfo().OpenRead();
-				IoHash hash = await IoHash.ComputeAsync(s, token);
-				cache.TryAdd(file, hash);
-			});
-
-			return cache.ToDictionary();
-		}
-
-		static async Task<uint> UploadBuildAsync(HttpClient httpClient, string projectName, string branch, string baselineBranch, string platform, string name, Dictionary<string, object> metadata)
-		{
-			CbWriter writer = new CbWriter();
-			writer.BeginObject();
-			writer.WriteString("project", projectName);
-			writer.WriteString("branch", branch);
-			writer.WriteString("baselineBranch", baselineBranch);
-			writer.WriteString("platform", platform);
-			writer.WriteString("name", name);
-			foreach (KeyValuePair<string, object> pair in metadata)
-			{
-				if (pair.Value is int value)
-				{
-					writer.WriteInteger(pair.Key, value);
-					continue;
-				}
-				writer.WriteString(pair.Key, pair.Value.ToString());
-			}
-			writer.EndObject();
-
-			CbObject buildObject = writer.ToObject();
-
-			using HttpContent requestContent = new ByteArrayContent(buildObject.GetView().ToArray());
-			requestContent.Headers.ContentType = new MediaTypeHeaderValue(CustomMediaTypeNames.UnrealCompactBinary);
-			using HttpResponseMessage result = await httpClient.PutAsync(new Uri($"", UriKind.Relative), requestContent);
-			result.EnsureSuccessStatusCode();
-			// TODO: return chunk size once Jupiter is updated to respond with it
-			//PutBuildResponse response = await result.Content.ReadAsCompactBinaryAsync<PutBuildResponse>();
-			//return response.ChunkSize;
-			return 0;
-		}
-
-		static async Task FinalizeBuildAsync(HttpClient httpClient)
-		{
-			using HttpResponseMessage result = await httpClient.PostAsync(new Uri("finalize", UriKind.Relative), null);
-			result.EnsureSuccessStatusCode();
-		}
-
-		static async Task<(List<IoHash>, IoHash)> UploadBuildPartObjectAsync(HttpClient httpClient, ArtifactName artifactName, CbObjectId buildPartId, DirectoryReference baseDir, List<FileReference> files, Dictionary<FileReference, IoHash> hashCache, uint chunkSizeHint)
-		{
-			CbWriter writer = new CbWriter();
-			writer.BeginObject();
-			writer.BeginObject("files");
-
-			writer.BeginUniformArray("paths", CbFieldType.String);
-			foreach (FileReference file in files)
-			{
-				// paths are always forwarded slashes
-				string path = file.MakeRelativeTo(baseDir).Replace('\\','/');
-				writer.WriteStringValue(path);
-			}
-			writer.EndUniformArray();
-
-			// encode chmod attributes if we are running on non-windows
-			if (!OperatingSystem.IsWindows())
-			{
-				writer.BeginUniformArray("chmod", CbFieldType.IntegerPositive);
-				foreach (FileReference file in files)
-				{
-					writer.WriteIntegerValue((uint)File.GetUnixFileMode(file.FullName));
-				}
-				writer.EndUniformArray();
-			}
-			else if (OperatingSystem.IsWindows())
-			{
-				// encode windows file attributes when running on windows
-				writer.BeginUniformArray("attr", CbFieldType.IntegerPositive);
-				foreach (FileReference file in files)
-				{
-					// paths are always forwarded slashes
-					writer.WriteIntegerValue((uint)File.GetAttributes(file.FullName));
-				}
-				writer.EndUniformArray();
-			}
-
-			writer.BeginUniformArray("rawhashes", CbFieldType.BinaryAttachment);
-			foreach (FileReference file in files)
-			{
-				IoHash hash = hashCache[file];
-				writer.WriteBinaryAttachmentValue(hash);
-			}
-			writer.EndUniformArray();
-
-			writer.BeginUniformArray("rawsizes", CbFieldType.IntegerPositive);
-			foreach (FileReference file in files)
-			{
-				// 0 if not chunked, the chunk size if chunked
-				writer.WriteIntegerValue(file.ToFileInfo().Length);
-			}
-			writer.EndUniformArray();
-			writer.WriteInteger("chunksizehint", chunkSizeHint);
-
-			writer.EndObject(); // files
-			writer.EndObject();
-
-			CbObject buildPartObject = writer.ToObject();
-			IoHash partHash = IoHash.Compute(buildPartObject.GetView().Span);
-			string partName = artifactName.ToString();
-
-			using HttpContent requestContent = new ByteArrayContent(buildPartObject.GetView().ToArray());
-			requestContent.Headers.ContentType = new MediaTypeHeaderValue(CustomMediaTypeNames.UnrealCompactBinary);
-			using HttpResponseMessage result = await httpClient.PutAsync(new Uri($"parts/{buildPartId}/{partName}", UriKind.Relative), requestContent);
-			result.EnsureSuccessStatusCode();
-
-			NeedsResponse response = await result.Content.ReadAsCompactBinaryAsync<NeedsResponse>();
-			return (response.Needs, partHash);
-		}
-
-		static async Task<List<IoHash>> FinalizeBuildPartAsync(HttpClient httpClient, CbObjectId partId, IoHash partHash)
-		{
-			using HttpResponseMessage result = await httpClient!.PostAsync(new Uri($"parts/{partId}/finalize/{partHash}", UriKind.Relative), null);
-			result.EnsureSuccessStatusCode();
-
-			NeedsResponse response = await result.Content.ReadAsCompactBinaryAsync<NeedsResponse>();
-			return response.Needs;
-		}
-
-		static async Task<long> UploadBlobAsync(HttpClient httpClient, CbObjectId partId, FileReference fileReference, IoHash hashOfFile)
-		{
-			FileInfo fi = fileReference.ToFileInfo();
-			await using Stream sourceStream = fi.OpenRead();
-			long rawLength = fi.Length;
-			bool compressContent = true;
-			long contentLength;
-			Stream? streamToSend = null;
-			string? tempFileName = null;
-			try
-			{
-				MediaTypeHeaderValue contentType;
-
-				bool useTempFile = rawLength > UseTempFileAtBytes;
-
-				if (compressContent)
-				{
-					if (useTempFile)
-					{
-						tempFileName = Path.GetTempFileName();
-						FileInfo tempFileInfo = new FileInfo(tempFileName!);
-						{
-							await using Stream s = tempFileInfo.OpenWrite();
-							CompressedBuffer.CompressContent(s, OodleCompressorType.Kraken, OodleCompressionLevel.VeryFast, sourceStream);
-						}
-
-						tempFileInfo.Refresh();
-						streamToSend = tempFileInfo.OpenRead();
-						contentLength = tempFileInfo.Length;
-					}
-					else
-					{
-						MemoryStream ms = new MemoryStream();
-						CompressedBuffer.CompressContent(ms, OodleCompressorType.Kraken, OodleCompressionLevel.VeryFast, sourceStream);
-						ms.Position = 0;
-						
-						streamToSend = ms;
-						contentLength = ms.Length;
-					}
-
-					contentType = new MediaTypeHeaderValue(CustomMediaTypeNames.UnrealCompressedBuffer);
-
-				}
-				else
-				{
-					// submit the uncompressed content
-					streamToSend = sourceStream;
-					contentType = new MediaTypeHeaderValue(MediaTypeNames.Application.Octet);
-					contentLength = fi.Length;
-				}
-
-				/*if (rawLength > MultipartAtBytes)
-				{
-					await MultipartUploadFileAsync(httpClient, partId, hashOfFile, streamToSend, contentType, idealChunkSize);
-				}
-				else*/
-				{
-					using HttpContent requestContent = new StreamContent(streamToSend);
-					requestContent.Headers.ContentType = contentType;
-
-					using HttpResponseMessage result = await httpClient!.PutAsync(new Uri($"parts/{partId}/blobs/{hashOfFile}", UriKind.Relative), requestContent);
-					result.EnsureSuccessStatusCode();
-				}
-			}
-			finally
-			{
-				if (streamToSend != null)
-				{
-					await streamToSend.DisposeAsync();
-				}
-
-				if (tempFileName != null)
-				{
-					if (File.Exists(tempFileName))
-					{
-						File.Delete(tempFileName);
-					}
-				}
-			}
-
-			return contentLength;
-		}
-
+		
 		async Task RegisterHordeArtifactAsync(Uri buildUri, ArtifactName name, ArtifactType type, StreamId streamId, CommitId commitId, Dictionary<string, object> metadata)
 		{		
 			// Add keys used to collate and find builds
-			List<string> keys = new List<string>();			
+			List<string> keys = new List<string>();
 
 			// Add keys for the job that's executing
 			string? jobId = Environment.GetEnvironmentVariable("UE_HORDE_JOBID");
@@ -769,7 +453,7 @@ namespace AutomationTool.Tasks
 			foreach (KeyValuePair<string, object> pair in metadata)
 			{
 				metadataList.Add($"{pair.Key}={pair.Value.ToString()}");
-			}			
+			}
 		
 			// TODO: Investigate registering in a zen artifact namespace 
 			IHordeClient hordeClient = CommandUtils.ServiceProvider.GetRequiredService<IHordeClient>();
@@ -789,8 +473,6 @@ namespace AutomationTool.Tasks
 			await artifact.CompleteAsync(rootRef);
 		}	
 
-		private const long UseTempFileAtBytes = Int32.MaxValue; // use a backing file when the uncompressed file is too large to fit into memory
-
 		/// <inheritdoc/>
 		public override void Write(XmlWriter writer)
 			=> Write(writer, _parameters);
