Fix issue if not staging the exe is specified.
Only delete and perform IOS exe copy if there's actually an exe.

diff --git a/Engine/Source/Programs/AutomationTool/Apple/ApplePlatform.Automation.cs b/Engine/Source/Programs/AutomationTool/Apple/ApplePlatform.Automation.cs
--- a/Engine/Source/Programs/AutomationTool/Apple/ApplePlatform.Automation.cs
+++ b/Engine/Source/Programs/AutomationTool/Apple/ApplePlatform.Automation.cs
@@ -110,18 +110,21 @@ public abstract class ApplePlatform : Platform
 		{
 			RootDirName = SC.ShortProjectName;
 		}
-		FileReference BinaryPath = FileReference.Combine(SC.StageDirectory, RootDirName, "Binaries", SC.PlatformDir, SC.StageExecutables[0]);
-		InternalUtils.SafeDeleteFile(BinaryPath.FullName, true);
-		// this may leave the binaries directory empty, but the Mac needs the Binaries/Mac dir to exist (see FMacPlatformProcess::BaseDir())
-		// so plop a file down in it's place
-		if (SC.StageTargetPlatform.PlatformType == UnrealTargetPlatform.Mac)
+		if (SC.StageExecutables.Any())
 		{
-		    Directory.CreateDirectory(BinaryPath.Directory.FullName);
-			File.WriteAllText(Path.Combine(BinaryPath.Directory.FullName, ".binariesdir"), "");
-		}
+			FileReference BinaryPath = FileReference.Combine(SC.StageDirectory, RootDirName, "Binaries", SC.PlatformDir, SC.StageExecutables[0]);
+			InternalUtils.SafeDeleteFile(BinaryPath.FullName, true);
+			// this may leave the binaries directory empty, but the Mac needs the Binaries/Mac dir to exist (see FMacPlatformProcess::BaseDir())
+			// so plop a file down in it's place
+			if (SC.StageTargetPlatform.PlatformType == UnrealTargetPlatform.Mac)
+			{
+				Directory.CreateDirectory(BinaryPath.Directory.FullName);
+				File.WriteAllText(Path.Combine(BinaryPath.Directory.FullName, ".binariesdir"), "");
+			}
 
-		DirectoryReference AppPath = new DirectoryReference(BinaryPath.FullName + ".app");
-		InternalUtils.SafeDeleteDirectory(AppPath.FullName, true);
+			DirectoryReference AppPath = new DirectoryReference(BinaryPath.FullName + ".app");
+			InternalUtils.SafeDeleteDirectory(AppPath.FullName, true);
+		}
 
 		if (AppleExports.UseModernXcode(Params.RawProjectPath))
 		{
