Optimizing error parser in summary by caching Linq enumerate results to a list

#rnx

diff --git a/Engine/Source/Programs/AutomationTool/Gauntlet/Unreal/Automation/UE.Automation.cs b/Engine/Source/Programs/AutomationTool/Gauntlet/Unreal/Automation/UE.Automation.cs
--- a/Engine/Source/Programs/AutomationTool/Gauntlet/Unreal/Automation/UE.Automation.cs
+++ b/Engine/Source/Programs/AutomationTool/Gauntlet/Unreal/Automation/UE.Automation.cs
@@ -955,13 +955,13 @@ namespace UE
 					// Tests failed so list that as our primary cause of failure
 					if (TestResults != null && TestResults.Failed > 0)
 					{
-						ExitReason = string.Format("{0} of {1} test(s) failed", TestResults.Failed, TestResults.Tests.Count());
+						ExitReason = string.Format("{0} of {1} test(s) failed", TestResults.Failed, TestResults.Tests.Count);
 						ExitCode = -1;
 						return UnrealProcessResult.TestFailure;
 					}
 
 					// If no tests were run then that's a failure (possibly a bad RunTest argument?)
-					if (TestResults == null || TestResults.Tests.Count() == 0)
+					if (TestResults == null || TestResults.Tests.Count == 0)
 					{
 						ExitReason = "No tests were executed!";
 						ExitCode = -1;
@@ -1045,27 +1045,28 @@ namespace UE
 				UnrealAutomatedTestPassResults JsonTestPassResults = GetTestPassResults(EditorRole.LogSummary);
 
 				// Filter our tests into categories
-				IEnumerable<UnrealAutomatedTestResult> AllTests = JsonTestPassResults.Tests.Where(T => !T.WasSkipped);
-				IEnumerable<UnrealAutomatedTestResult> IncompleteTests = AllTests.Where(T => !T.IsComplete);
-				IEnumerable<UnrealAutomatedTestResult> FailedTests = AllTests.Where(T => T.IsComplete && T.HasFailed);
-				IEnumerable<UnrealAutomatedTestResult> TestsWithWarnings = AllTests.Where(T => T.HasSucceeded && T.HasWarnings);
+				List<UnrealAutomatedTestResult> AllTests = JsonTestPassResults.Tests.Where(T => !T.WasSkipped).ToList();
+				List<UnrealAutomatedTestResult> IncompleteTests = AllTests.Where(T => !T.IsComplete).ToList();
+				List<UnrealAutomatedTestResult> FailedTests = AllTests.Where(T => T.IsComplete && T.HasFailed).ToList();
+				List<UnrealAutomatedTestResult> TestsWithWarnings = AllTests.Where(T => T.HasSucceeded && T.HasWarnings).ToList();
 
 				Action<IEnumerable<object>> NotifyMoreIfNeeded = E =>
 				{
-					if (E.Count() > MaxEventsDisplayPerTest)
+					int Count = E.Count();
+					if (Count > MaxEventsDisplayPerTest)
 					{
-						Log.Info(" (and {Count} more)", E.Count() - MaxEventsDisplayPerTest);
+						Log.Info(" (and {Count} more)", Count - MaxEventsDisplayPerTest);
 					}
 				};
 				// If there were abnormal exits then look only at the failed and incomplete tests only to avoid confusing things.
 				if (GetRolesThatExitedAbnormally().Any())
 				{
-					if (AllTests.Count() == 0)
+					if (AllTests.Count == 0)
 					{
 						Log.Error(KnownLogEvents.Gauntlet_UnrealEngineTestEvent, " * No tests were executed.");
 						Log.Info("");
 					}
-					else if (FailedTests.Count() > 0 || IncompleteTests.Count() > 0)
+					else if (FailedTests.Count > 0 || IncompleteTests.Count > 0)
 					{
 						var Failures = FailedTests.Concat(IncompleteTests.Where(T => T.State == TestStateType.InProcess));
 						if (Failures.Any())
@@ -1090,10 +1091,10 @@ namespace UE
 								Log.Info("");
 							}
 						}
-						var NotRun = IncompleteTests.Where(T => T.State == TestStateType.NotRun);
+						List<UnrealAutomatedTestResult> NotRun = IncompleteTests.Where(T => T.State == TestStateType.NotRun).ToList();
 						if (NotRun.Any())
 						{
-							int TotalNotRun = NotRun.Count();
+							int TotalNotRun = NotRun.Count;
 							Log.Info(" ### The following Engine test(s) were not run:");
 							Log.Info(string.Join("\n", NotRun.Take(10).Select(T => " * "+T.FullTestPath)));
 							if (TotalNotRun > 10)
@@ -1106,7 +1107,7 @@ namespace UE
 				}
 				else
 				{
-					if (AllTests.Count() == 0)
+					if (AllTests.Count == 0)
 					{
 						Log.Error(KnownLogEvents.Gauntlet_TestEvent, " * No tests were executed.");
 
@@ -1135,7 +1136,7 @@ namespace UE
 					else
 					{
 						// Now list the tests that failed
-						if (FailedTests.Count() > 0)
+						if (FailedTests.Count > 0)
 						{
 							Log.Info(" ### The following Engine test(s) failed:");
 
@@ -1152,7 +1153,7 @@ namespace UE
 							}
 						}
 
-						if (TestsWithWarnings.Count() > 0)
+						if (TestsWithWarnings.Count > 0)
 						{
 							Log.Info(" ### The following Engine test(s) completed with warnings:");
 
@@ -1170,7 +1171,7 @@ namespace UE
 							}
 						}
 
-						if (IncompleteTests.Count() > 0)
+						if (IncompleteTests.Count > 0)
 						{
 							Log.Info(" ### The following Engine test(s) timed out or did not run:");
 
@@ -1191,10 +1192,10 @@ namespace UE
 						// show a brief summary at the end where it's most visible
 						List<string> TestSummary = new List<string>();
 
-						int PassedTests = AllTests.Count() - (FailedTests.Count() + IncompleteTests.Count());
-						int TestsPassedWithoutWarnings = PassedTests - TestsWithWarnings.Count();
+						int PassedTests = AllTests.Count - (FailedTests.Count + IncompleteTests.Count);
+						int TestsPassedWithoutWarnings = PassedTests - TestsWithWarnings.Count;
 
-						TestSummary.Add(string.Format(" * {0} Test(s) Requested", AllTests.Count()));
+						TestSummary.Add(string.Format(" * {0} Test(s) Requested", AllTests.Count));
 
 						// Print out a summary of each category of result
 						if (TestsPassedWithoutWarnings > 0)
@@ -1202,19 +1203,19 @@ namespace UE
 							TestSummary.Add(string.Format(" * {0} Test(s) Passed", TestsPassedWithoutWarnings));
 						}
 
-						if (TestsWithWarnings.Count() > 0)
+						if (TestsWithWarnings.Count > 0)
 						{
-							TestSummary.Add(string.Format(" * {0} Test(s) Passed with warnings", TestsWithWarnings.Count()));
+							TestSummary.Add(string.Format(" * {0} Test(s) Passed with warnings", TestsWithWarnings.Count));
 						}
 
-						if (FailedTests.Count() > 0)
+						if (FailedTests.Count > 0)
 						{
-							TestSummary.Add(string.Format(" * {0} Test(s) Failed", FailedTests.Count()));
+							TestSummary.Add(string.Format(" * {0} Test(s) Failed", FailedTests.Count));
 						}
 
-						if (IncompleteTests.Count() > 0)
+						if (IncompleteTests.Count > 0)
 						{
-							TestSummary.Add(string.Format(" * {0} Test(s) didn't complete", IncompleteTests.Count()));
+							TestSummary.Add(string.Format(" * {0} Test(s) didn't complete", IncompleteTests.Count));
 						}
 
 						Log.Info(" ### Summary");
