Cloud Artifacts - Built in appending horde url into the create cloud artifact task to avoid having to build this in every graph we want it in and defaulted the http version to be read from env variable if not explicitly set. Reduces the amount of redundant configuration when using the tasks.
Now we fixup the artifact name to fit into our conventions (no starting + and prefer - as a seperator)
Removed non-zen-cli path from retrieve artifact (not used)

diff --git a/Engine/Build/Graph/Tasks/BuildAndTestProject.xml b/Engine/Build/Graph/Tasks/BuildAndTestProject.xml
--- a/Engine/Build/Graph/Tasks/BuildAndTestProject.xml
+++ b/Engine/Build/Graph/Tasks/BuildAndTestProject.xml
@@ -1,6 +1,5 @@
 <?xml version='1.0' ?>
 <BuildGraph xmlns="http://www.epicgames.com/BuildGraph" xmlns:xsi="http://www.w3.org/2001/XMLSchema-instance" xsi:schemaLocation="http://www.epicgames.com/BuildGraph ../Schema.xsd" >
-	<EnvVar Name="UE_HORDE_URL"/>
 	<EnvVar Name="UE_HORDE_JOBID"/>
 	<EnvVar Name="UE_HORDE_STEPID"/>
 	<EnvVar Name="UE_HORDE_TEMPLATEID"/>
@@ -763,10 +762,7 @@
 					<ForEach Name="PackageConfiguration" Values="$(PackageConfigurations)" Separator="+">
 						<Tag Files="$(ProjectOutputDirectory)/$(StagedPlatformFolder)/..." With="#$(PlatformStageNodeName) Files"/>
 
-						<Property Name="BuildUrl" Value=""/>
-						<Property Name="BuildUrl" Value="$(UE_HORDE_URL)job/$(UE_HORDE_JOBID)?step=$(UE_HORDE_STEPID)" If="$(Horde)"/>
-
-						<CreateCloudArtifact Name="$(BuildNamePath)" HordeArtifactName="Cloud-$(ProjectName)-Staged-$(StagedPlatformFolder)-$(PackageConfiguration)" Type="staged-build" BaseDir="$(ProjectOutputDirectory)/$(StagedPlatformFolder)" Metadata="Configuration=$(PackageConfiguration);dashboard-category=Staged Builds - Cloud;BuildUrl=$(BuildUrl)" Project="$(TargetProject)" Platform="$(StagedPlatformFolder)" HttpVersion="$(UE-CloudPublishHttpVersion)" Files="#$(PlatformStageNodeName) Files"/>
+						<CreateCloudArtifact Name="$(BuildNamePath)" HordeArtifactName="Cloud-$(ProjectName)-Staged-$(StagedPlatformFolder)-$(PackageConfiguration)" Type="staged-build" BaseDir="$(ProjectOutputDirectory)/$(StagedPlatformFolder)" Metadata="Configuration=$(PackageConfiguration);dashboard-category=Staged Builds - Cloud" Project="$(TargetProject)" Platform="$(StagedPlatformFolder)" Files="#$(PlatformStageNodeName) Files"/>
 					</ForEach>
 				</Node>
 			</Do>
@@ -822,10 +818,7 @@
 						<Command Name="BuildCookRun" Arguments="$(BCRArgs) -configuration=$(PackageConfiguration) $(StageAndPackageArguments) -skipbuild -skipcook -skippak $(CloudPackageSkipBCRArguments) -archive -archivedirectory=&quot;$(ProjectOutputDirectory)/Packages/$(StagedPlatformFolder)/$(PackageConfiguration)&quot;" />
 						<Tag Files="$(ProjectOutputDirectory)/Packages/$(StagedPlatformFolder)/$(PackageConfiguration)/..." With="#$(PlatformPackageNodeName) $(PackageConfiguration) Files"/>
 
-						<Property Name="BuildUrl" Value=""/>
-						<Property Name="BuildUrl" Value="$(UE_HORDE_URL)job/$(UE_HORDE_JOBID)?step=$(UE_HORDE_STEPID)" If="$(Horde)"/>
-
-						<CreateCloudArtifact Name="$(BuildNamePath)" HordeArtifactName="Cloud-$(ProjectName)-Packaged-$(StagedPlatformFolder)-$(PackageConfiguration)" Type="packaged-build" BaseDir="$(ProjectOutputDirectory)/Packages/$(StagedPlatformFolder)/$(PackageConfiguration)" Metadata="Configuration=$(PackageConfiguration);dashboard-category=Packaged Builds - Cloud;BuildUrl=$(BuildUrl)" Project="$(TargetProject)" Platform="$(StagedPlatformFolder)" HttpVersion="$(UE-CloudPublishHttpVersion)" Files="#$(PlatformPackageNodeName) $(PackageConfiguration) Files"/>
+						<CreateCloudArtifact Name="$(BuildNamePath)" HordeArtifactName="Cloud-$(ProjectName)-Packaged-$(StagedPlatformFolder)-$(PackageConfiguration)" Type="packaged-build" BaseDir="$(ProjectOutputDirectory)/Packages/$(StagedPlatformFolder)/$(PackageConfiguration)" Metadata="Configuration=$(PackageConfiguration);dashboard-category=Packaged Builds - Cloud" Project="$(TargetProject)" Platform="$(StagedPlatformFolder)" Files="#$(PlatformPackageNodeName) $(PackageConfiguration) Files"/>
 					</ForEach>
 				</Node>
 			</Do>
diff --git a/Engine/Source/Programs/AutomationTool/BuildGraph/Tasks/CreateCloudArtifactTask.cs b/Engine/Source/Programs/AutomationTool/BuildGraph/Tasks/CreateCloudArtifactTask.cs
--- a/Engine/Source/Programs/AutomationTool/BuildGraph/Tasks/CreateCloudArtifactTask.cs
+++ b/Engine/Source/Programs/AutomationTool/BuildGraph/Tasks/CreateCloudArtifactTask.cs
@@ -164,17 +164,32 @@ namespace AutomationTool.Tasks
 		/// <param name="tagNameToFileSet">Mapping from tag names to the set of files they include.</param>
 		public override async Task ExecuteAsync(JobContext job, HashSet<FileReference> buildProducts, Dictionary<string, HashSet<FileReference>> tagNameToFileSet)
 		{
-			ArtifactName name = new ArtifactName(_parameters.Name);
+			string name = SanitizeBucketValue(_parameters.Name);
 			ArtifactType type = new ArtifactType(_parameters.Type);
-			ArtifactName hordeArtifactName = String.IsNullOrEmpty(_parameters.HordeArtifactName) ? name : new ArtifactName(_parameters.HordeArtifactName);
+			ArtifactName hordeArtifactName = String.IsNullOrEmpty(_parameters.HordeArtifactName) ? new ArtifactName(name) : new ArtifactName(_parameters.HordeArtifactName);
 			List<string> metadataParam = (_parameters.Metadata ?? String.Empty).Split(';', StringSplitOptions.RemoveEmptyEntries | StringSplitOptions.TrimEntries).ToList();
 
 			Dictionary<string, object> metadata = new Dictionary<string, object>();
-			string project = SanitizeBucketValue(_parameters.Project.GetFileNameWithoutExtension());
 			string platform = SanitizeBucketValue(_parameters.Platform);
-			ProjectProperties properties = ProjectUtils.GetProjectProperties(_parameters.Project);
-			ConfigHierarchy config = properties.EngineConfigs[HostPlatform.Current.HostEditorPlatform];
-			bool foundConfig = config.TryGetValueGeneric("StorageServers", "Cloud", out CloudConfiguration cloudConfig);
+			
+			string project = _parameters.Project.GetFileNameWithoutExtension();
+			FileReference? projectPath = _parameters.Project;
+			if (!File.Exists(_parameters.Project.FullName))
+			{
+				// if the project is not an uproject that exists we do the scanning for engine inis without the project specific overrides
+				projectPath = null;
+				Logger.LogInformation("\'{Project}\' is not a valid uproject file. Assuming it is a project name and not loading inis, you will need to specify project options explicitly.", _parameters.Project);
+				project = _parameters.Project.FullName;
+			}
+
+			bool foundConfig = false;
+			CloudConfiguration? cloudConfig = null;
+			if (projectPath != null)
+			{
+				ProjectProperties properties = ProjectUtils.GetProjectProperties(projectPath);
+				ConfigHierarchy config = properties.EngineConfigs[HostPlatform.Current.HostEditorPlatform];
+				foundConfig = config.TryGetValueGeneric("StorageServers", "Cloud", out cloudConfig);
+			}
 
 			foreach (string line in metadataParam)
 			{
@@ -196,18 +211,36 @@ namespace AutomationTool.Tasks
 			}
 
 			// Add keys for the job that's executing
+			string? stepId = null;
 			string? jobId = Environment.GetEnvironmentVariable("UE_HORDE_JOBID");
 			if (!String.IsNullOrEmpty(jobId))
 			{
 				metadata.Add("job", jobId);
 
-				string? stepId = Environment.GetEnvironmentVariable("UE_HORDE_STEPID");
+				stepId = Environment.GetEnvironmentVariable("UE_HORDE_STEPID");
 				if (!String.IsNullOrEmpty(stepId))
 				{
 					metadata.Add("step", stepId);
 				}
 			}
 
+			string? hordeUrl = Environment.GetEnvironmentVariable("UE_HORDE_URL");
+			// if we are running in horde and have the required environment variables we append a link back to the horde job into metadata
+			if (!String.IsNullOrEmpty(hordeUrl) && !String.IsNullOrEmpty(jobId) && !String.IsNullOrEmpty(stepId))
+			{
+				metadata.Add("buildurl", $"{hordeUrl}job/{jobId}?step={stepId}");
+			}
+
+			string httpVersion = _parameters.HttpVersion;
+			if (String.IsNullOrEmpty(httpVersion) || httpVersion == "None")
+			{
+				string? httpVersionEnvironmentVariable = Environment.GetEnvironmentVariable("UE-CloudPublishHttpVersion");
+				if (!String.IsNullOrEmpty(httpVersionEnvironmentVariable))
+				{
+					httpVersion = httpVersionEnvironmentVariable;
+				}
+			}
+			
 			// Figure out the current change and stream id
 			StreamId streamIdTyped;
 			if (!String.IsNullOrEmpty(_parameters.StreamId))
@@ -230,10 +263,10 @@ namespace AutomationTool.Tasks
 			metadata.Add("stream", streamId);
 			
 			string baselineBranch =_parameters.BaselineBranch;
-			if (foundConfig && String.IsNullOrEmpty(baselineBranch))
+			if (foundConfig && cloudConfig != null && String.IsNullOrEmpty(baselineBranch))
 			{
 				// no explicit baseline branch set and we have found a storage server config, using that value
-				baselineBranch = cloudConfig.BuildsBaselineBranch;
+				baselineBranch = cloudConfig.Value.BuildsBaselineBranch;
 			}
 			if (String.IsNullOrEmpty(baselineBranch))
 			{
@@ -285,9 +318,9 @@ namespace AutomationTool.Tasks
 			}
 
 			string? cloudConfigHost = null;
-			if (foundConfig && !String.IsNullOrEmpty(cloudConfig.Host))
+			if (foundConfig && cloudConfig != null && !String.IsNullOrEmpty(cloudConfig.Value.Host))
 			{
-				cloudConfigHost = cloudConfig.Host;
+				cloudConfigHost = cloudConfig.Value.Host;
 				if (cloudConfigHost.Contains(';', StringComparison.OrdinalIgnoreCase))
 				{
 					// if its a list pick the first element
@@ -339,9 +372,9 @@ namespace AutomationTool.Tasks
 			{
 				ns = Environment.GetEnvironmentVariable(publishNamespaceEnvVar)!;
 			}
-			else if (foundConfig && !String.IsNullOrEmpty(cloudConfig.BuildsNamespace))
+			else if (foundConfig && cloudConfig != null && !String.IsNullOrEmpty(cloudConfig.Value.BuildsNamespace))
 			{
-				ns = cloudConfig.BuildsNamespace;
+				ns = cloudConfig.Value.BuildsNamespace;
 			}
 			else
 			{
@@ -407,7 +440,7 @@ namespace AutomationTool.Tasks
 				Branch = streamId,
 				BaselineBranch = baselineBranch,
 				Platform = platform
-			}, allowMultipart: _parameters.AllowMultipart, assumeHttp2: String.Equals(_parameters.HttpVersion, "http2-only", StringComparison.OrdinalIgnoreCase), forceClean: _parameters.Force);
+			}, allowMultipart: _parameters.AllowMultipart, assumeHttp2: String.Equals(httpVersion, "http2-only", StringComparison.OrdinalIgnoreCase), forceClean: _parameters.Force);
 
 			Uri buildUri = new Uri($"{descriptorHost}/api/v2/builds/{ns}/{bucketId}/{buildId}");
 
@@ -423,6 +456,9 @@ namespace AutomationTool.Tasks
 
 		internal static string SanitizeBucketValue(string s)
 		{
+			// trim starting + and then replace any + and . with -
+			s = s.TrimStart('+');
+			s = s.Replace("+", "-", StringComparison.OrdinalIgnoreCase);
 			// . have special meaning in the bucket names, its used as a separator between components, as such component can not themselves contain .
 			return s.Replace(".", "-", StringComparison.OrdinalIgnoreCase);
 		}
diff --git a/Engine/Source/Programs/AutomationTool/BuildGraph/Tasks/RetrieveCloudArtifactTask.cs b/Engine/Source/Programs/AutomationTool/BuildGraph/Tasks/RetrieveCloudArtifactTask.cs
--- a/Engine/Source/Programs/AutomationTool/BuildGraph/Tasks/RetrieveCloudArtifactTask.cs
+++ b/Engine/Source/Programs/AutomationTool/BuildGraph/Tasks/RetrieveCloudArtifactTask.cs
@@ -2,12 +2,10 @@
 
 using System;
 using System.Collections.Generic;
-using System.Diagnostics;
 using System.IO;
 using System.Linq;
 using System.Net.Http;
 using System.Net.Http.Headers;
-using System.Net.Mime;
 using System.Threading.Tasks;
 using System.Xml;
 using AutomationTool.Tasks.CloudDDC;
@@ -18,7 +16,6 @@ using EpicGames.Horde.Streams;
 using EpicGames.Serialization;
 using Microsoft.Extensions.Logging;
 using UnrealBuildTool;
-using Exception = System.Exception;
 
 #nullable enable
 
@@ -113,12 +110,6 @@ namespace AutomationTool.Tasks
 		[TaskParameter(Optional = true)]
 		public bool AllowMultipart { get; set; } = false;
 
-		/// <summary>
-		/// Use the zen cli to upload the build. Generally faster but experimental.
-		/// </summary>
-		[TaskParameter(Optional = true)]
-		public bool UseZenCli { get; set; } = true;
-
 		/// <summary>
 		/// Set the explicit http version to use. None to use http handshaking.
 		/// </summary>
@@ -155,11 +146,24 @@ namespace AutomationTool.Tasks
 			}
 
 			string project = _parameters.Project.GetFileNameWithoutExtension();
+			FileReference? projectPath = _parameters.Project;
+			if (!File.Exists(_parameters.Project.FullName))
+			{
+				// if the project is not an uproject that exists we do the scanning for engine inis without the project specific overrides
+				projectPath = null;
+				Logger.LogInformation("\'{Project}\' is not a valid file. Loading engine configs without project overrides.", _parameters.Project);
+			}
 			string platform = _parameters.Platform;
 
-			ProjectProperties properties = ProjectUtils.GetProjectProperties(_parameters.Project);
-			ConfigHierarchy config = properties.EngineConfigs[HostPlatform.Current.HostEditorPlatform];
-			bool foundConfig = config.TryGetValueGeneric("StorageServers", "Cloud", out CloudConfiguration cloudConfig);
+			bool foundConfig = false;
+			CloudConfiguration? cloudConfig = null;
+
+			if (projectPath != null)
+			{
+				ProjectProperties properties = ProjectUtils.GetProjectProperties(projectPath);
+				ConfigHierarchy config = properties.EngineConfigs[HostPlatform.Current.HostEditorPlatform];
+				foundConfig = config.TryGetValueGeneric("StorageServers", "Cloud", out cloudConfig);
+			}
 
 			// Figure out the current change and stream id
 			StreamId streamId;
@@ -185,9 +189,9 @@ namespace AutomationTool.Tasks
 			{
 				host = _parameters.Host;
 			}
-			else if (foundConfig && !String.IsNullOrEmpty(cloudConfig.Host))
+			else if (foundConfig && cloudConfig != null && !String.IsNullOrEmpty(cloudConfig.Value.Host))
 			{
-				host = cloudConfig.Host;
+				host = cloudConfig.Value.Host;
 			}
 			else
 			{
@@ -207,9 +211,9 @@ namespace AutomationTool.Tasks
 			{
 				ns = _parameters.Namespace;
 			}
-			else if (foundConfig && !String.IsNullOrEmpty(cloudConfig.BuildsNamespace))
+			else if (foundConfig && cloudConfig != null && !String.IsNullOrEmpty(cloudConfig.Value.BuildsNamespace))
 			{
-				ns = cloudConfig.BuildsNamespace;
+				ns = cloudConfig.Value.BuildsNamespace;
 			}
 			else
 			{
@@ -224,6 +228,16 @@ namespace AutomationTool.Tasks
 				}
 			}
 
+			string httpVersion = _parameters.HttpVersion;
+			if (String.IsNullOrEmpty(httpVersion) || httpVersion == "None")
+			{
+				string? httpVersionEnvironmentVariable = Environment.GetEnvironmentVariable("UE-CloudPublishHttpVersion");
+				if (!String.IsNullOrEmpty(httpVersionEnvironmentVariable))
+				{
+					httpVersion = httpVersionEnvironmentVariable;
+				}
+			}
+
 			string accessToken;
 			if (!String.IsNullOrEmpty(_parameters.AccessToken))
 			{
@@ -269,14 +283,14 @@ namespace AutomationTool.Tasks
 				}
 			}
 
-			ArtifactName? name = null;
+			string? name = null;
 			if (!String.IsNullOrEmpty(_parameters.Name))
 			{
-				name = new ArtifactName(_parameters.Name);
+				name = _parameters.Name;
 			}
 
 			ArtifactType type = new ArtifactType(_parameters.Type);
-			List<string> keys = (_parameters.Keys ?? String.Empty).Split(';', StringSplitOptions.RemoveEmptyEntries | StringSplitOptions.TrimEntries).ToList();
+			List<string> keys = (_parameters.Keys).Split(';', StringSplitOptions.RemoveEmptyEntries | StringSplitOptions.TrimEntries).ToList();
 
 			CbWriter queryWriter = new CbWriter();
 			queryWriter.BeginObject(); // root
@@ -312,7 +326,7 @@ namespace AutomationTool.Tasks
 			if (name != null)
 			{
 				queryWriter.BeginObject("name");
-				queryWriter.WriteString("$eq", name.ToString());
+				queryWriter.WriteString("$eq", name);
 				queryWriter.EndObject();
 			}
 
@@ -358,16 +372,7 @@ namespace AutomationTool.Tasks
 			CbObjectId artifact = artifactIds.First();
 			Logger.LogInformation("Found artifact {ArtifactId}", artifact);
 
-			if (_parameters.UseZenCli)
-			{
-				BuildDownload.DownloadBuild(host, outputDir, ns, bucketId, artifact.ToString(), accessToken, allowMultipart: _parameters.AllowMultipart, assumeHttp2: String.Equals(_parameters.HttpVersion, "http2-only", StringComparison.OrdinalIgnoreCase));
-				return;
-			}
-
-			// Upload the files
-			Stopwatch timer = Stopwatch.StartNew();
-			await DownloadArtifactAsync(httpClient, artifact, outputDir.ToDirectoryInfo());
-			Logger.LogInformation("Retrieved artifact {ArtifactId} in {Time:n1}s", artifact, timer.Elapsed.TotalSeconds);
+			BuildDownload.DownloadBuild(host, outputDir, ns, bucketId, artifact.ToString(), accessToken, allowMultipart: _parameters.AllowMultipart, assumeHttp2: String.Equals(httpVersion, "http2-only", StringComparison.OrdinalIgnoreCase));
 		}
 
 		static HttpClient BuildHttpClient(string host, string ns, string bucketId, string accessToken)
@@ -384,104 +389,6 @@ namespace AutomationTool.Tasks
 			httpClient.Timeout = TimeSpan.FromMinutes(5.0); // bump timeout for each request as we sometimes have larger files
 			return httpClient;
 		}
-
-		static async Task DownloadArtifactAsync(HttpClient httpClient, CbObjectId buildId, System.IO.DirectoryInfo outputDir)
-		{
-			using HttpResponseMessage result = await httpClient.GetAsync(new Uri($"{buildId}", UriKind.Relative));
-			result.EnsureSuccessStatusCode();
-
-			byte[] response = await result.Content.ReadAsByteArrayAsync();
-			CbObject o = new CbObject(response);
-			if (o["parts"].Equals(CbField.Empty))
-			{
-				throw new Exception("Failed to find parts in build, unexpected result, unable to download");
-			}
-
-			CbObject parts = o["parts"].AsObject();
-			foreach (CbField part in parts)
-			{
-				Utf8String partName = part.Name;
-				CbObjectId partId = part.AsObjectId();
-
-				Logger.LogInformation("Downloading part \'{PartName}\' with id {PartId}", partName, partId);
-				using HttpResponseMessage partResult = await httpClient.GetAsync(new Uri($"{buildId}/parts/{partId}", UriKind.Relative));
-				partResult.EnsureSuccessStatusCode();
-
-				byte[] partResponse = await partResult.Content.ReadAsByteArrayAsync();
-				CbObject partObject = new CbObject(partResponse);
-				CbObject filesObject = partObject["files"].AsObject();
-				List<string> paths = filesObject["paths"].AsArray().Select(field => field.AsString()).ToList();
-				List<long> rawsizes = filesObject["rawsizes"].AsArray().Select(field => field.AsInt64()).ToList();
-				List<IoHash> rawhashes = filesObject["rawhashes"].AsArray().Select(field => field.AsHash()).ToList();
-				//uint chunkSizeHint = filesObject["chunksizehint"].AsUInt32();
-
-				string attributeField = OperatingSystem.IsWindows() ? "attr" : "chmod";
-				List<int> attributes =  filesObject[attributeField].AsArray().Select(field => field.AsInt32()).ToList();
-
-				List<(string, long, IoHash, int)> foundFiles = [];
-				for (int i = 0; i < paths.Count; i++)
-				{
-					string path = paths[i];
-					long rawSize = rawsizes[i];
-					IoHash hash = rawhashes[i];
-					int attr = attributes[i];
-
-					foundFiles.Add((path, rawSize, hash, attr));
-				}
-
-				await Parallel.ForEachAsync(foundFiles, async (field, token) =>
-				{
-					(string fileName, long length, IoHash hash, int attr) = field;
-					string fullPath = Path.Combine(outputDir.FullName, fileName);
-					FileInfo outputFile = new FileInfo(fullPath);
-					await DownloadBlobAsync(httpClient, outputFile, buildId, partId, hash);
-					if (OperatingSystem.IsWindows())
-					{
-						File.SetAttributes(outputFile.FullName, (FileAttributes)attr);
-					}
-					else
-					{
-						File.SetUnixFileMode(outputFile.FullName, (UnixFileMode)attr);
-					}
-				});
-			}
-		}
-
-		private static async Task DownloadBlobAsync(HttpClient httpClient, FileInfo fullPath, CbObjectId buildId, CbObjectId partId, IoHash attachmentId)
-		{
-			using HttpRequestMessage request = new HttpRequestMessage(HttpMethod.Get, new Uri($"{buildId}/parts/{partId}/blobs/{attachmentId}", UriKind.Relative));
-			// accept all forms of blobs, we handle them based on the response mime type
-			request.Headers.Add("Accept", "*/*");
-
-			using HttpResponseMessage result = await httpClient!.SendAsync(request);
-			result.EnsureSuccessStatusCode();
-
-			// Make sure the directories are created if they are missing
-			fullPath.Directory?.Create();
-			string contentType = result.Content.Headers.ContentType!.ToString();
-			if (contentType == CustomMediaTypeNames.UnrealCompressedBuffer)
-			{
-				// compressed
-				long? contentLength = result.Content.Headers.ContentLength;
-				if (contentLength == null)
-				{
-					throw new Exception("Content length is missing");
-				}
-				await using Stream s = fullPath.OpenWrite();
-				await CompressedBuffer.DecompressContentAsync(await result.Content.ReadAsStreamAsync(), (ulong)contentLength.Value, s);
-			}
-			else if (contentType == MediaTypeNames.Application.Octet)
-			{
-				// uncompressed, just write to disk
-				await using Stream s = fullPath.OpenWrite();
-				await result.Content.CopyToAsync(s);
-			}
-			else
-			{
-				throw new Exception($"Unknown content type {contentType} when downloading blob {attachmentId}");
-			}
-		}
-
 		static async Task<List<CbObjectId>> SearchArtifactsAsync(HttpClient httpClient, CbObject queryObject)
 		{
 			using ByteArrayContent content = new ByteArrayContent(queryObject.GetView().ToArray());
@@ -490,11 +397,11 @@ namespace AutomationTool.Tasks
 			using HttpRequestMessage request = new HttpRequestMessage(HttpMethod.Post, new Uri($"search", UriKind.Relative));
 			request.Content = content;
 			request.Headers.Add("Accept", CustomMediaTypeNames.UnrealCompactBinary);
-			using HttpResponseMessage result = await httpClient!.SendAsync(request);
+			using HttpResponseMessage result = await httpClient.SendAsync(request);
 			result.EnsureSuccessStatusCode();
 
 			SearchResult searchResult = await result.Content.ReadAsCompactBinaryAsync<SearchResult>();
-			return searchResult.Results.Select(result => result.BuildId).ToList();
+			return searchResult.Results.Select(r => r.BuildId).ToList();
 		}
 
 		/// <inheritdoc/>
