Propagate config overrides in SetupCustomStageCopyHandler

#rnx

diff --git a/Engine/Source/Programs/AutomationTool/Scripts/CopyBuildToStagingDirectory.Automation.cs b/Engine/Source/Programs/AutomationTool/Scripts/CopyBuildToStagingDirectory.Automation.cs
--- a/Engine/Source/Programs/AutomationTool/Scripts/CopyBuildToStagingDirectory.Automation.cs
+++ b/Engine/Source/Programs/AutomationTool/Scripts/CopyBuildToStagingDirectory.Automation.cs
@@ -5982,7 +5982,7 @@ namespace AutomationScripts
 		private static void SetupCustomStageCopyHandler(ProjectParams Params, DeploymentContext SC)
 		{
 			string CustomStageCopyHandlerName = null;
-			ConfigHierarchy GameIni = ConfigCache.ReadHierarchy(ConfigHierarchyType.Game, Params.RawProjectPath.Directory, SC.StageTargetPlatform.PlatformType, SC.CustomConfig);
+			ConfigHierarchy GameIni = ConfigCache.ReadHierarchy(ConfigHierarchyType.Game, Params.RawProjectPath.Directory, SC.StageTargetPlatform.PlatformType, SC.CustomConfig, Params.AdditionalPakOptions.Split());
 			GameIni.GetString("/Script/UnrealEd.ProjectPackagingSettings", "CustomStageCopyHandler", out CustomStageCopyHandlerName);
 
 			if (!string.IsNullOrEmpty(CustomStageCopyHandlerName))
