Fixed issue with reading engine inis for cloud artifact tasks.
Quoted paths sent to zen-cli when using Create and Retrieve cloud artifacts.

diff --git a/Engine/Source/Programs/AutomationTool/BuildGraph/Tasks/CreateCloudArtifactTask.cs b/Engine/Source/Programs/AutomationTool/BuildGraph/Tasks/CreateCloudArtifactTask.cs
--- a/Engine/Source/Programs/AutomationTool/BuildGraph/Tasks/CreateCloudArtifactTask.cs
+++ b/Engine/Source/Programs/AutomationTool/BuildGraph/Tasks/CreateCloudArtifactTask.cs
@@ -188,7 +188,8 @@ namespace AutomationTool.Tasks
 			{
 				ProjectProperties properties = ProjectUtils.GetProjectProperties(projectPath);
 				ConfigHierarchy config = properties.EngineConfigs[HostPlatform.Current.HostEditorPlatform];
-				foundConfig = config.TryGetValueGeneric("StorageServers", "Cloud", out cloudConfig);
+				foundConfig = config.TryGetValueGeneric("StorageServers", "Cloud", out CloudConfiguration foundCloudConfig);
+				cloudConfig = foundCloudConfig;
 			}
 
 			foreach (string line in metadataParam)
@@ -574,7 +575,7 @@ namespace AutomationTool.Tasks
 				string metadataString = File.ReadAllText(metadataJsonPath);
 				string http2Options = assumeHttp2 ? "--assume-http2" : "";
 				// pass the access token via environment variable to avoid showing it on the cli
-				string cmdline = $"builds upload --url {host} --namespace {namespaceId} --bucket {bucketId} --build-id {buildId} --build-part-name {partName} --build-part-id {partId} --metadata-path {metadataJsonPath} --access-token-env UE-CloudDataCacheAccessToken --create-build --plain-progress --allow-multipart={allowMultipart} {http2Options} --clean={forceClean} --verify --verbose {baseDir}";
+				string cmdline = $"builds upload --url {host} --namespace {namespaceId} --bucket {bucketId} --build-id {buildId} --build-part-name \"{partName}\" --build-part-id {partId} --metadata-path \"{metadataJsonPath}\" --access-token-env UE-CloudDataCacheAccessToken --create-build --plain-progress --allow-multipart={allowMultipart} {http2Options} --clean={forceClean} --verify --verbose \"{baseDir}\"";
 				Log.Logger.LogInformation("Using metadata '{Metadata}'", metadataString);
 				CommandUtils.RunAndLog(CommandUtils.CmdEnv, zenExe.FullName,cmdline , Options: CommandUtils.ERunOptions.Default, EnvVars: new Dictionary<string, string> {{"UE-CloudDataCacheAccessToken", accessToken}});
 			}
diff --git a/Engine/Source/Programs/AutomationTool/BuildGraph/Tasks/RetrieveCloudArtifactTask.cs b/Engine/Source/Programs/AutomationTool/BuildGraph/Tasks/RetrieveCloudArtifactTask.cs
--- a/Engine/Source/Programs/AutomationTool/BuildGraph/Tasks/RetrieveCloudArtifactTask.cs
+++ b/Engine/Source/Programs/AutomationTool/BuildGraph/Tasks/RetrieveCloudArtifactTask.cs
@@ -162,7 +162,8 @@ namespace AutomationTool.Tasks
 			{
 				ProjectProperties properties = ProjectUtils.GetProjectProperties(projectPath);
 				ConfigHierarchy config = properties.EngineConfigs[HostPlatform.Current.HostEditorPlatform];
-				foundConfig = config.TryGetValueGeneric("StorageServers", "Cloud", out cloudConfig);
+				foundConfig = config.TryGetValueGeneric("StorageServers", "Cloud", out CloudConfiguration foundCloudConfig);
+				cloudConfig = foundCloudConfig;
 			}
 
 			// Figure out the current change and stream id
@@ -447,7 +448,7 @@ namespace AutomationTool.Tasks
 			string http2Options = assumeHttp2 ? "--assume-http2" : "";
 			// pass the access token via environment variable to avoid showing it on the cli
 			string cmdline =
-				$"builds download --url {host} --namespace {namespaceId} --bucket {bucketId} --access-token-env UE-CloudDataCacheAccessToken --plain-progress --allow-multipart={allowMultipart} {http2Options} {targetDir} {buildId}";
+				$"builds download --url {host} --namespace {namespaceId} --bucket {bucketId} --access-token-env UE-CloudDataCacheAccessToken --plain-progress --allow-multipart={allowMultipart} {http2Options} \"{targetDir}\" {buildId}";
 			CommandUtils.RunAndLog(CommandUtils.CmdEnv, zenExe.FullName, cmdline, Options: CommandUtils.ERunOptions.Default,
 				EnvVars: new Dictionary<string, string> { { "UE-CloudDataCacheAccessToken", accessToken } });
 		}
