Changing iOS tests to use non bulk builds for faster install time

#rnx

diff --git a/Engine/Source/Programs/AutomationTool/Gauntlet/Platform/IOS/Gauntlet.TargetDeviceIOS.cs b/Engine/Source/Programs/AutomationTool/Gauntlet/Platform/IOS/Gauntlet.TargetDeviceIOS.cs
--- a/Engine/Source/Programs/AutomationTool/Gauntlet/Platform/IOS/Gauntlet.TargetDeviceIOS.cs
+++ b/Engine/Source/Programs/AutomationTool/Gauntlet/Platform/IOS/Gauntlet.TargetDeviceIOS.cs
@@ -413,21 +413,50 @@ namespace Gauntlet
 			string EngineSavedDirectory = "/Documents/Engine/Saved";
 			string ProjectSavedDirectory = string.Format("/Documents/{0}/Saved", ProjectName);
 
+			string ListEngineSaved = string.Format("--bundle_id {0} --list={1}", BundleID, EngineSavedDirectory);
 			string CleanEngine = string.Format("--bundle_id {0} --rmtree {1}", BundleID, EngineSavedDirectory);
-			string CleanProject = string.Format("--bundle_id {0} --rmtree {1}", BundleID, ProjectSavedDirectory);
 
-			IProcessResult Result = ExecuteIOSDeployCommand(CleanEngine);
+			// The output of the --list command will print all files in that path
+			// so if the output contains the path, it exists
+			// if no output is provided, the path does not exist
+			
+			IProcessResult Result = ExecuteIOSDeployCommand(ListEngineSaved);
 			Result.WaitForExit();
 			if (Result.ExitCode != 0)
 			{
-				throw new DeviceException("Cleaning Engine artifacts for {0} exited with code {1}:\n{2}", BundleID, Result.ExitCode, Result.Output);
+				throw new DeviceException("Listing Engine artifacts for {0} exited with code {1}:\n{2}", BundleID, Result.ExitCode, Result.Output);
+			}
+			
+			if (Result.Output.Contains(EngineSavedDirectory))
+			{
+				Result = ExecuteIOSDeployCommand(CleanEngine);
+				Result.WaitForExit();
+				if (Result.ExitCode != 0)
+				{
+					throw new DeviceException("Cleaning Engine artifacts for {0} exited with code {1}:\n{2}", BundleID, Result.ExitCode, Result.Output);
+				}	
 			}
 
-			Result = ExecuteIOSDeployCommand(CleanProject);
+			string ListProjectSaved = string.Format("--bundle_id {0} --list={1}", BundleID, ProjectSavedDirectory);
+			string CleanProject = string.Format("--bundle_id {0} --rmtree {1}", BundleID, ProjectSavedDirectory);
+			
+			Result = ExecuteIOSDeployCommand(ListProjectSaved);
 			Result.WaitForExit();
 			if (Result.ExitCode != 0)
 			{
-				throw new DeviceException("Cleaning Project artifacts for {0} exited with code {1}:\n{2}", BundleID, Result.ExitCode, Result.Output);
+				throw new DeviceException("Listing Project artifacts for {0} exited with code {1}:\n{2}", BundleID, Result.ExitCode, Result.Output);
+			}
+			
+			// The same applies here for the --list
+			if (Result.Output.Contains(ListProjectSaved))
+			{
+				Result = ExecuteIOSDeployCommand(CleanProject);
+				Result.WaitForExit();
+				if (Result.ExitCode != 0)
+				{
+					throw new DeviceException("Cleaning Project artifacts for {0} exited with code {1}:\n{2}", BundleID, Result.ExitCode,
+						Result.Output);
+				}
 			}
 		}
 
