Fixed issues with download Cloud Artifacts for use in tests. Not enabled by default, set RetrieveUsingCloudArtifacts=true to enable.

diff --git a/Engine/Build/Graph/Tasks/BuildAndTestProject.xml b/Engine/Build/Graph/Tasks/BuildAndTestProject.xml
--- a/Engine/Build/Graph/Tasks/BuildAndTestProject.xml
+++ b/Engine/Build/Graph/Tasks/BuildAndTestProject.xml
@@ -860,7 +860,7 @@
 					<Property Name="GauntletBuildArg" Value="$(GauntletBuildArg) -destlocalinstalldir=&quot;$(LocalInstallDirectory)&quot;" If="!$(IsLocalPlatform) And $(WithDeviceReservation)"/>
 				</Do>
 
-				<Do If="$(Horde)">
+				<Do If="$(Horde) and $(RetrieveUsingCloudArtifacts) == ''">
 					<Property Name="HordeArtifactPlatformRetrieveWithConfig" Value="$(HordeArtifactPlatformRetrieve) $(TestConfiguration) (Horde)"/>
 					<Node Name="$(HordeArtifactPlatformRetrieveWithConfig)" Requires="$(TestPrerequisites)">
 						<Switch>
@@ -875,6 +875,22 @@
 					<Property Name="TestPrerequisites" Value="$(TestPrerequisites);$(HordeArtifactPlatformRetrieveWithConfig)" />
 				</Do>
 
+				<Do If="$(RetrieveUsingCloudArtifacts)">
+					<Property Name="CloudArtifactPlatformRetrieveWithConfig" Value="$(CloudArtifactPlatformRetrieve) $(TestConfiguration) (Cloud)"/>
+					<Node Name="$(CloudArtifactPlatformRetrieveWithConfig)" Requires="$(TestPrerequisites)">
+						<Switch>
+							<Case If="$(UsePackaged)">
+								<RetrieveCloudArtifact Name="$(BuildNamePath)" Type="packaged-build" OutputDir="$(ProjectOutputDirectory)/Packages/$(StagedPlatformFolder)/$(TestConfiguration)" Project="$(TargetProject)" Platform="$(StagedPlatformFolder)" />
+							</Case>
+							<Default>
+								<RetrieveCloudArtifact Name="$(BuildNamePath)" Type="staged-build" OutputDir="$(ProjectOutputDirectory)/$(StagedPlatformFolder)" Project="$(TargetProject)" Platform="$(StagedPlatformFolder)" />
+							</Default>
+						</Switch>
+					</Node>
+
+					<Property Name="TestPrerequisites" Value="$(TestPrerequisites);$(CloudArtifactPlatformRetrieveWithConfig)" />
+				</Do>
+
 				<Property Name="GauntletTestArgs" Value="$(GauntletTargetTestArgs)"/>
 				<Property Name="GauntletTestArgs" Value="$(GauntletTargetTestAltArgs)" If="ContainsItem('$(AlternateDevicePoolPlatforms)', '$(TargetPlatform)', ';')"/>
 
diff --git a/Engine/Source/Programs/AutomationTool/BuildGraph/Tasks/RetrieveCloudArtifactTask.cs b/Engine/Source/Programs/AutomationTool/BuildGraph/Tasks/RetrieveCloudArtifactTask.cs
--- a/Engine/Source/Programs/AutomationTool/BuildGraph/Tasks/RetrieveCloudArtifactTask.cs
+++ b/Engine/Source/Programs/AutomationTool/BuildGraph/Tasks/RetrieveCloudArtifactTask.cs
@@ -6,6 +6,8 @@ using System.IO;
 using System.Linq;
 using System.Net.Http;
 using System.Net.Http.Headers;
+using System.Net.Http.Json;
+using System.Net.Mime;
 using System.Threading.Tasks;
 using System.Xml;
 using AutomationTool.Tasks.CloudDDC;
@@ -291,7 +293,7 @@ namespace AutomationTool.Tasks
 			}
 
 			ArtifactType type = new ArtifactType(_parameters.Type);
-			List<string> keys = (_parameters.Keys).Split(';', StringSplitOptions.RemoveEmptyEntries | StringSplitOptions.TrimEntries).ToList();
+			List<string> keys = (_parameters.Keys ?? String.Empty).Split(';', StringSplitOptions.RemoveEmptyEntries | StringSplitOptions.TrimEntries).ToList();
 
 			CbWriter queryWriter = new CbWriter();
 			queryWriter.BeginObject(); // root
@@ -355,14 +357,14 @@ namespace AutomationTool.Tasks
 
 			CbObject queryObject = queryWriter.ToObject();
 			DirectoryReference outputDir = ResolveDirectory(_parameters.OutputDir);
-			string bucketId = $"{CreateCloudArtifactTask.SanitizeBucketValue(project)}.{CreateCloudArtifactTask.SanitizeBucketValue(type.ToString())}.{CreateCloudArtifactTask.SanitizeBucketValue(streamId.ToString())}.{CreateCloudArtifactTask.SanitizeBucketValue(platform)}";
+			string bucketId = $"{CreateCloudArtifactTask.SanitizeBucketValue(project)}.{CreateCloudArtifactTask.SanitizeBucketValue(type.ToString())}.{CreateCloudArtifactTask.SanitizeBucketValue(streamId.ToString())}.{CreateCloudArtifactTask.SanitizeBucketValue(platform)}".ToLowerInvariant();
 			using HttpClient httpClient = BuildHttpClient(host, ns, bucketId, accessToken);
 
 			List<CbObjectId> artifactIds = await SearchArtifactsAsync(httpClient, queryObject);
 
 			if (artifactIds.Count == 0)
 			{
-				throw new AutomationException($"Unable to find any artifact matching given criteria in namespace {ns} and bucket {bucketId}.");
+				throw new AutomationException($"Unable to find any artifact matching given criteria in namespace {ns} and bucket {bucketId} criteria: {queryObject.ToJson()}");
 			}
 
 			if (!_parameters.AllowMultipleMatches && artifactIds.Count != 1)
@@ -392,17 +394,17 @@ namespace AutomationTool.Tasks
 		}
 		static async Task<List<CbObjectId>> SearchArtifactsAsync(HttpClient httpClient, CbObject queryObject)
 		{
-			using ByteArrayContent content = new ByteArrayContent(queryObject.GetView().ToArray());
-			content.Headers.ContentType = new MediaTypeHeaderValue(CustomMediaTypeNames.UnrealCompactBinary);
+			using StringContent content = new StringContent(queryObject.ToJson());
+			content.Headers.ContentType = new MediaTypeHeaderValue(MediaTypeNames.Application.Json);
 
 			using HttpRequestMessage request = new HttpRequestMessage(HttpMethod.Post, new Uri($"search", UriKind.Relative));
 			request.Content = content;
-			request.Headers.Add("Accept", CustomMediaTypeNames.UnrealCompactBinary);
+			request.Headers.Add("Accept", MediaTypeNames.Application.Json);
 			using HttpResponseMessage result = await httpClient.SendAsync(request);
 			result.EnsureSuccessStatusCode();
 
-			SearchResult searchResult = await result.Content.ReadAsCompactBinaryAsync<SearchResult>();
-			return searchResult.Results.Select(r => r.BuildId).ToList();
+			SearchResult? searchResult = await result.Content.ReadFromJsonAsync<SearchResult>();
+			return searchResult!.Results.Select(r => r.BuildId).ToList();
 		}
 
 		/// <inheritdoc/>
@@ -422,8 +424,6 @@ namespace AutomationTool.Tasks
 	{
 		[CbField("buildId")]
 		public CbObjectId BuildId { get; set; }
-		[CbField("metadata")]
-		public CbObject Metadata { get; set; } = null!;
 	}
 
 	class SearchResult
