Add tags support to AutomatedTestSessionData (Horde) report

#rnx

diff --git a/Engine/Source/Programs/AutomationTool/Gauntlet/Framework/Utils/Gauntlet.HordeReport.cs b/Engine/Source/Programs/AutomationTool/Gauntlet/Framework/Utils/Gauntlet.HordeReport.cs
--- a/Engine/Source/Programs/AutomationTool/Gauntlet/Framework/Utils/Gauntlet.HordeReport.cs
+++ b/Engine/Source/Programs/AutomationTool/Gauntlet/Framework/Utils/Gauntlet.HordeReport.cs
@@ -8,7 +8,6 @@ using System.Text.Json;
 using System.Text.RegularExpressions;
 using System.Text.Json.Serialization;
 using AutomationTool;
-using System.Security.Cryptography;
 using Microsoft.CodeAnalysis;
 using EpicGames.Core;
 using Logging = Microsoft.Extensions.Logging;
@@ -1291,7 +1290,8 @@ namespace Gauntlet
 				public TestPhaseOutcome outcome { get; set; }
 				public string errorFingerprint { get; set; }
 				public bool? hasWarning { get; set; }
-				public List<string> deviceKeys { get; set; }
+				public HashSet<string> tags { get; set; }
+				public HashSet<string> deviceKeys { get; set; }
 				public string eventStreamPath { get; set; }
 
 				private TestSession session { get; set; }
@@ -1302,8 +1302,9 @@ namespace Gauntlet
 				{
 					key = string.IsNullOrEmpty(InKey) ? GenerateKey(InName) : ValidateKey(InKey);
 					name = InName;
-					deviceKeys = new List<string>();
+					deviceKeys = new HashSet<string>();
 					outcome = TestPhaseOutcome.Unknown;
+					tags = null;
 					errorFingerprint = null;
 					eventStream = new TestEventStream(session?.report, this);
 					AttachToSession(InSession);
@@ -1403,6 +1404,19 @@ namespace Gauntlet
 					}
 				}
 
+				/// <summary>
+				/// Add tags to the phase
+				/// </summary>
+				/// <param name="NewTags"></param>
+				public void AddTags(IEnumerable<string> NewTags)
+				{
+					if (tags == null)
+					{
+						tags = new();
+					}
+					tags.UnionWith(NewTags);
+				}
+
 				/// <summary>
 				/// Get the event stream associated with this phase
 				/// </summary>
@@ -1424,6 +1438,7 @@ namespace Gauntlet
 				public Dictionary<string, TestDevice> devices { get; set; }
 				public int version { get { return report.Version; } }
 				public Dictionary<string, string> metadata { get { return report.Metadata; } }
+				public HashSet<string> tags { get; set; }
 
 				[JsonIgnore]
 				// back pointer
@@ -1435,6 +1450,7 @@ namespace Gauntlet
 					phases = new List<TestPhase>();
 					devices = new Dictionary<string, TestDevice>();
 					report = InReport;
+					tags = null;
 				}
 			}
 
@@ -1561,6 +1577,19 @@ namespace Gauntlet
 				return CurrentPhase;
 			}
 
+			/// <summary>
+			/// Add tags to the test
+			/// </summary>
+			/// <param name="NewTags"></param>
+			public virtual void AddTestTags(IEnumerable<string> NewTags)
+			{
+				if (Data.tags == null)
+				{
+					Data.tags = new();
+				}
+				Data.tags.UnionWith(NewTags);
+			}
+
 			/// <summary>
 			/// Represent the event stream of a phase
 			/// </summary>
@@ -2167,6 +2196,11 @@ namespace Gauntlet
 					}
 					// update phase
 					Phase.eventStreamPath = OutputFilePath;
+					// propagate phase tags to test tags
+					if (Phase.tags != null)
+					{
+						AddTestTags(Phase.tags);
+					}
 				}
 
 				return base.GetReportDependencies();
@@ -2213,7 +2247,7 @@ namespace Gauntlet
 					foreach (UnrealAutomatedTestResult InTestResult in InTestPassResults.Tests)
 					{
 						TestPhase Phase = OutTestPassResults.AddPhase(InTestResult.FullTestPath);
-						Phase.deviceKeys = InTestResult.DeviceInstance;
+						Phase.deviceKeys = InTestResult.DeviceInstance.ToHashSet();
 						string TestDateTime = InTestResult.DateTime;
 						DateTime PhaseTime = ConvertUETimeStringToDateTime(TestDateTime);
 						Phase.SetTiming(PhaseTime, InTestResult.Duration);
