Make GetCachedConfiguration() protected and use it when applicable in UE.Automation and AutomationTestPerf nodes

#rnx

diff --git a/Engine/Plugins/Performance/AutomatedPerfTesting/Build/Scripts/AutomatedPerfTestNode.cs b/Engine/Plugins/Performance/AutomatedPerfTesting/Build/Scripts/AutomatedPerfTestNode.cs
--- a/Engine/Plugins/Performance/AutomatedPerfTesting/Build/Scripts/AutomatedPerfTestNode.cs
+++ b/Engine/Plugins/Performance/AutomatedPerfTesting/Build/Scripts/AutomatedPerfTestNode.cs
@@ -305,7 +305,7 @@ namespace AutomatedPerfTest
 						};
 						
 						Log.Info("Creating perf server importer with build name {BuildName}", OriginalBuildName);
-						string DataSourceName = GetConfiguration().DataSourceName;
+						string DataSourceName = GetCachedConfiguration().DataSourceName;
 						string ImportDirOverride = Globals.Params.ParseValue("PerfReportServerImportDir", null);
 						Logger.LogInformation("Creating PRS Importer for data source '{0}' and import dir override (if any) '{1}'.",
 							DataSourceName, ImportDirOverride);
@@ -326,7 +326,7 @@ namespace AutomatedPerfTest
 
 							// todo update this so it associates videos with the correct CSVs
 							IEnumerable<CsvImportEntry> CsvImportEntries = ImportEntries as CsvImportEntry[] ?? ImportEntries.ToArray();
-							if (GetConfiguration().DoInsightsTrace)
+							if (GetCachedConfiguration().DoInsightsTrace)
 							{
 								string InsightsFilename = Path.GetFileNameWithoutExtension(CsvImportEntries.First().CsvFilename)
 									.Replace(".csv", ".utrace");
@@ -349,7 +349,7 @@ namespace AutomatedPerfTest
 								}
 							}
 
-							if (GetConfiguration().DoVideoCapture)
+							if (GetCachedConfiguration().DoVideoCapture)
 							{
 								string VideoPath = Path.Combine(ArtifactPath, "Client", "Videos");
 								string[] VideoFiles = Directory.GetFiles(VideoPath, "*.mp4");
@@ -785,20 +785,6 @@ namespace AutomatedPerfTest
 				}
 			}
 		}
-
-		/// <summary>
-		/// Returns the cached version of our config. Avoids repeatedly calling GetConfiguration() on derived nodes
-		/// </summary>
-		/// <returns></returns>
-		private TConfigClass GetCachedConfiguration()
-		{
-			if (CachedConfig == null)
-			{
-				return GetConfiguration();
-			}
-
-			return CachedConfig;
-		}
 	}
 
 	/// <summary>
diff --git a/Engine/Source/Programs/AutomationTool/Gauntlet/Unreal/Automation/UE.Automation.cs b/Engine/Source/Programs/AutomationTool/Gauntlet/Unreal/Automation/UE.Automation.cs
--- a/Engine/Source/Programs/AutomationTool/Gauntlet/Unreal/Automation/UE.Automation.cs
+++ b/Engine/Source/Programs/AutomationTool/Gauntlet/Unreal/Automation/UE.Automation.cs
@@ -478,7 +478,7 @@ namespace UE
 			{
 				string BaseName = base.Name;
 
-				var Config = GetConfiguration();
+				var Config = GetCachedConfiguration();
 				if (Config is AutomationTestConfig)
 				{
 					var AutomationConfig = Config as AutomationTestConfig;
@@ -500,7 +500,7 @@ namespace UE
 			get
 			{
 				string BaseSuite = base.Suite;
-				var Config = GetConfiguration();
+				var Config = GetCachedConfiguration();
 				if (Config is AutomationTestConfig)
 				{
 					var AutomationConfig = Config as AutomationTestConfig;
@@ -533,7 +533,7 @@ namespace UE
 			get
 			{
 				string Base = "UE.Automation";
-				var Config = GetConfiguration();
+				var Config = GetCachedConfiguration();
 				if (Config is AutomationTestConfig)
 				{
 					var AutomationConfig = Config as AutomationTestConfig;
@@ -559,7 +559,7 @@ namespace UE
 			LogReader = null;
 			TestPassResults = null;
 
-			if (GetConfiguration() is AutomationTestConfig Config)
+			if (GetCachedConfiguration() is AutomationTestConfig Config)
 			{
 				if ((Config.WriteTestResultsForHorde || Config.ResumeOnCriticalFailure) && string.IsNullOrEmpty(Config.ReportExportPath))
 				{
@@ -590,7 +590,7 @@ namespace UE
 			LogReader = null;
 			TestPassResults = null;
 
-			if (GetConfiguration() is AutomationTestConfig Config)
+			if (GetCachedConfiguration() is AutomationTestConfig Config)
 			{
 				Config.GenerateSessionID();
 			}
@@ -646,7 +646,7 @@ namespace UE
 			if (TestPassResults == null)
 			{
 				string AutomationReportPath = string.Empty;
-				if (GetConfiguration() is AutomationTestConfig Config)
+				if (GetCachedConfiguration() is AutomationTestConfig Config)
 				{
 					AutomationReportPath = Config.ReportExportPath;
 					if (!string.IsNullOrEmpty(AutomationReportPath))
@@ -707,8 +707,8 @@ namespace UE
 		private void UpdateTestStateOnCriticalFailure(UnrealAutomatedTestPassResults JsonTestPassResults)
 		{
 			bool HasTimeout = RoleResults != null && RoleResults.Where(R => R.ProcessResult == UnrealProcessResult.TimeOut).Any();
-			string HordeArtifactPath = GetConfiguration().HordeArtifactPath;
-			var MainRole = GetConfiguration().GetMainRequiredRole();
+			string HordeArtifactPath = GetCachedConfiguration().HordeArtifactPath;
+			var MainRole = GetCachedConfiguration().GetMainRequiredRole();
 			// Get any critical error and push it to json report and resave it.
 			UnrealLog.CallstackMessage FatalError = null;
 			UnrealRoleResult FatalErrorRoleResult = null;
@@ -849,7 +849,7 @@ namespace UE
 			{
 				// The test pass did not run at all
 				Log.Verbose("Found not-run tests: {Count}", JsonTestPassResults.NotRun);
-				if ((GetConfiguration() is AutomationTestConfig Config) && Config.ResumeOnCriticalFailure && !HasTimeout)
+				if ((GetCachedConfiguration() is AutomationTestConfig Config) && Config.ResumeOnCriticalFailure && !HasTimeout)
 				{
 					// Reschedule test to resume from last 'in-process' test.
 					if (SetToRetryIfPossible())
@@ -980,7 +980,7 @@ namespace UE
 		public override ITestReport CreateReport(TestResult Result)
 		{
 			ITestReport Report = null;
-			if (GetConfiguration() is AutomationTestConfig Config)
+			if (GetCachedConfiguration() is AutomationTestConfig Config)
 			{
 				// Parse the json test pass results and look for critical failure to add to report
 				var TestPassResults = GetTestPassResults(null);
@@ -1228,7 +1228,7 @@ namespace UE
 				{
 					string AutomationReportPath = string.Empty;
 					string AutomationReportURL = string.Empty;
-					if (GetConfiguration() is AutomationTestConfig Config)
+					if (GetCachedConfiguration() is AutomationTestConfig Config)
 					{
 						AutomationReportPath = Config.ReportExportPath;
 						AutomationReportURL = Config.ReportURL;
diff --git a/Engine/Source/Programs/AutomationTool/Gauntlet/Unreal/Base/Gauntlet.UnrealTestNode.cs b/Engine/Source/Programs/AutomationTool/Gauntlet/Unreal/Base/Gauntlet.UnrealTestNode.cs
--- a/Engine/Source/Programs/AutomationTool/Gauntlet/Unreal/Base/Gauntlet.UnrealTestNode.cs
+++ b/Engine/Source/Programs/AutomationTool/Gauntlet/Unreal/Base/Gauntlet.UnrealTestNode.cs
@@ -435,11 +435,11 @@ namespace Gauntlet
 		/// Returns the cached version of our config. Avoids repeatedly calling GetConfiguration() on derived nodes
 		/// </summary>
 		/// <returns></returns>
-		private TConfigClass GetCachedConfiguration()
+		protected TConfigClass GetCachedConfiguration()
 		{
 			if (CachedConfig == null)
 			{
-				return GetConfiguration();
+				CachedConfig = GetConfiguration();
 			}
 
 			return CachedConfig;
