LaunchMultiServer support for bot clients

Added "-ProxyBotCount=#" command-line argument to LaunchMultiServer UAT script.

This will launch the specified number of additional clients with these arguments "-log -nullrhi -bot -unattended -nosplash ".

It also gives the bot clients a "NewConsole" window by default (you can override with "-nonewconsole" like usual).

diff --git a/Engine/Source/Programs/AutomationTool/Scripts/LaunchMultiServer.Automation.cs b/Engine/Source/Programs/AutomationTool/Scripts/LaunchMultiServer.Automation.cs
--- a/Engine/Source/Programs/AutomationTool/Scripts/LaunchMultiServer.Automation.cs
+++ b/Engine/Source/Programs/AutomationTool/Scripts/LaunchMultiServer.Automation.cs
@@ -333,6 +333,25 @@ namespace AutomationTool
 								var ProxyClientAdditionalParams = String.Format("127.0.0.1:{0} -log=DstmProxyClient-{1}.log -SaveWinPos={1}", ProxyListeningPort, (ClientIndex+1));
 								StartProcessesForServerInfos(UniqueServerInfos, ProcessLaunchType.ProxyClient, ProxyClientAdditionalParams);
 							}
+
+							if (ParseParam("proxybotcount"))
+							{
+								int BotCount = 0;
+								string BotCountStr = ParseParamValue("proxybotcount", null);
+								Int32.TryParse(BotCountStr, out BotCount);
+								
+								for (int BotIndex = 0; BotIndex < BotCount; BotIndex++)
+								{
+									int BotNumber = BotIndex + 1;
+									var BotClientAdditionalParams = String.Format("127.0.0.1:{0} -log=DstmProxyBot-{1}.log -log -nullrhi -bot -unattended -nosplash ", ProxyListeningPort, BotNumber);
+									BotClientAdditionalParams += String.Format("-ConsoleTitle=\"{0} Proxy Bot Client {1}\" ", ProjectUtils.GetShortProjectName(ProjectPath), BotNumber);
+									if (!ParseParam("nonewconsole"))
+									{
+										BotClientAdditionalParams += " -newconsole ";
+									}
+									StartProcessesForServerInfos(UniqueServerInfos, ProcessLaunchType.ProxyClient, BotClientAdditionalParams);
+								}
+							}
 						}
 					}
 				}
@@ -356,7 +375,7 @@ namespace AutomationTool
 				(
 					Command: this,
 					RawProjectPath: ProjectPath,
-					DedicatedServer: LaunchType == ProcessLaunchType.Server || LaunchType ==ProcessLaunchType.ProxyServer,
+					DedicatedServer: LaunchType == ProcessLaunchType.Server || LaunchType == ProcessLaunchType.ProxyServer,
 					Client: LaunchType == ProcessLaunchType.DirectClient || LaunchType == ProcessLaunchType.ProxyClient
 				);
 
