[Gauntlet] Add -McpRegion and -McpSubRegion epic config args
#tests Local runs, prf not relevant

diff --git a/Engine/Source/Programs/AutomationTool/Gauntlet/Unreal/Game/EpicGame.TestConfig.cs b/Engine/Source/Programs/AutomationTool/Gauntlet/Unreal/Game/EpicGame.TestConfig.cs
--- a/Engine/Source/Programs/AutomationTool/Gauntlet/Unreal/Game/EpicGame.TestConfig.cs
+++ b/Engine/Source/Programs/AutomationTool/Gauntlet/Unreal/Game/EpicGame.TestConfig.cs
@@ -95,6 +95,18 @@ namespace EpicGame
 		[AutoParam]
 		public bool RemoteServer = false;
 
+		/// <summary>
+		/// Optional specifier for client mcp region
+		/// </summary>
+		[AutoParam]
+		public string McpRegion = string.Empty;
+
+		/// <summary>
+		/// Optional specifier for client mcp subregion
+		/// </summary>
+		[AutoParam]
+		public string McpSubRegion = string.Empty;
+
 		/// Set traces to collect for editor roles, saved to a generated file.
 		/// </summary>
 		[AutoParam]
@@ -240,6 +252,17 @@ namespace EpicGame
 					{
 						McpString += string.Format(" -epicapp={0} -buildidoverride={1}", EpicApp, BuildIDOverride);
 					}
+
+					bool bRequestRegion = !string.IsNullOrEmpty(McpRegion);
+					if (ConfigRole.RoleType.IsClient() && bRequestRegion)
+					{
+						McpString += string.Format(" -McpRegion={0}", McpRegion);
+
+						if (bRequestRegion && !string.IsNullOrEmpty(McpSubRegion))
+						{
+							McpString += string.Format(" -McpSubRegion={0}", McpSubRegion);
+						}
+					}
 				}
 
 				if (FastCook)
