Update ReportItemizedExecutableCode to emit a sections report and write them to the build metadata folder.

#rnx

diff --git a/Engine/Source/Programs/AutomationTool/Scripts/ReportItemizedExecutableCode.cs b/Engine/Source/Programs/AutomationTool/Scripts/ReportItemizedExecutableCode.cs
--- a/Engine/Source/Programs/AutomationTool/Scripts/ReportItemizedExecutableCode.cs
+++ b/Engine/Source/Programs/AutomationTool/Scripts/ReportItemizedExecutableCode.cs
@@ -37,8 +37,24 @@ namespace ReportItemizedExecutableCode.Automation
 	[Help("Reports itemized binary size in a format consumed by PerfReportServer")]
 	[ParamHelp("ExeToItemize", "Absolute path to the binary to itemize", ParamType = typeof(FileReference))]
 	[ParamHelp("SymbolsExe", "Absolute path to the symbolicated binary", ParamType = typeof(FileReference))]
+	[ParamHelp("ReportPath", "Absolute path to the directory to write the reports to.", ParamType = typeof(DirectoryReference))]
 	class ReportItemizedExecutableCode : BuildCommand
 	{
+		private string[] GetReportDestDirectories()
+		{
+			List<string> DestPaths = new List<string>()
+			{
+				CmdEnv.LogFolder
+			};
+
+			DirectoryReference ReportDir = ParseOptionalDirectoryReferenceParam("ReportPath");
+			if (ReportDir != null)
+			{
+				DestPaths.Add(ReportDir.FullName);
+			}
+			return DestPaths.ToArray();
+		}
+
 		/** Attempts to assign a reasonable module name from a source file name */
 		string FindModuleName(string SourcePath)
 		{
@@ -167,97 +183,77 @@ namespace ReportItemizedExecutableCode.Automation
 		protected ModuleInfos ItemizeExecutableWithBloaty(FileReference Binary, FileReference SymbolsBinary)
 		{
 			ModuleInfos ModInfos = new ModuleInfos();
-			string BloatyExePath = Unreal.RootDirectory.ToString() + @"\Engine\Extras\ThirdPartyNotUE\Bloaty\bloaty.exe";
 
-			if (File.Exists(BloatyExePath))
+			if (FileReference.Exists(Binary))
 			{
-				if (FileReference.Exists(Binary))
+				string Args = string.Format("-d compileunits -n 0 --csv {0}", Binary.FullName);
+				if (SymbolsBinary != null && FileReference.Exists(SymbolsBinary))
 				{
-					string Args = string.Format("-d compileunits -n 0 --csv {0}", Binary.FullName);
-					if (SymbolsBinary != null && FileReference.Exists(SymbolsBinary))
-					{
-						Args += string.Format(" --debug-file={0}", SymbolsBinary.FullName);
-					}
+					Args += string.Format(" --debug-file={0}", SymbolsBinary.FullName);
+				}
 
-					// Hack!
-					//string Results = HardCodedResults;
+				// Hack!
+				//string[] LineArray = RunBloaty(Args, HardCodedResults);
 
-					int ExitCode = 0;
-					string Results = UnrealBuildTool.Utils.RunLocalProcessAndReturnStdOut(BloatyExePath, Args, out ExitCode);
+				string[] LineArray = RunBloaty(Args);
 
-					if (ExitCode == 0)
+				// since we told bloaty to output csv, the result will be already machine readable. Just make sure that the first line is as we expect
+				const string ExpectedCSVOutput = "compileunits,vmsize,filesize";
+				if (LineArray[0] != ExpectedCSVOutput)
+				{
+					throw new AutomationException("CSV output ({0}) does not match expectation ({1})", LineArray[0], ExpectedCSVOutput);
+				}
+
+				uint Total = 0;
+
+				for (int IdxLine = 1; IdxLine < LineArray.Length; ++IdxLine)
+				{
+					string[] LineElems = LineArray[IdxLine].Split(',');
+					if (LineElems.Length != 3)
 					{
-						// Get the results from Bloaty as an array of lines
-						string[] LineArray = Results.Split(new char[] { '\r', '\n' }, StringSplitOptions.RemoveEmptyEntries);
+						throw new AutomationException("CSV line ({0}) does not match expectation ({1}) - comma in source file?", LineArray[IdxLine], ExpectedCSVOutput);
+					}
 
-						// since we told bloaty to output csv, the result will be already machine readable. Just make sure that the first line is as we expect
-						const string ExpectedCSVOutput = "compileunits,vmsize,filesize";
-						if (LineArray[0] != ExpectedCSVOutput)
+					string ModuleName = SanitizeModuleName(FindModuleName(LineElems[0]));
+					uint Size = uint.Parse(LineElems[1]);
+					Total += Size;
+
+					if (ValidateModuleName(ModuleName, LineArray[IdxLine]))
+					{
+						if (ModInfos.Modules.ContainsKey(ModuleName))
 						{
-							throw new AutomationException("CSV output ({0}) does not match expectation ({1})", LineArray[0], ExpectedCSVOutput);
+							ModInfos.Modules[ModuleName] += Size;
 						}
-
-						uint Total = 0;
-						
-						for (int IdxLine = 1; IdxLine < LineArray.Length; ++IdxLine)
+						else
 						{
-							string[] LineElems = LineArray[IdxLine].Split(',');
-							if (LineElems.Length != 3)
-							{
-								throw new AutomationException("CSV line ({0}) does not match expectation ({1}) - comma in source file?", LineArray[IdxLine], ExpectedCSVOutput);
-							}
-
-							string ModuleName = SanitizeModuleName(FindModuleName(LineElems[0]));
-							uint Size = uint.Parse(LineElems[1]);
-							Total += Size;
-
-							if (ValidateModuleName(ModuleName, LineArray[IdxLine]))
-							{
-								if (ModInfos.Modules.ContainsKey(ModuleName))
-								{
-									ModInfos.Modules[ModuleName] += Size;
-								}
-								else
-								{
-									ModInfos.Modules.Add(ModuleName, Size);
-								}
-
-								if (ModInfos.ModuleNameToSources.ContainsKey(ModuleName))
-								{
-									ModInfos.ModuleNameToSources[ModuleName].Add(LineElems[0]);
-								}
-								else
-								{
-									List<string> Sources = new List<string>();
-									Sources.Add(LineElems[0]);
-									ModInfos.ModuleNameToSources.Add(ModuleName, Sources);
-								}
-							}
-							else
-							{
-								ModInfos.SourcesWithoutModule.Add(LineArray[IdxLine]);
-								ModInfos.TotalSizeForSourcesWithoutModule += Size;
-							}
+							ModInfos.Modules.Add(ModuleName, Size);
 						}
 
-						// add a _Total metric
-						ModInfos.Modules.Add("_Total", Total);
-
-						return ModInfos;
+						if (ModInfos.ModuleNameToSources.ContainsKey(ModuleName))
+						{
+							ModInfos.ModuleNameToSources[ModuleName].Add(LineElems[0]);
+						}
+						else
+						{
+							List<string> Sources = new List<string>();
+							Sources.Add(LineElems[0]);
+							ModInfos.ModuleNameToSources.Add(ModuleName, Sources);
+						}
 					}
 					else
 					{
-						throw new AutomationException("bloaty failed with exit code {0} and output '{1}'", ExitCode, Results);
+						ModInfos.SourcesWithoutModule.Add(LineArray[IdxLine]);
+						ModInfos.TotalSizeForSourcesWithoutModule += Size;
 					}
 				}
-				else
-				{
-					throw new AutomationException("binary to itemize is not found at '{0}'", Binary.FullName);
-				}
+
+				// add a _Total metric
+				ModInfos.Modules.Add("_Total", Total);
+				return ModInfos;
 			}
 			else
 			{
-				throw new AutomationException("bloaty.exe is not found at '{0}'", BloatyExePath);
+				throw new AutomationException("binary to itemize is not found at '{0}'", Binary.FullName);
 			}
 		}
 
@@ -270,6 +266,18 @@ namespace ReportItemizedExecutableCode.Automation
 		public override void ExecuteBuild()
 		{
 			FileReference ExecutableFileToItemize = ParseRequiredFileReferenceParam("ExeToItemize");
+
+			string[] DestDirs = GetReportDestDirectories();
+
+			string ItemizedReportName = ExecutableFileToItemize.GetFileNameWithoutExtension() + "_ItemizedModules.csv";
+			WriteItemizedReport(ExecutableFileToItemize, DestDirs.Select(DirPath => Path.Combine(DirPath, ItemizedReportName)).ToArray());
+
+			string SectionsReportName = ExecutableFileToItemize.GetFileNameWithoutExtension() + "_Sections.csv";
+			WriteSectionsReport(ExecutableFileToItemize, DestDirs.Select(DirPath => Path.Combine(DirPath, SectionsReportName)).ToArray());
+		}
+
+		private void WriteItemizedReport(FileReference ExecutableFileToItemize, string[] DestPaths)
+		{
 			FileReference ExecutableFileWithSymbols = ParseOptionalFileReferenceParam("SymbolsExe");
 
 			// Hack for quick iteration on formatting (run bloaty manually and just point this at the csv file)
@@ -343,17 +351,92 @@ namespace ReportItemizedExecutableCode.Automation
 				System.Console.WriteLine("{0}, {1:F2}, {2}", ModInfo.Key, ModInfo.Value / (1024.0 * 1024.0), ModInfo.Value);
 			}
 
-			string CSVName = ExecutableFileToItemize.GetFileNameWithoutExtension() + "_ItemizedModules.csv";
-			CSVName = Path.Combine(CmdEnv.LogFolder, CSVName);
-			
-			System.Console.WriteLine($"\nWriting itemized module information to: {CSVName}");
+			foreach (string DestPath in DestPaths)
+			{
+				System.Console.WriteLine($"\nWriting itemized module information to: {DestPath}");
+
+				string DirPath = Path.GetDirectoryName(DestPath);
+				Directory.CreateDirectory(DirPath);
 
-			File.WriteAllText(CSVName, Builder.ToString());
+				File.WriteAllText(DestPath, Builder.ToString());
 
-			System.Console.WriteLine("\n\n");
+				System.Console.WriteLine("\n\n");
+			}
 		}
 
 		// Hack for a quick iteration on formatting - stores the results from bloaty.
 		//string HardCodedResults;
+
+		private void WriteSectionsReport(FileReference ExecutableFile, string[] DestPaths)
+		{
+			string Args = string.Format("-s vm --csv {0}", ExecutableFile.FullName);
+			string[] Lines = RunBloaty(Args);
+
+			string ExpectedColumns = "sections,vmsize,filesize";
+			if (Lines[0] != ExpectedColumns)
+			{
+				throw new AutomationException("CSV output ({0}) does not match expectation ({1})", Lines[0], ExpectedColumns);
+			}
+
+			StringBuilder Builder = new StringBuilder("Section,Megabytes,Bytes\n");
+			foreach (string Line in Lines.Skip(1))
+			{
+				string[] LineValues = Line.Split(",");
+				string SectionName = SanitizeModuleName(LineValues[0]).Replace("-", "");
+
+				uint VMSize = uint.Parse(LineValues[1]);
+
+				Builder.AppendLine(string.Format("{0},{1:F2},{2}", SectionName, VMSize / (1024.0 * 1024.0), VMSize));
+			}
+
+			foreach (string DestPath in DestPaths)
+			{
+				System.Console.WriteLine($"\nWriting sections report to: {DestPath}");
+
+				string DirPath = Path.GetDirectoryName(DestPath);
+				Directory.CreateDirectory(DirPath);
+
+				File.WriteAllText(DestPath, Builder.ToString());
+
+				System.Console.WriteLine("\n\n");
+			}
+		}
+
+		/** Runs bloaty with the given args and returns an array of the output. */
+		protected static string[] RunBloaty(string Args, string ResultsOverride = null)
+		{
+			string BloatyExePath = Unreal.RootDirectory.ToString() + @"\Engine\Extras\ThirdPartyNotUE\Bloaty\bloaty.exe";
+
+			if (File.Exists(BloatyExePath))
+			{
+				int ExitCode = 0;
+
+				string Results = null;
+				if (ResultsOverride != null)
+				{
+					Results = ResultsOverride;
+				}
+				else
+				{
+					Results = UnrealBuildTool.Utils.RunLocalProcessAndReturnStdOut(BloatyExePath, Args, out ExitCode);
+				}
+
+				if (ExitCode == 0)
+				{
+					// Get the results from Bloaty as an array of lines
+					string[] LineArray = Results.Split(new char[] { '\r', '\n' }, StringSplitOptions.RemoveEmptyEntries);
+
+					return LineArray;
+				}
+				else
+				{
+					throw new AutomationException("bloaty failed with exit code {0} and output '{1}'", ExitCode, Results);
+				}
+			}
+			else
+			{
+				throw new AutomationException("bloaty.exe is not found at '{0}'", BloatyExePath);
+			}
+		}
 	}
 }
