Improve how the arguments are shown for a job.

diff --git a/Engine/Source/Programs/AutomationTool/Gauntlet/Framework/Gauntlet.TestExecutor.cs b/Engine/Source/Programs/AutomationTool/Gauntlet/Framework/Gauntlet.TestExecutor.cs
--- a/Engine/Source/Programs/AutomationTool/Gauntlet/Framework/Gauntlet.TestExecutor.cs
+++ b/Engine/Source/Programs/AutomationTool/Gauntlet/Framework/Gauntlet.TestExecutor.cs
@@ -5,6 +5,7 @@ using System.Collections.Generic;
 using System.IO;
 using System.Linq;
 using System.Net;
+using System.Text.RegularExpressions;
 using System.Threading;
 using EpicGames.Core;
 using Gauntlet.Utils;
@@ -659,6 +660,10 @@ namespace Gauntlet
 			// display how the test could be ran locally
 			string RunLocalString = TestInfo.TestNode.GetRunLocalCommand(BuildCommandThatLaunchedExecutor);
 			Log.Info("How to run locally: ({0})", RunLocalString);
+			
+			// Also, showing the arguments each in its row, so it is easier to read.
+			string FormattedArgs = Regex.Replace(RunLocalString, @" (?=[^-]*(-|$))", "\n\t");
+			Log.Info("Arguments used:\n\t{0}", FormattedArgs);
 		}
 
 		/// <summary>
