Create directory for Android Gauntlet logcat artifact

diff --git a/Engine/Source/Programs/AutomationTool/Gauntlet/Platform/Android/Gauntlet.TargetDeviceAndroid.cs b/Engine/Source/Programs/AutomationTool/Gauntlet/Platform/Android/Gauntlet.TargetDeviceAndroid.cs
--- a/Engine/Source/Programs/AutomationTool/Gauntlet/Platform/Android/Gauntlet.TargetDeviceAndroid.cs
+++ b/Engine/Source/Programs/AutomationTool/Gauntlet/Platform/Android/Gauntlet.TargetDeviceAndroid.cs
@@ -1617,6 +1617,10 @@ namespace Gauntlet
 			string SavedDirectory = Path.Combine(Install.AndroidDevice.LocalCachePath, "Saved");
 			string LogcatFilename = "Logcat.log";
 			// Save logcat dump to local artifact path.
+			if (!Directory.Exists(SavedDirectory))
+			{
+				Directory.CreateDirectory(SavedDirectory);
+			}
 			File.WriteAllText(Path.Combine(SavedDirectory, LogcatFilename), LogcatResult.Output);
 			bHaveSavedArtifacts = true;
 		}
