Fixed issues with cloud artifact uploads for release builds (were no uproject file exists)

diff --git a/Engine/Source/Programs/AutomationTool/BuildGraph/Tasks/CreateCloudArtifactTask.cs b/Engine/Source/Programs/AutomationTool/BuildGraph/Tasks/CreateCloudArtifactTask.cs
--- a/Engine/Source/Programs/AutomationTool/BuildGraph/Tasks/CreateCloudArtifactTask.cs
+++ b/Engine/Source/Programs/AutomationTool/BuildGraph/Tasks/CreateCloudArtifactTask.cs
@@ -61,7 +61,7 @@ namespace AutomationTool.Tasks
 		/// The uproject this artifact is created for
 		/// </summary>
 		[TaskParameter(Optional = true)]
-		public FileReference Project { get; set; } = null!;
+		public FileReference? Project { get; set; } = null!;
 
 		/// <summary>
 		/// Override the project name, defaults to the filename of the uproject
@@ -178,17 +178,19 @@ namespace AutomationTool.Tasks
 			Dictionary<string, object> metadata = new Dictionary<string, object>();
 			string platform = SanitizeBucketValue(_parameters.Platform);
 			
-			string project = _parameters.Project.GetFileNameWithoutExtension();
 			FileReference? projectPath = _parameters.Project;
-			if (!File.Exists(_parameters.Project.FullName))
+			string project = projectPath != null ? projectPath.GetFileNameWithoutExtension() : "";
+
+			bool projectPathExists = projectPath != null && File.Exists(projectPath.FullName);
+			if (!projectPathExists)
 			{
-				// if the project is not an uproject that exists we do the scanning for engine inis without the project specific overrides
 				projectPath = null;
 				Logger.LogInformation("\'{Project}\' is not a valid uproject file. Unable to load inis, you will need to specify project options explicitly.", _parameters.Project);
 				project = _parameters.ProjectNameOverride;
 				if (String.IsNullOrEmpty(project))
 				{
 					Logger.LogError("No project name override set and unable to locate uproject file '{Project}'. Unable to determine what the project name should be.", _parameters.Project);
+					throw new AutomationException("Unable to determine what the project name should be");
 				}
 			}
 
