LaunchMultiServer support for multiple proxy servers

Addeed "-proxyservercount=#" that sets how many proxy servers to run.  When used with "-proxyclientcount=#" or "-proxybotcount=#", clients and/or bots are assigned to each proxy in evenly distributed fashion.  Note that currently proxy servers are limited to 4 connected clients.

diff --git a/Engine/Source/Programs/AutomationTool/Scripts/LaunchMultiServer.Automation.cs b/Engine/Source/Programs/AutomationTool/Scripts/LaunchMultiServer.Automation.cs
--- a/Engine/Source/Programs/AutomationTool/Scripts/LaunchMultiServer.Automation.cs
+++ b/Engine/Source/Programs/AutomationTool/Scripts/LaunchMultiServer.Automation.cs
@@ -295,7 +295,7 @@ namespace AutomationTool
 					}
 					else
 					{
-						Logger.LogInformation("MultiServer server defintion with duplicate ID {0}, ignoring duplicate entry {1}", ConfigServerDefs[i].LocalId, ConfigServerDefs[i]);
+						Logger.LogInformation("MultiServer server definition with duplicate ID {0}, ignoring duplicate entry {1}", ConfigServerDefs[i].LocalId, ConfigServerDefs[i]);
 					}
 				}
 
@@ -314,43 +314,71 @@ namespace AutomationTool
 						}
 						else if (ClientValue == "proxy")
 						{
-							int ClientCount = 1;
+							int ClientCountTotal = 1;
 							if (ParseParam("proxyclientcount"))
 							{
 								string ClientCountStr = ParseParamValue("proxyclientcount", null);
-								Int32.TryParse(ClientCountStr, out ClientCount);
+								Int32.TryParse(ClientCountStr, out ClientCountTotal);
 							}
 
-							var ProxyListeningPort = 8000;
-
-							ProjectGameConfig.TryGetValue(ProxyConfigSection, "ListenPort", out ProxyListeningPort);
-
-							Logger.LogInformation("Starting proxy server on port {0} connected to MultiServer instances", ProxyListeningPort);
-							StartProcessesForServerInfos(UniqueServerInfos, ProcessLaunchType.ProxyServer, String.Format(" -port={0} ", ProxyListeningPort));
+							int BotCountTotal = 0;
+							if (ParseParam("proxybotcount"))
+							{
+								string ProxyBotCountStr = ParseParamValue("proxybotcount", null);
+								Int32.TryParse(ProxyBotCountStr, out BotCountTotal);
+							}
 
-							for (int ClientIndex = 0; ClientIndex < ClientCount; ClientIndex++)
+							int ProxyServerCount = 1;
+							if (ParseParam("proxyservercount"))
 							{
-								var ProxyClientAdditionalParams = String.Format("127.0.0.1:{0} -log=DstmProxyClient-{1}.log -SaveWinPos={1}", ProxyListeningPort, (ClientIndex+1));
-								StartProcessesForServerInfos(UniqueServerInfos, ProcessLaunchType.ProxyClient, ProxyClientAdditionalParams);
+								string ProxyServerCountStr = ParseParamValue("proxyservercount", null);
+								Int32.TryParse(ProxyServerCountStr, out ProxyServerCount);
 							}
 
-							if (ParseParam("proxybotcount"))
+							var FirstProxyListeningPort = 8000;
+							ProjectGameConfig.TryGetValue(ProxyConfigSection, "ListenPort", out FirstProxyListeningPort);
+
+							int TotalClientCount = 0;
+							int TotalBotCount = 0;
+
+							var ProjectName = ProjectUtils.GetShortProjectName(ProjectPath);
+
+							var ProxyListeningPort = FirstProxyListeningPort;
+							for (int ProxyServerIndex = 0; ProxyServerIndex < ProxyServerCount; ++ProxyServerIndex)
 							{
-								int BotCount = 0;
-								string BotCountStr = ParseParamValue("proxybotcount", null);
-								Int32.TryParse(BotCountStr, out BotCount);
+								Logger.LogInformation("Starting proxy server on port {0} connected to MultiServer instances", ProxyListeningPort);
+								{
+									var ProxyServerAdditionalParams = String.Format(" -port={0} -ConsoleTitle=\"{1} Proxy Server {2}\" -SessionName=\"{1} Proxy Server {2}\" -log={1}ProxyServer-{2}.log ", ProxyListeningPort, ProjectName, ProxyServerIndex + 1);
+									StartProcessesForServerInfos(UniqueServerInfos, ProcessLaunchType.ProxyServer, ProxyServerAdditionalParams);
+								}
+
+								int NumClientsPerServer = ClientCountTotal / ProxyServerCount;
+								int NumExtraClients = ClientCountTotal % ProxyServerCount;
+								int NumClientsToSpawn = NumClientsPerServer + (ProxyServerIndex < NumExtraClients ? 1 : 0);
+								for (int ClientIndex = 0; ClientIndex < NumClientsToSpawn; ClientIndex++)
+								{
+									++TotalClientCount;
+									var ProxyClientAdditionalParams = String.Format("127.0.0.1:{0} -ConsoleTitle=\"{1} Proxy Client {2}\" -SessionName=\"{1} Proxy Client {2}\" -log={1}ProxyClient-{2}.log -SaveWinPos={2} ", ProxyListeningPort, ProjectName, TotalClientCount);
+					
+									StartProcessesForServerInfos(UniqueServerInfos, ProcessLaunchType.ProxyClient, ProxyClientAdditionalParams);
+								}
 								
-								for (int BotIndex = 0; BotIndex < BotCount; BotIndex++)
+								int NumBotsPerServer = BotCountTotal / ProxyServerCount;
+								int NumExtraBots = BotCountTotal % ProxyServerCount;
+								int NumBotsToSpawn = NumBotsPerServer + (ProxyServerIndex < NumExtraBots ? 1 : 0);
+								for (int BotIndex = 0; BotIndex < NumBotsToSpawn; BotIndex++)
 								{
-									int BotNumber = BotIndex + 1;
-									var BotClientAdditionalParams = String.Format("127.0.0.1:{0} -log=DstmProxyBot-{1}.log -log -nullrhi -bot -unattended -nosplash ", ProxyListeningPort, BotNumber);
-									BotClientAdditionalParams += String.Format("-ConsoleTitle=\"{0} Proxy Bot Client {1}\" ", ProjectUtils.GetShortProjectName(ProjectPath), BotNumber);
+									++TotalBotCount;
+									var BotClientAdditionalParams = String.Format("127.0.0.1:{0} -log={2}ProxyBot-{1}.log -log -nullrhi -bot -unattended -nosplash ", ProxyListeningPort, TotalBotCount, ProjectName);
+									BotClientAdditionalParams += String.Format("-ConsoleTitle=\"{0} Proxy Bot Client {1}\" -SessionName=\"{0} Proxy Bot Client {1}\" ", ProjectName, TotalBotCount);
 									if (!ParseParam("nonewconsole"))
 									{
 										BotClientAdditionalParams += " -newconsole ";
 									}
 									StartProcessesForServerInfos(UniqueServerInfos, ProcessLaunchType.ProxyClient, BotClientAdditionalParams);
 								}
+								
+								++ProxyListeningPort;
 							}
 						}
 					}
@@ -399,11 +427,6 @@ namespace AutomationTool
 			if (LaunchType == ProcessLaunchType.Server || LaunchType == ProcessLaunchType.ProxyServer)
 			{
 				CommonArgs += Params.MapToRun + " -server -log ";
-
-				// Here are some examples of adding parameters for debugging purposes (note the trailing space within the quotes)
-				// CommonArgs += "-dpcvars=net.MaxConstructedPartialBunchSizeBytes=131071,dstm.SimpleCoinBots=0 ";
-				// CommonArgs += "-logcmds=\"LogActor Verbose, LogRemoteObject Verbose, LogDstmPushMigration VeryVerbose, LogRepCompares Log, LogNet Verbose\" ";
-
 				if (!ParseParam("nonewconsole"))
 				{
 					CommonArgs += " -newconsole ";
@@ -427,14 +450,14 @@ namespace AutomationTool
 
 						if (LaunchType == ProcessLaunchType.Server)
 						{
-							Args += String.Format("-MultiServerLocalId={0} -MultiServerListenPort={1} -port={2} -log=DstmMultiServer-{0}.log ", ServerInfo.Definition.LocalId, ServerInfo.Definition.ListenPort, ServerInfo.GameListenPort);
+							Args += String.Format("-MultiServerLocalId={0} -MultiServerListenPort={1} -port={2} -log={3}MultiServer-{0}.log ", ServerInfo.Definition.LocalId, ServerInfo.Definition.ListenPort, ServerInfo.GameListenPort, Params.ShortProjectName);
 							Args += String.Format("-ConsoleTitle=\"{0} Server ID {1}\" -SessionName=\"{0} Server ID {1}\" ", Params.ShortProjectName, ServerInfo.Definition.LocalId);
 							Args += Params.ServerCommandline + " ";
 						}
 						else if (LaunchType == ProcessLaunchType.DirectClient)
 						{
 							Args += String.Format("127.0.0.1:{0} ", ServerInfo.GameListenPort) + (Params.Cook ? "" : "-game ");
-							Args += String.Format(" -log=DstmDirectClient-{0}.log", ServerInfo.Definition.LocalId);
+							Args += String.Format(" -log={1}DirectClient-{0}.log", ServerInfo.Definition.LocalId, Params.ShortProjectName);
 							Args += String.Format(" -SessionName=\"{0} Direct Client | Server ID {1}\" ", Params.ShortProjectName, ServerInfo.Definition.LocalId);
 							Args += Params.ClientCommandline + " ";
 
@@ -482,9 +505,6 @@ namespace AutomationTool
 
 					if (ServerInfos.Count() > 0)
 					{
-						Args += String.Format("-ConsoleTitle=\"{0} Proxy Server\" -SessionName=\"{0} Proxy Server\" ", Params.ShortProjectName);
-						Args += "-log=DstmProxyServer.log ";
-
 						// Example of adding additional debug parameters for the proxy
 						// Args += "-logcmds=\"LogNetProxy VeryVerbose\" ";
 
@@ -554,8 +574,6 @@ namespace AutomationTool
 					Args += Params.ClientCommandline + " ";
 					Args += "-windowed ";
 					
-					Args += String.Format(" -SessionName=\"{0} Proxy Client \" ", Params.ShortProjectName);
-					
 					if (ParseParam("notimeouts"))
 					{
 						Args += "-notimeouts ";
