Upload Build Metadata To PRS when builds are published.

+ Refactored the PerfReportServerImporter to support the new inline manifests.
+ Added a UAT script to create an inline import manifest containing build metadata provded by the build graph.
+ Added this node to the graph after most standard target publish steps.

#rnx

diff --git a/Engine/Source/Programs/AutomationTool/Gauntlet/Framework/Utils/Gauntlet.CsvImporter.cs b/Engine/Source/Programs/AutomationTool/Gauntlet/Framework/Utils/Gauntlet.CsvImporter.cs
--- a/Engine/Source/Programs/AutomationTool/Gauntlet/Framework/Utils/Gauntlet.CsvImporter.cs
+++ b/Engine/Source/Programs/AutomationTool/Gauntlet/Framework/Utils/Gauntlet.CsvImporter.cs
@@ -2,30 +2,74 @@
 
 using System;
 using System.Collections.Generic;
+using System.Diagnostics.CodeAnalysis;
 using System.Linq;
+using System.IO;
 
 namespace Gauntlet
 {
+	public enum CsvImportType
+	{
+		/// <summary>
+		/// Standard import pointing to a csv file.
+		/// </summary>
+		CsvFile,
+
+		/// <summary>
+		/// An import that has no csv and any metadata is provided via the manifest.
+		/// </summary>
+		Inline
+	}
+
+	public interface ICsvMetadata
+	{
+		/// <summary>
+		/// Get as a completed dict;
+		/// </summary>
+		public Dictionary<string, string> ToDict();
+	}
+
 	/// <summary>
 	/// Interface for a generic Csv Import.
 	/// CsvImportEntry should suit most needs, but if you need something custom you can implement this interface.
 	/// </summary>
 	public interface ICsvImport
 	{
+		/// <summary>
+		/// The type of import.
+		/// </summary>
+		public CsvImportType ImportType { get; }
+
+		/// <summary>
+		/// Friendly name of the import for logging purposes.
+		/// </summary>
+		public string ImportName { get; }
+
 		/// <summary>
 		/// The path to the csv file.
 		/// </summary>
-		public string CsvFilename { get; set; }
-		
+		[AllowNull]
+		public string CsvFilename { get; }
+
 		/// <summary>
-		/// The path to the log file.
+		/// The path to the log file (Optional).
 		/// </summary>
+		[AllowNull]
 		public string LogFilename { get; set; }
-		
+
 		/// <summary>
 		/// Any additional metadata associated with this csv.
+		/// Corresponds to 'CustomFields' in the datasource config.
 		/// </summary>
-		public Dictionary<string, string> Metadata { get; }
+		public Dictionary<string, string> CustomFields { get; }
+
+		/// <summary>
+		/// Metadata for this import which varies on the import type.
+		/// For CsvFile imports, this is additional metadata that's included with the csv's metadata (not currently implemented).
+		/// For inline imports, this is the only source of metadata as there's no csv file.
+		/// </summary>
+		[AllowNull]
+		public ICsvMetadata Metadata { get; set; }
 
 		/// <summary>
 		/// Any additional files that should be bundled with this csv.
@@ -46,16 +90,63 @@ namespace Gauntlet
 	/// </summary>
 	public class CsvImportEntry : ICsvImport
 	{
-		public string CsvFilename { get; set; }
-		public string LogFilename { get; set; }
-		public Dictionary<string, string> Metadata { get; }
+		public CsvImportType ImportType { get; private set; } = CsvImportType.CsvFile;
+		public string ImportName { get; private set; }
+
+		public string CsvFilename { get; private set; } = null;
+		public string LogFilename { get; set; } = null;
+		public Dictionary<string, string> CustomFields { get; private set; }
+		public ICsvMetadata Metadata { get; set; } = null;
 		public IReadOnlyDictionary<string, string> AdditionalFiles { get; private set; }
 
-		public CsvImportEntry(string CsvFilename, string LogFilename = null, Dictionary<string, string> Metadata = null)
+		/// <summary>
+		/// CsvFile import constructor.
+		/// Deprecated: Please use CreateCsvFileImportEntry.
+		/// </summary>
+		/// <param name="CsvFilename">The path to the csv to import.</param>
+		/// <param name="LogFilename">Optional path to a corresponding log to bundle with this csv.</param>
+		/// <param name="CustomFields">Optional metadata corresponding to 'CustomFields' in the datasource config.</param>
+		public CsvImportEntry(string CsvFilename, string LogFilename = null, Dictionary<string, string> CustomFields = null)
 		{
+			this.ImportType = CsvImportType.CsvFile;
+			this.ImportName = Path.GetFileName(CsvFilename);
 			this.CsvFilename = CsvFilename;
 			this.LogFilename = LogFilename;
-			this.Metadata = Metadata ?? new Dictionary<string, string>();
+			this.CustomFields = CustomFields ?? new Dictionary<string, string>();
+		}
+
+		/// <summary>
+		/// Creates an import entry for a csv file.
+		/// </summary>
+		/// <param name="CsvFilename">The path to the csv to import.</param>
+		/// <param name="LogFilename">Optional path to a corresponding log to bundle with this csv.</param>
+		/// <param name="CustomFields">Optional metadata corresponding to 'CustomFields' in the datasource config.</param>
+		public static CsvImportEntry CreateCsvFileImportEntry(string CsvFilename, string LogFilename = null, Dictionary<string, string> CustomFields = null)
+		{
+			return new CsvImportEntry(CsvFilename, LogFilename, CustomFields);
+		}
+
+		/// <summary>
+		/// Creates an inline Import entry.
+		/// </summary>
+		/// <param name="ImportName">Identifier for this import, primarily for logging.</param>
+		/// <param name="InlineMetadata">Metadata that would normally be provided by a csv. Must contain necessary fields.</param>
+		/// <param name="CustomFields">Optional datasource specific metadata. Corresponds to "CustomFields" in datasource config.</param>
+		public static CsvImportEntry CreateInlineImportEntry(string ImportName, ICsvMetadata CsvMetadata, Dictionary<string, string> CustomFields = null)
+		{
+			if (CsvMetadata == null)
+			{
+				throw new ArgumentNullException("CsvMetadata", "Inline imports must have valid metadata.");
+			}
+
+			return new CsvImportEntry()
+			{
+				ImportType = CsvImportType.Inline,
+				ImportName = ImportName,
+				Metadata = CsvMetadata,
+				CustomFields = CustomFields ?? new Dictionary<string, string>(),
+				CsvFilename = null
+			};
 		}
 
 		public void AddAdditionalFile(string FileType, string Filename)
@@ -72,6 +163,10 @@ namespace Gauntlet
 			MutableAdditionalFiles[FileType] = Filename;
 			AdditionalFiles = MutableAdditionalFiles;
 		}
+
+		private CsvImportEntry()
+		{
+		}
 	}
 
 	/// <summary>
