Adding -noconsole to the launch script which will disable logging to the consoles, but still allow logging to the files.  With VeryVerbose logging, the console output makes the games unplayable.  You may load up the files in an Editor like Notepad++ which will auto-reload to debug.

#rnx

diff --git a/Engine/Source/Programs/AutomationTool/Scripts/LaunchMultiServer.Automation.cs b/Engine/Source/Programs/AutomationTool/Scripts/LaunchMultiServer.Automation.cs
--- a/Engine/Source/Programs/AutomationTool/Scripts/LaunchMultiServer.Automation.cs
+++ b/Engine/Source/Programs/AutomationTool/Scripts/LaunchMultiServer.Automation.cs
@@ -427,10 +427,22 @@ namespace AutomationTool
 			if (LaunchType == ProcessLaunchType.Server || LaunchType == ProcessLaunchType.ProxyServer)
 			{
 				CommonArgs += Params.MapToRun + " -server -log ";
+
+				// Here are some examples of adding parameters for debugging purposes (note the trailing space within the quotes)
+				// CommonArgs += "-dpcvars=net.MaxConstructedPartialBunchSizeBytes=131071,dstm.SimpleCoinBots=0 ";
+				// CommonArgs += "-logcmds=\"LogActor Verbose, LogRemoteObject Verbose, LogDstmPushMigration VeryVerbose, LogRepCompares Log, LogNet Verbose\" ";
+
 				if (!ParseParam("nonewconsole"))
 				{
 					CommonArgs += " -newconsole ";
 				}
+
+				// Use this option when logging with VeryVerbose to avoid unusable debugging windows (too much spew)
+				// The logging will still go to the output files (and you can use a text editor like Notepad++ to auto-reload them)
+				if (ParseParam("noconsole"))
+				{
+					CommonArgs += " -nodebugoutput -noconsole ";
+				}
 			}
 
 			CommonArgs += " -messaging ";
