Implement timeout on ADB and AFS command

#rnx

diff --git a/Engine/Source/Programs/AutomationTool/Gauntlet/Platform/Android/Gauntlet.TargetDeviceAndroid.cs b/Engine/Source/Programs/AutomationTool/Gauntlet/Platform/Android/Gauntlet.TargetDeviceAndroid.cs
--- a/Engine/Source/Programs/AutomationTool/Gauntlet/Platform/Android/Gauntlet.TargetDeviceAndroid.cs
+++ b/Engine/Source/Programs/AutomationTool/Gauntlet/Platform/Android/Gauntlet.TargetDeviceAndroid.cs
@@ -99,6 +99,9 @@ namespace Gauntlet
 		/// </summary>
 		public string DeviceLogPath { get; protected set; }
 
+		[AutoParam((int)(60 * 15))]
+		public int MaxInstallTime { get; protected set; }
+
 		/// <summary>
 		/// Path to a command line if installed
 		/// </summary>
@@ -112,6 +115,8 @@ namespace Gauntlet
 
 		public TargetDeviceAndroid(string InDeviceName = "", AndroidDeviceData DeviceData = null, string InCachePath = null)
 		{
+			AutoParam.ApplyParamsAndDefaults(this, Globals.Params.AllArguments);
+
 			AdbCredentialCache.AddInstance(DeviceData);
 
 			// If no device name or its 'default' then use the first default device
@@ -547,24 +552,24 @@ namespace Gauntlet
 		/// <param name="bShouldLogCommand"></param>
 		/// <param name="bPauseErrorParsing"></param>
 		/// <returns></returns>
-		public IProcessResult RunAdbDeviceCommand(string Args, bool Wait = true, bool bShouldLogCommand = false, bool bPauseErrorParsing = false)
+		public IProcessResult RunAdbDeviceCommand(string Args, bool Wait = true, bool bShouldLogCommand = false, bool bPauseErrorParsing = false, int TimeoutSec = 60 * 15)
 		{
 			if (string.IsNullOrEmpty(DeviceName) == false)
 			{
 				Args = string.Format("-s {0} {1}", DeviceName, Args);
 			}
 
-			return RunAdbGlobalCommand(Args, Wait, bShouldLogCommand, bPauseErrorParsing);
+			return RunAdbGlobalCommand(Args, Wait, bShouldLogCommand, bPauseErrorParsing, TimeoutSec);
 		}
 
-		public IProcessResult RunAFSDeviceCommand(string Args, bool Wait = true, bool bShouldLogCommand = false, bool bPauseErrorParsing = false)
+		public IProcessResult RunAFSDeviceCommand(string Args, bool Wait = true, bool bShouldLogCommand = false, bool bPauseErrorParsing = false, int TimeoutSec = 60 * 15)
 		{
 			if (string.IsNullOrEmpty(DeviceName) == false)
 			{
 				Args = string.Format("-s {0} {1}", DeviceName, Args);
 			}
 
-			return RunAFSGlobalCommand(Args, Wait, bShouldLogCommand, bPauseErrorParsing);
+			return RunAFSGlobalCommand(Args, Wait, bShouldLogCommand, bPauseErrorParsing, TimeoutSec);
 		}
 
 		/// <summary>
@@ -598,7 +603,7 @@ namespace Gauntlet
 		/// <param name="bShouldLogCommand"></param>
 		/// <param name="bPauseErrorParsing"></param>
 		/// <returns></returns>
-		public static IProcessResult RunAdbGlobalCommand(string Args, bool Wait = true, bool bShouldLogCommand = false, bool bPauseErrorParsing = false)
+		public static IProcessResult RunAdbGlobalCommand(string Args, bool Wait = true, bool bShouldLogCommand = false, bool bPauseErrorParsing = false, int TimeoutSec = 60 * 15)
 		{
 			CommandUtils.ERunOptions RunOptions = CommandUtils.ERunOptions.AppMustExist | CommandUtils.ERunOptions.NoWaitForExit | CommandUtils.ERunOptions.SpewIsVerbose;
 			RunOptions |= Log.IsVeryVerbose? ERunOptions.AllowSpew : ERunOptions.NoLoggingOfRunCommand;
@@ -616,7 +621,15 @@ namespace Gauntlet
 
 				if (Wait)
 				{
-					Process.WaitForExit();
+					Process.ProcessObject.WaitForExit(TimeoutSec * 1000);
+
+					if (!Process.HasExited)
+					{
+						Log.Info("ADB Command 'adb {Args}' timed out after {Timeout} sec.", Args, TimeoutSec);
+
+						Process.ProcessObject?.Kill();
+						Process.ProcessObject?.WaitForExit(15000);
+					}
 				}
 			}
 
@@ -641,7 +654,7 @@ namespace Gauntlet
 			return new LongProcessResult(AdbCommand, Args, RunOptions, LocalCache: LocalCache);
 		}
 
-		public static IProcessResult RunAFSGlobalCommand(string Args, bool Wait = true, bool bShouldLogCommand = false, bool bPauseErrorParsing = false)
+		public static IProcessResult RunAFSGlobalCommand(string Args, bool Wait = true, bool bShouldLogCommand = false, bool bPauseErrorParsing = false, int TimeoutSec = 60 * 15 )
 		{
 			CommandUtils.ERunOptions RunOptions = CommandUtils.ERunOptions.AppMustExist | CommandUtils.ERunOptions.NoWaitForExit | CommandUtils.ERunOptions.SpewIsVerbose;
 
@@ -667,7 +680,15 @@ namespace Gauntlet
 
 				if (Wait)
 				{
-					Process.WaitForExit();
+					Process.ProcessObject.WaitForExit(TimeoutSec * 1000);
+
+					if (!Process.HasExited)
+					{
+						Log.Info("AFS Command 'UnrealAndroidFileTool {Args}' timed out after {Timeout} sec.", Args, TimeoutSec);
+
+						Process.ProcessObject?.Kill();
+						Process.ProcessObject?.WaitForExit(15000);
+					}
 				}
 			}
 
@@ -876,14 +897,14 @@ namespace Gauntlet
 					if (UseAFS)
 					{
 						Log.Info("Copying {0} to {1} via AFS", QuotedSourcePath, DestPath);
-						AdbResult = RunAFSDeviceCommand(string.Format("-p \"{0}\" -k \"{1}\" push \"{2}\" \"{3}\"", PackageName, AFSToken, SourcePath, DestPath));
+						AdbResult = RunAFSDeviceCommand(string.Format("-p \"{0}\" -k \"{1}\" push \"{2}\" \"{3}\"", PackageName, AFSToken, SourcePath, DestPath), TimeoutSec: MaxInstallTime);
 					}
 					else
 					{
 						RunAdbDeviceCommand(string.Format("shell mkdir -p {0}/", FileDirectory));
 
 						Log.Info("Copying {0} to {1} via adb push", QuotedSourcePath, DestPath);
-						AdbResult = RunAdbDeviceCommand(string.Format("push {0} {1}", QuotedSourcePath, DestPath));
+						AdbResult = RunAdbDeviceCommand(string.Format("push {0} {1}", QuotedSourcePath, DestPath), TimeoutSec: MaxInstallTime);
 					}
 
 					// Note: Presently, AdbResult NEVER reports failures to push when using AFS. Should be fixed at some point.
@@ -900,12 +921,12 @@ namespace Gauntlet
 								if (UseAFS)
 								{
 									Log.Info("Retrying to copy via AFS...");
-									AdbResult = RunAFSDeviceCommand(string.Format("-p {0} -k {1} push \"{2}\" \"{3}\"", PackageName, AFSToken, SourcePath, DestPath));
+									AdbResult = RunAFSDeviceCommand(string.Format("-p {0} -k {1} push \"{2}\" \"{3}\"", PackageName, AFSToken, SourcePath, DestPath), TimeoutSec: MaxInstallTime);
 								}
 								else
 								{
 									Log.Info("Retrying to copy via adb push...");
-									AdbResult = RunAdbDeviceCommand(string.Format("push {0} {1}", QuotedSourcePath, DestPath));
+									AdbResult = RunAdbDeviceCommand(string.Format("push {0} {1}", QuotedSourcePath, DestPath), TimeoutSec: MaxInstallTime);
 								}
 								if (AdbResult.ExitCode != 0)
 								{
