Fix build automation & module discovery for merged modules

- Use platform-specific decorated separator when applicable on merged targets
- Allow platforms to have a more complex module discovery method
- Update test and tools to reflect those changes

#rnx

diff --git a/Engine/Source/Programs/AutomationTool/AutomationUtils/ProjectParams.cs b/Engine/Source/Programs/AutomationTool/AutomationUtils/ProjectParams.cs
--- a/Engine/Source/Programs/AutomationTool/AutomationUtils/ProjectParams.cs
+++ b/Engine/Source/Programs/AutomationTool/AutomationUtils/ProjectParams.cs
@@ -423,6 +423,7 @@ namespace AutomationTool
 			this.NumClients = InParams.NumClients;
 			this.Compressed = InParams.Compressed;
 			this.ForceUncompressed = InParams.ForceUncompressed;
+			this.UseMergedModules = InParams.UseMergedModules;
 			this.AdditionalPakOptions = InParams.AdditionalPakOptions;
 			this.AdditionalIoStoreOptions = InParams.AdditionalIoStoreOptions;
 			this.ForceOodleDllVersion = InParams.ForceOodleDllVersion;
@@ -499,6 +500,7 @@ namespace AutomationTool
 			bool? Clean = null,
 			bool? Compressed = null,
 			bool? ForceUncompressed = null,
+			bool? UseMergedModules = null,
 			string AdditionalPakOptions = null,
 			string AdditionalIoStoreOptions = null,
 			string ForceOodleDllVersion = null,
@@ -820,6 +822,7 @@ namespace AutomationTool
 			this.SkipCookingEditorContent = GetParamValueIfNotSpecified(Command, SkipCookingEditorContent, this.SkipCookingEditorContent, "SkipCookingEditorContent");
 			this.Compressed = GetParamValueIfNotSpecified(Command, Compressed, this.Compressed, "compressed");
 			this.ForceUncompressed = GetParamValueIfNotSpecified(Command, ForceUncompressed, this.ForceUncompressed, "ForceUncompressed");
+			this.UseMergedModules = GetParamValueIfNotSpecified(Command, UseMergedModules, this.UseMergedModules, "UseMergedModules");
 			this.AdditionalPakOptions = ParseParamValueIfNotSpecified(Command, AdditionalPakOptions, "AdditionalPakOptions");
 			if (!string.IsNullOrEmpty(this.NoZenAutoLaunch))
 			{
@@ -1961,6 +1964,11 @@ namespace AutomationTool
 		/// Do not compress packages during cook, override game ProjectPackagingSettings to force it off
 		/// </summary>
 		public bool ForceUncompressed;
+
+		/// <summary>
+		/// Assume that executables will be built with merged modules
+		/// </summary>
+		public bool UseMergedModules;
 		
 		/// <summary>
 		/// Compress packages during cook, override game ProjectPackagingSettings and Platform Hardware Compression settings to force it on
diff --git a/Engine/Source/Programs/AutomationTool/Scripts/CopyBuildToStagingDirectory.Automation.cs b/Engine/Source/Programs/AutomationTool/Scripts/CopyBuildToStagingDirectory.Automation.cs
--- a/Engine/Source/Programs/AutomationTool/Scripts/CopyBuildToStagingDirectory.Automation.cs
+++ b/Engine/Source/Programs/AutomationTool/Scripts/CopyBuildToStagingDirectory.Automation.cs
@@ -5866,7 +5866,20 @@ namespace AutomationScripts
 					string Exe = Target;
 					if (Config != UnrealTargetConfiguration.Development)
 					{
-						Exe = Target + "-" + PlatformName + "-" + Config.ToString();
+						string Separator = "-";
+
+						// We need to account for the decorated separator, rather than hardcoding the full file pattern
+						foreach (var ProjectTarget in Params.ProjectTargets)
+						{
+							if (Target == ProjectTarget.TargetName && Params.UseMergedModules)
+							{
+								TargetRules ValidatedRules = ProjectTarget.Rules;
+								TargetRules.ValidateMergedTarget(ValidatedRules, StagePlatform.Type);
+								Separator = ValidatedRules.DecoratedSeparator;
+								break;
+							}
+						}
+						Exe = Target + Separator + PlatformName + Separator + Config.ToString();
 					}
 
 					// append the architecture part if needed
diff --git a/Engine/Source/Programs/UnrealBuildTool/Configuration/TargetRules.cs b/Engine/Source/Programs/UnrealBuildTool/Configuration/TargetRules.cs
--- a/Engine/Source/Programs/UnrealBuildTool/Configuration/TargetRules.cs
+++ b/Engine/Source/Programs/UnrealBuildTool/Configuration/TargetRules.cs
@@ -3138,6 +3138,17 @@ namespace UnrealBuildTool
 			return rules;
 		}
 
+		/// <summary>
+		/// Validate the target rules for its current platform when using modular build
+		/// </summary>
+		/// <param name="target">Rules to validate</param>
+		/// <param name="platform">Explicit platform to use</param>
+		public static void ValidateMergedTarget(TargetRules target, UnrealTargetPlatform platform)
+		{
+			UEBuildPlatform BuildPlatform = UEBuildPlatform.GetBuildPlatform(platform);
+			BuildPlatform.ValidateMergedTarget(target);
+		}
+		
 		/// <summary>
 		/// Constructor.
 		/// </summary>
diff --git a/Engine/Source/Programs/UnrealBuildTool/Configuration/UEBuildPlatform.cs b/Engine/Source/Programs/UnrealBuildTool/Configuration/UEBuildPlatform.cs
--- a/Engine/Source/Programs/UnrealBuildTool/Configuration/UEBuildPlatform.cs
+++ b/Engine/Source/Programs/UnrealBuildTool/Configuration/UEBuildPlatform.cs
@@ -790,6 +790,17 @@ namespace UnrealBuildTool
 		/// Validate a target's settings
 		/// </summary>
 		public virtual void ValidateTarget(TargetRules Target)
+		{
+			if (Target.bMergeModules)
+			{
+				ValidateMergedTarget(Target);
+			}
+		}
+
+		/// <summary>
+		/// Validate a target's settings for merged modules
+		/// </summary>
+		public virtual void ValidateMergedTarget(TargetRules Target)
 		{
 		}
 
diff --git a/Engine/Source/Programs/UnrealBuildTool/Platform/Android/UEBuildAndroid.cs b/Engine/Source/Programs/UnrealBuildTool/Platform/Android/UEBuildAndroid.cs
--- a/Engine/Source/Programs/UnrealBuildTool/Platform/Android/UEBuildAndroid.cs
+++ b/Engine/Source/Programs/UnrealBuildTool/Platform/Android/UEBuildAndroid.cs
@@ -191,6 +191,8 @@ namespace UnrealBuildTool
 
 		public override void ValidateTarget(TargetRules Target)
 		{
+			base.ValidateTarget(Target);
+
 			if (!String.IsNullOrWhiteSpace(Environment.GetEnvironmentVariable("CLANG_STATIC_ANALYZER_MODE")))
 			{
 				Target.StaticAnalyzer = StaticAnalyzer.Default;
diff --git a/Engine/Source/Programs/UnrealBuildTool/Platform/IOS/UEBuildIOS.cs b/Engine/Source/Programs/UnrealBuildTool/Platform/IOS/UEBuildIOS.cs
--- a/Engine/Source/Programs/UnrealBuildTool/Platform/IOS/UEBuildIOS.cs
+++ b/Engine/Source/Programs/UnrealBuildTool/Platform/IOS/UEBuildIOS.cs
@@ -749,6 +749,8 @@ namespace UnrealBuildTool
 
 		public override void ValidateTarget(TargetRules Target)
 		{
+			base.ValidateTarget(Target);
+
 			if (!String.IsNullOrWhiteSpace(Environment.GetEnvironmentVariable("CLANG_STATIC_ANALYZER_MODE")))
 			{
 				Target.StaticAnalyzer = StaticAnalyzer.Default;
diff --git a/Engine/Source/Programs/UnrealBuildTool/Platform/Linux/UEBuildLinux.cs b/Engine/Source/Programs/UnrealBuildTool/Platform/Linux/UEBuildLinux.cs
--- a/Engine/Source/Programs/UnrealBuildTool/Platform/Linux/UEBuildLinux.cs
+++ b/Engine/Source/Programs/UnrealBuildTool/Platform/Linux/UEBuildLinux.cs
@@ -208,6 +208,8 @@ namespace UnrealBuildTool
 
 		public override void ValidateTarget(TargetRules Target)
 		{
+			base.ValidateTarget(Target);
+
 			if (!String.IsNullOrWhiteSpace(Environment.GetEnvironmentVariable("CLANG_STATIC_ANALYZER_MODE")))
 			{
 				Target.StaticAnalyzer = StaticAnalyzer.Default;
diff --git a/Engine/Source/Programs/UnrealBuildTool/Platform/Mac/UEBuildMac.cs b/Engine/Source/Programs/UnrealBuildTool/Platform/Mac/UEBuildMac.cs
--- a/Engine/Source/Programs/UnrealBuildTool/Platform/Mac/UEBuildMac.cs
+++ b/Engine/Source/Programs/UnrealBuildTool/Platform/Mac/UEBuildMac.cs
@@ -337,7 +337,9 @@ namespace UnrealBuildTool
 		}
 
 		public override void ValidateTarget(TargetRules Target)
-		{			
+		{
+			base.ValidateTarget(Target);
+
 			if (!String.IsNullOrWhiteSpace(Environment.GetEnvironmentVariable("CLANG_STATIC_ANALYZER_MODE")))
 			{
 				Target.StaticAnalyzer = StaticAnalyzer.Default;
diff --git a/Engine/Source/Programs/UnrealBuildTool/Platform/Windows/UEBuildWindows.cs b/Engine/Source/Programs/UnrealBuildTool/Platform/Windows/UEBuildWindows.cs
--- a/Engine/Source/Programs/UnrealBuildTool/Platform/Windows/UEBuildWindows.cs
+++ b/Engine/Source/Programs/UnrealBuildTool/Platform/Windows/UEBuildWindows.cs
@@ -1203,6 +1203,8 @@ namespace UnrealBuildTool
 		/// </summary>
 		public override void ValidateTarget(TargetRules Target)
 		{
+			base.ValidateTarget(Target);
+
 			if (Platform == UnrealTargetPlatform.Win64)
 			{
 				Target.WindowsPlatform.Architecture = Target.Architecture;// == UnrealArch.Default ? UnrealArch.X64 : Target.Architecture;
diff --git a/Engine/Source/Runtime/Core/Private/GenericPlatform/GenericPlatformProcess.cpp b/Engine/Source/Runtime/Core/Private/GenericPlatform/GenericPlatformProcess.cpp
--- a/Engine/Source/Runtime/Core/Private/GenericPlatform/GenericPlatformProcess.cpp
+++ b/Engine/Source/Runtime/Core/Private/GenericPlatform/GenericPlatformProcess.cpp
@@ -232,6 +232,11 @@ const FString FGenericPlatformProcess::GetModulesDirectory()
 	return FPaths::Combine(*FPaths::EngineDir(), TEXT("Binaries"), FPlatformProcess::GetBinariesSubdirectory());
 }
 
+bool FGenericPlatformProcess::ModuleExists(const FString& Filename)
+{
+	return FPaths::FileExists(Filename);
+}
+
 void FGenericPlatformProcess::LaunchURL( const TCHAR* URL, const TCHAR* Parms, FString* Error )
 {
 	UE_LOG(LogHAL, Fatal, TEXT("FGenericPlatformProcess::LaunchURL not implemented on this platform"));
diff --git a/Engine/Source/Runtime/Core/Private/Modules/ModuleManager.cpp b/Engine/Source/Runtime/Core/Private/Modules/ModuleManager.cpp
--- a/Engine/Source/Runtime/Core/Private/Modules/ModuleManager.cpp
+++ b/Engine/Source/Runtime/Core/Private/Modules/ModuleManager.cpp
@@ -355,7 +355,7 @@ bool FModuleManager::IsModuleUpToDate(const FName InModuleName) const
 
 	for (const TPair<FName, FString>& Pair : ModulePathMap)
 	{
-		if (!FPaths::FileExists(Pair.Value))
+		if (!FPlatformProcess::ModuleExists(Pair.Value))
 		{
 			return false;
 		}
@@ -779,7 +779,7 @@ IModuleInterface* FModuleManager::LoadModuleWithFailureReason(const FName InModu
 	{
 #if !IS_MONOLITHIC
 		ModuleInfoRef ModuleInfo = FoundModulePtr.ToSharedRef();
-		if (ModuleInfo->Filename.IsEmpty() || !FPaths::FileExists(ModuleInfo->Filename))
+		if (ModuleInfo->Filename.IsEmpty() || !FPlatformProcess::ModuleExists(ModuleInfo->Filename))
 		{
 			TMap<FName, FString> ModulePathMap;
 			FindModulePaths(*InModuleName.ToString(), ModulePathMap);
@@ -937,7 +937,7 @@ IModuleInterface* FModuleManager::LoadModuleWithFailureReason(const FName InModu
 
 		UE_LOG(LogModuleManager, Verbose, TEXT("ModuleManager: Load Module '%s' DLL '%s'"), *InModuleName.ToString(), *ModuleInfo->Filename);
 
-		if (ModuleInfo->Filename.IsEmpty() || !FPaths::FileExists(ModuleInfo->Filename))
+		if (ModuleInfo->Filename.IsEmpty() || !FPlatformProcess::ModuleExists(ModuleInfo->Filename))
 		{
 			TMap<FName, FString> ModulePathMap;
 			FindModulePaths(*InModuleName.ToString(), ModulePathMap);
@@ -961,7 +961,7 @@ IModuleInterface* FModuleManager::LoadModuleWithFailureReason(const FName InModu
 		ModuleInfo->Handle = nullptr;
 
 		// Skip this check if file manager has not yet been initialized
-		if (FPaths::FileExists(ModuleFileToLoad))
+		if (FPlatformProcess::ModuleExists(ModuleFileToLoad))
 		{
 			ModuleInfo->Handle = InternalLoadLibrary(InModuleName, ModuleFileToLoad);
 			
diff --git a/Engine/Source/Runtime/Core/Public/GenericPlatform/GenericPlatformProcess.h b/Engine/Source/Runtime/Core/Public/GenericPlatform/GenericPlatformProcess.h
--- a/Engine/Source/Runtime/Core/Public/GenericPlatform/GenericPlatformProcess.h
+++ b/Engine/Source/Runtime/Core/Public/GenericPlatform/GenericPlatformProcess.h
@@ -406,6 +406,13 @@ struct FGenericPlatformProcess
 	 * @see GetModulePrefix, GetModuleExtension
 	 */
 	static CORE_API const FString GetModulesDirectory();
+
+	/**
+	 * Used only by platforms with DLLs, this confirms the existence of a given module.
+	 * This should be called with the proper extension for modules.
+	 * Returning true should ensure that GetDllHandle() will find the same file.
+	 */
+	static CORE_API bool ModuleExists(const FString& Filename);
 	
 	/**
 	 * Launch a uniform resource locator (i.e. http://www.epicgames.com/unreal).
