Adding the ability to conditionally restart tests in Gauntlet

#rnx

diff --git a/Engine/Source/Programs/AutomationTool/Gauntlet/Unreal/Base/Gauntlet.UnrealTestNode.cs b/Engine/Source/Programs/AutomationTool/Gauntlet/Unreal/Base/Gauntlet.UnrealTestNode.cs
--- a/Engine/Source/Programs/AutomationTool/Gauntlet/Unreal/Base/Gauntlet.UnrealTestNode.cs
+++ b/Engine/Source/Programs/AutomationTool/Gauntlet/Unreal/Base/Gauntlet.UnrealTestNode.cs
@@ -267,7 +267,7 @@ namespace Gauntlet
 
 		static protected DateTime SessionStartTime = DateTime.MinValue;
 
-		public int Retries { get; private set; } = 0;
+		public int Retries { get; protected set; } = 0;
 		private int MaxRetries = 3;
 		public virtual bool CanRetry() { return Retries < MaxRetries; }
 		public virtual bool SetToRetryIfPossible()
@@ -278,8 +278,34 @@ namespace Gauntlet
 			SetUnrealTestResult(TestResult.WantRetry);
 			return true;
 		}
+
+		/// <summary>
+		/// PLEASE READ BEFORE OVERRIDING:
+		/// This should only be used in very specific instances to retry for user defined dependency/service failure strings
+		/// Example: Useful for unexpected unattended log outs during testing
+		/// </summary>
+		protected virtual bool CustomRetryLogic()
+		{
+			return false;
+		}
+
 		public virtual bool IsSetToRetry()
 		{
+			if (CustomRetryLogic())
+			{
+				if (CanRetry())
+				{
+					Retries++;
+					SetUnrealTestResult(TestResult.WantRetry);
+					Log.Error($"CustomRetryLogic returned true on retry {Retries}/{MaxRetries}");
+					return true;
+				}
+				else
+				{
+					Log.Error($"Max retries triggered by CustomRetryLogic has exceeded, test has failed {Retries}/{MaxRetries} times. Creating report.");
+					return false;
+				}
+			}
 			return GetTestResult() == TestResult.WantRetry;
 		}
 
@@ -1255,6 +1281,7 @@ namespace Gauntlet
 		/// <returns></returns>
 		public virtual ITestReport CreateReport(TestResult Result, UnrealTestContext Context, UnrealBuildSource Build, IEnumerable<UnrealRoleResult> RoleResults, string ArtifactPath)
 		{
+			if (IsSetToRetry()) { return null; }
 			if (GetCachedConfiguration().WriteTestResultsForHorde)
 			{
 				// write test report for Horde
