Changing Android to use non-bulk as default

#rnx

diff --git a/Engine/Source/Programs/AutomationTool/Gauntlet/Platform/Android/Gauntlet.TargetDeviceAndroid.cs b/Engine/Source/Programs/AutomationTool/Gauntlet/Platform/Android/Gauntlet.TargetDeviceAndroid.cs
--- a/Engine/Source/Programs/AutomationTool/Gauntlet/Platform/Android/Gauntlet.TargetDeviceAndroid.cs
+++ b/Engine/Source/Programs/AutomationTool/Gauntlet/Platform/Android/Gauntlet.TargetDeviceAndroid.cs
@@ -387,9 +387,6 @@ namespace Gauntlet
 			{
 				CopyOBBFiles(AppConfig, Build.FilesToInstall, Package);
 			}
-
-			// Now, create and push the UECommandline.txt file
-			CopyCommandlineFile(AppConfig, AppConfig.CommandLine, Package, AppConfig.ProjectName, Build.UsesExternalFilesDir);
 		}
 
 		public IAppInstall CreateAppInstall(UnrealAppConfig AppConfig)
@@ -409,7 +406,7 @@ namespace Gauntlet
 
 			PopulateDirectoryMappings(Path.GetDirectoryName(DeviceArtifactPath));
 
-			return new AndroidAppInstall(this, AppConfig.ProjectName, Build.AndroidPackageName, AppConfig.CommandLine, AppConfig.ProjectFile, AppConfig.Build.Configuration);
+			return new AndroidAppInstall(AppConfig, this, AppConfig.ProjectName, Build.AndroidPackageName, AppConfig.CommandLine, AppConfig.ProjectFile, AppConfig.Build.Configuration);
 		}
 
 		public void CopyAdditionalFiles(IEnumerable<UnrealFileToCopy> FilesToCopy)
@@ -442,6 +439,10 @@ namespace Gauntlet
 				PowerOn();
 			}
 
+			// Now, create and push the UECommandline.txt file
+			bool bUseExternalFileDirectory = (Install.AppConfig.Build as AndroidBuild)?.UsesExternalFilesDir ?? false;
+			CopyCommandlineFile(Install.AppConfig, Install.AppConfig.CommandLine, Install.AndroidPackageName, Install.AppConfig.ProjectName, bUseExternalFileDirectory);
+
 			// kill any currently running instance:
 			KillRunningProcess(Install.AndroidPackageName);
 
@@ -1418,7 +1419,7 @@ namespace Gauntlet
 
 			CopyCommandlineFile(AppConfig, AppConfig.CommandLine, Build.AndroidPackageName, AppConfig.ProjectName, Build.UsesExternalFilesDir);
 
-			AndroidAppInstall AppInstall = new AndroidAppInstall(this, InApkPath: ApkPath, AppConfig.ProjectName, Build.AndroidPackageName, AppConfig.CommandLine, AppConfig.ProjectFile, Build.Configuration);
+			AndroidAppInstall AppInstall = new AndroidAppInstall(AppConfig, this, InApkPath: ApkPath, AppConfig.ProjectName, Build.AndroidPackageName, AppConfig.CommandLine, AppConfig.ProjectFile, Build.Configuration);
 
 			return AppInstall;
 		}
@@ -1444,9 +1445,12 @@ namespace Gauntlet
 		public FileReference ProjectFile { get; protected set; }
 
 		public UnrealTargetConfiguration Configuration { get; protected set; }
+		
+		public UnrealAppConfig AppConfig { get; protected set; }
 
-		public AndroidAppInstall(TargetDeviceAndroid InDevice, string InName, string InAndroidPackageName, string InCommandLine, FileReference InProjectFile, UnrealTargetConfiguration InConfiguration, string InAppTag = "UE")
+		public AndroidAppInstall(UnrealAppConfig InAppConfig, TargetDeviceAndroid InDevice, string InName, string InAndroidPackageName, string InCommandLine, FileReference InProjectFile, UnrealTargetConfiguration InConfiguration, string InAppTag = "UE")
 		{
+			AppConfig = InAppConfig;
 			AndroidDevice = InDevice;
 			Name = InName;
 			AndroidPackageName = InAndroidPackageName;
@@ -1456,8 +1460,9 @@ namespace Gauntlet
 			Configuration = InConfiguration;
 		}
 
-		public AndroidAppInstall(TargetDeviceAndroid InDevice, string InApkPath, string InName, string InAndroidPackageName, string InCommandLine, FileReference InProjectFile, UnrealTargetConfiguration InConfiguration, string InAppTag = "UE")
+		public AndroidAppInstall(UnrealAppConfig InAppConfig, TargetDeviceAndroid InDevice, string InApkPath, string InName, string InAndroidPackageName, string InCommandLine, FileReference InProjectFile, UnrealTargetConfiguration InConfiguration, string InAppTag = "UE")
 		{
+			AppConfig = InAppConfig;
 			AndroidDevice = InDevice;
 			ApkPath = InApkPath;
 			Name = InName;
