[Gauntlet] Modifies build sources to ignore builds when the executable's platform does not match the requested platform

diff --git a/Engine/Source/Programs/AutomationTool/Gauntlet/Unreal/BuildSource/Gauntlet.StagedBuildSource.cs b/Engine/Source/Programs/AutomationTool/Gauntlet/Unreal/BuildSource/Gauntlet.StagedBuildSource.cs
--- a/Engine/Source/Programs/AutomationTool/Gauntlet/Unreal/BuildSource/Gauntlet.StagedBuildSource.cs
+++ b/Engine/Source/Programs/AutomationTool/Gauntlet/Unreal/BuildSource/Gauntlet.StagedBuildSource.cs
@@ -291,10 +291,16 @@ namespace Gauntlet
 
 				foreach (FileSystemInfo App in Binaries)
 				{
-					UnrealTargetConfiguration Config = UnrealHelpers.GetConfigurationFromExecutableName(InProjectName, App.Name);
-					UnrealTargetRole Role = UnrealHelpers.GetRoleFromExecutableName(InProjectName, App.Name);
-					string Flavor = UnrealHelpers.GetBuildFlavorFromExecutableName(InProjectName, App.Name);
+					UnrealHelpers.ConfigInfo ConfigInfo = UnrealHelpers.GetUnrealConfigFromFileName(InProjectName, App.Name);
+					UnrealTargetConfiguration Config = ConfigInfo.Configuration;
+					UnrealTargetRole Role = ConfigInfo.RoleType;
+					string Flavor = ConfigInfo.Flavor;
 
+					if (ConfigInfo.Platform.HasValue && InPlatform != ConfigInfo.Platform.Value)
+					{
+						Log.Verbose("{PackageName} platform {Platform} does not match, skipping.", App.Name, ConfigInfo.Platform.Value);
+						continue;
+					}
 
 					// if we have a verificatio BuildSource object, verify the build is usable before we even create the build
 					if (VerificationBuildSource == null || VerificationBuildSource.ShouldMakeBuildAvailable(App.FullName))
