LaunchMultiServer UAT script now passes -servercmdline argument to the proxy server.

diff --git a/Engine/Source/Programs/AutomationTool/Scripts/LaunchMultiServer.Automation.cs b/Engine/Source/Programs/AutomationTool/Scripts/LaunchMultiServer.Automation.cs
--- a/Engine/Source/Programs/AutomationTool/Scripts/LaunchMultiServer.Automation.cs
+++ b/Engine/Source/Programs/AutomationTool/Scripts/LaunchMultiServer.Automation.cs
@@ -557,6 +557,7 @@ namespace AutomationTool
 						}
 					}
 
+					Args += Params.ServerCommandline + " ";
 					Args += AdditionalParameters;
 
 					var NewProcess = Run(App, Args, null, ERunOptions.Default | ERunOptions.NoWaitForExit | ERunOptions.NoStdOutRedirect);
