Add ability to specify destination project ID in ZenImportSnapshotTask, and avoid swallowing exceptions when running commands.

[RN] none

diff --git a/Engine/Source/Programs/AutomationTool/BuildGraph/Tasks/ZenImportSnapshotTask.cs b/Engine/Source/Programs/AutomationTool/BuildGraph/Tasks/ZenImportSnapshotTask.cs
--- a/Engine/Source/Programs/AutomationTool/BuildGraph/Tasks/ZenImportSnapshotTask.cs
+++ b/Engine/Source/Programs/AutomationTool/BuildGraph/Tasks/ZenImportSnapshotTask.cs
@@ -15,8 +15,6 @@ using System.Xml;
 using EpicGames.Core;
 using EpicGames.ProjectStore;
 
-using Microsoft.Extensions.Logging;
-
 using UnrealBuildBase;
 
 namespace AutomationTool.Tasks
@@ -74,6 +72,12 @@ namespace AutomationTool.Tasks
 		/// </summary>
 		[TaskParameter(Optional = true)]
 		public int RemoteZenPort { get; set; } = 8558;
+
+		/// <summary>
+		/// Optional. Destination project ID to import snapshots into.
+		/// </summary>
+		[TaskParameter(Optional = true)]
+		public string DestinationProjectId { get; set; }
 	}
 
 	/// <summary>
@@ -229,7 +233,7 @@ namespace AutomationTool.Tasks
 			}
 			catch (CommandUtils.CommandFailedException e)
 			{
-				Logger.LogWarning("{Text}", e.ToString());
+				throw new AutomationException("Zen command failed: {0}", e.ToString());
 			}
 		}
 
@@ -395,7 +399,7 @@ namespace AutomationTool.Tasks
 				host = _parameters.RemoteZenHost + ":" + _parameters.RemoteZenPort.ToString();
 			}
 
-			string projectId = GetProjectID();
+			string projectId = String.IsNullOrEmpty(_parameters.DestinationProjectId) ? GetProjectID() : _parameters.DestinationProjectId;
 
 			DirectoryReference rootDir = Unreal.RootDirectory;
 			DirectoryReference engineDir = new DirectoryReference(Path.Combine(Unreal.RootDirectory.FullName, "Engine"));
