Add Initialization phase report in TestDataV2
Clearly identify Finalize report method
Avoid superfluous or redundant error when handling EngineTest report

#rnx

diff --git a/Engine/Source/Programs/AutomationTool/Gauntlet/Framework/Base/Gauntlet.TestReport.cs b/Engine/Source/Programs/AutomationTool/Gauntlet/Framework/Base/Gauntlet.TestReport.cs
--- a/Engine/Source/Programs/AutomationTool/Gauntlet/Framework/Base/Gauntlet.TestReport.cs
+++ b/Engine/Source/Programs/AutomationTool/Gauntlet/Framework/Base/Gauntlet.TestReport.cs
@@ -84,6 +84,11 @@ namespace Gauntlet
 		/// </summary>
 		/// <returns></returns>
 		Dictionary<string, object> GetReportDependencies();
+
+		/// <summary>
+		/// Call to finalize report
+		/// </summary>
+		void FinalizeReport();
 	}
 
 	/// <summary>
@@ -269,6 +274,10 @@ namespace Gauntlet
 		{
 			return new Dictionary<string, object>();
 		}
+
+		/// <inheritdoc/>
+		public virtual void FinalizeReport()
+		{ }
 	}
 
 	class JsonTryParseEnumConverter : JsonConverterFactory
diff --git a/Engine/Source/Programs/AutomationTool/Gauntlet/Framework/Utils/Gauntlet.HordeReport.cs b/Engine/Source/Programs/AutomationTool/Gauntlet/Framework/Utils/Gauntlet.HordeReport.cs
--- a/Engine/Source/Programs/AutomationTool/Gauntlet/Framework/Utils/Gauntlet.HordeReport.cs
+++ b/Engine/Source/Programs/AutomationTool/Gauntlet/Framework/Utils/Gauntlet.HordeReport.cs
@@ -1238,12 +1238,15 @@ namespace Gauntlet
 			public class TestDevice
 			{
 				public string name { get; set; }
+				[JsonIgnore]
+				public string key { get; private set; }
 				public string appInstanceLogPath { get; set; }
 				public Dictionary<string, string> metadata { get; set; }
 
-				public TestDevice(string InName)
+				public TestDevice(string InName, string InKey)
 				{
 					name = InName;
+					key = InKey;
 					metadata = new Dictionary<string, string>();
 				}
 			}
@@ -1394,7 +1397,7 @@ namespace Gauntlet
 				}
 
 				/// <summary>
-				/// Compile the aggregate error fingerprints
+				/// Compile the aggregate errors
 				/// </summary>
 				public void CompileErrorFingerprint()
 				{
@@ -1480,7 +1483,7 @@ namespace Gauntlet
 					Data.devices.Remove(Key);
 					Log.Verbose("Device (Key) '{Name}' already exists in the report. It is being replaced.", InName, Key);
 				}
-				TestDevice Device = new TestDevice(InName);
+				TestDevice Device = new TestDevice(InName, Key);
 				Data.devices.Add(Key, Device);
 				return Device;
 			}
@@ -1599,6 +1602,7 @@ namespace Gauntlet
 				static Regex Regex_numbertimestampstring = new Regex(@"\d[\d.,:]*", RegexOptions.Compiled);
 				// Redact path
 				static Regex Regex_pathstring = new Regex(@"(?<![A-Z])(?:[A-Z]:|/)[^ :]+[/\\]SYNC[/\\]", RegexOptions.Compiled);
+				static Regex Regex_pathstring2 = new Regex(@"(?:^|\s)/[\w/.:-]+", RegexOptions.Compiled);
 				// Redact hex strings
 				static Regex Regex_hexstring = new Regex(@"(?:0X|#)[0-9A-F]+", RegexOptions.Compiled);
 
@@ -1663,6 +1667,7 @@ namespace Gauntlet
 				{
 					Text = Text.Trim().ToUpperInvariant();
 					Text = Regex_pathstring.Replace(Text, "");
+					Text = Regex_pathstring2.Replace(Text, "");
 					Text = Regex_hexstring.Replace(Text, "H");
 					Text = Regex_numbertimestampstring.Replace(Text, "n");
 
@@ -2163,7 +2168,8 @@ namespace Gauntlet
 				return Phase.GetStream();
 			}
 
-			public override Dictionary<string, object> GetReportDependencies()
+			/// <inheritdoc/>
+			public override void FinalizeReport()
 			{
 				// store the event streams
 				string OutputDir = Path.Combine(OutputArtifactPath, "EventStreams");
@@ -2192,18 +2198,20 @@ namespace Gauntlet
 							Log.Error("Failed to save Event Stream for '{0}'. {1}", Phase.key, Ex);
 						}
 						// finalize error fingerprint
-						Phase.CompileErrorFingerprint(); 
+						Phase.CompileErrorFingerprint();
 					}
 					// update phase
 					Phase.eventStreamPath = OutputFilePath;
+					if (Phase.outcome == TestPhaseOutcome.Unknown)
+					{
+						Phase.SetOutcome(Phase.errorFingerprint == null ? TestStateType.Success : TestStateType.Fail);
+					}
 					// propagate phase tags to test tags
 					if (Phase.tags != null)
 					{
 						AddTestTags(Phase.tags);
 					}
 				}
-
-				return base.GetReportDependencies();
 			}
 
 			private static DateTime ConvertUETimeStringToDateTime(string StringTime)
@@ -2215,20 +2223,21 @@ namespace Gauntlet
 			/// Convert UnrealAutomatedTestPassResults to Horde data model
 			/// </summary>
 			/// <param name="InTestPassResults"></param>
-			/// <param name="InKey"></param>
-			/// <param name="InName"></param>
 			/// <param name="InReportPath"></param>
-			/// <param name="InHordeArtifactPath"></param>
 			/// <returns></returns>
-			public static AutomatedTestSessionData FromUnrealAutomatedTests(UnrealAutomatedTestPassResults InTestPassResults, string InKey, string InName, string InReportPath, string InHordeArtifactPath)
+			public void PopulateFromUnrealAutomatedTests(UnrealAutomatedTestPassResults InTestPassResults, string InReportPath)
 			{
-				AutomatedTestSessionData OutTestPassResults = new AutomatedTestSessionData(InKey, InName);
-				OutTestPassResults.SetOutputArtifactPath(InHordeArtifactPath);
+				if (string.IsNullOrEmpty(Data.summary.dateTime))
+				{
+					DateTime SessionTime = ConvertUETimeStringToDateTime(InTestPassResults.ReportCreatedOn);
+					SetSessionTiming(SessionTime, InTestPassResults.TotalDuration);
+				}
+
 				if (InTestPassResults.Devices != null)
 				{
 					foreach (UnrealAutomationDevice InDevice in InTestPassResults.Devices)
 					{
-						TestDevice ConvertedDevice = OutTestPassResults.AddDevice(InDevice.DeviceName, InDevice.Instance);
+						TestDevice ConvertedDevice = AddDevice(InDevice.DeviceName, InDevice.Instance);
 						ConvertedDevice.appInstanceLogPath = InDevice.AppInstanceLog;
 						ConvertedDevice.metadata.Add("platform", InDevice.Platform);
 						ConvertedDevice.metadata.Add("os_version", InDevice.OSVersion);
@@ -2240,13 +2249,11 @@ namespace Gauntlet
 						ConvertedDevice.metadata.Add("rhi", InDevice.RHI);
 					}
 				}
-				DateTime SessionTime = ConvertUETimeStringToDateTime(InTestPassResults.ReportCreatedOn);
-				OutTestPassResults.SetSessionTiming(SessionTime, InTestPassResults.TotalDuration);
 				if (InTestPassResults.Tests != null)
 				{
 					foreach (UnrealAutomatedTestResult InTestResult in InTestPassResults.Tests)
 					{
-						TestPhase Phase = OutTestPassResults.AddPhase(InTestResult.FullTestPath);
+						TestPhase Phase = AddPhase(InTestResult.FullTestPath);
 						Phase.deviceKeys = InTestResult.DeviceInstance.ToHashSet();
 						string TestDateTime = InTestResult.DateTime;
 						DateTime PhaseTime = ConvertUETimeStringToDateTime(TestDateTime);
@@ -2285,7 +2292,6 @@ namespace Gauntlet
 						}
 					}
 				}
-				return OutTestPassResults;
 			}
 		}
 
@@ -2314,6 +2320,7 @@ namespace Gauntlet
 				Dictionary<string, object> ExtraItems = null;
 				if (InData is BaseHordeReport InHordeReport)
 				{
+					InHordeReport.FinalizeReport();
 					ExtraItems = InHordeReport.GetReportDependencies();
 					InKey = InHordeReport.GetTestDataKey(InKey);
 					InData = InHordeReport.GetTestData();
diff --git a/Engine/Source/Programs/AutomationTool/Gauntlet/Unreal/Automation/UE.Automation.cs b/Engine/Source/Programs/AutomationTool/Gauntlet/Unreal/Automation/UE.Automation.cs
--- a/Engine/Source/Programs/AutomationTool/Gauntlet/Unreal/Automation/UE.Automation.cs
+++ b/Engine/Source/Programs/AutomationTool/Gauntlet/Unreal/Automation/UE.Automation.cs
@@ -957,7 +957,7 @@ namespace UE
 					{
 						ExitReason = string.Format("{0} of {1} test(s) failed", TestResults.Failed, TestResults.Tests.Count);
 						ExitCode = -1;
-						return UnrealProcessResult.TestFailure;
+						return UnrealProcessResult.EngineTestError;
 					}
 
 					// If no tests were run then that's a failure (possibly a bad RunTest argument?)
@@ -965,7 +965,7 @@ namespace UE
 					{
 						ExitReason = "No tests were executed!";
 						ExitCode = -1;
-						return UnrealProcessResult.TestFailure;
+						return UnrealProcessResult.InitializationFailure;
 					}
 				}
 			}
@@ -1109,7 +1109,7 @@ namespace UE
 				{
 					if (AllTests.Count == 0)
 					{
-						Log.Error(KnownLogEvents.Gauntlet_TestEvent, " * No tests were executed.");
+						Log.Error(KnownLogEvents.Gauntlet_UnrealEngineTestEvent, " * No tests were executed.");
 
 						IEnumerable<UnrealTestEvent> Errors = TestNodeEvents.Where(E => E.IsError).Distinct();
 						IEnumerable<UnrealTestEvent> Warnings = TestNodeEvents.Where(E => E.IsWarning).Distinct();
diff --git a/Engine/Source/Programs/AutomationTool/Gauntlet/Unreal/Base/Gauntlet.UnrealProcessResult.cs b/Engine/Source/Programs/AutomationTool/Gauntlet/Unreal/Base/Gauntlet.UnrealProcessResult.cs
--- a/Engine/Source/Programs/AutomationTool/Gauntlet/Unreal/Base/Gauntlet.UnrealProcessResult.cs
+++ b/Engine/Source/Programs/AutomationTool/Gauntlet/Unreal/Base/Gauntlet.UnrealProcessResult.cs
@@ -14,7 +14,7 @@ namespace Gauntlet
 		LoginFailed,					// Client never successfully logged in.
 		TimeOut,                        // A timeout occurred
 		TestFailure,                    // A test is known to have failed
-		UnrealError,                    // Unreal exited with an error code unrelated to test issues
+		EngineTestError,                // Unreal exited with an error code related to Engine Tests
 		Unknown,                        // Something not in the above
 	}
 }
diff --git a/Engine/Source/Programs/AutomationTool/Gauntlet/Unreal/Base/Gauntlet.UnrealTestNode.cs b/Engine/Source/Programs/AutomationTool/Gauntlet/Unreal/Base/Gauntlet.UnrealTestNode.cs
--- a/Engine/Source/Programs/AutomationTool/Gauntlet/Unreal/Base/Gauntlet.UnrealTestNode.cs
+++ b/Engine/Source/Programs/AutomationTool/Gauntlet/Unreal/Base/Gauntlet.UnrealTestNode.cs
@@ -1363,40 +1363,43 @@ namespace Gauntlet
 				float TimeElapse = (float)(DateTime.Now - SessionStartTime).TotalSeconds;
 				TestReport.SetSessionTiming(SessionTime, TimeElapse);
 				// Simple report has one phase
-				var Phase = TestReport.AddPhase("Main");
-				Phase.SetTiming(SessionTime, TimeElapse);
+				var InitPhase = TestReport.AddPhase("Initialization", "Initialization");
+				InitPhase.SetTiming(SessionTime, 0); // Unfortunately we are not tracking initialization duration at this point
+				var MainPhase = TestReport.AddPhase("Main");
+				MainPhase.SetTiming(SessionTime, TimeElapse);
 				// Propagate events
-				var Stream = Phase.GetStream();
-				foreach (UnrealTestEvent Event in TestNodeEvents)
-				{
-					Stream.AddEvent(Event);
-				}
 				if (RoleResults != null)
 				{
 					foreach (UnrealRoleResult Item in RoleResults)
 					{
 						string RoleType = Item.Artifacts.SessionRole.RoleType.ToString();
+						var RelevantPhase = (Item.ProcessResult == UnrealProcessResult.InitializationFailure || Item.ProcessResult == UnrealProcessResult.LoginFailed)?
+								InitPhase: MainPhase;
+						var RelevantStream = RelevantPhase.GetStream();
+
 						foreach (UnrealTestEvent Event in Item.Events)
 						{
-							Stream.AddEvent(Event);
+							RelevantStream.AddEvent(Event);
 						}
-						// Add Processed result
+
+						// Add Process result
 						if (Item.ProcessResult != UnrealProcessResult.ExitOk)
 						{
-							Stream.AddError($"{RoleType}: {Item.Summary}");
+							RelevantStream.AddError($"{RoleType}: {Item.Summary}");
 						}
 						// Since we are going through the Role, lets add the related device
-						ITargetDevice TargetDevice = Item.Artifacts.AppInstance.Device;
-						string DeviceName = TargetDevice.Platform == BuildHostPlatform.Current.Platform ? System.Environment.MachineName : TargetDevice.Name;
-						string DeviceKey = TargetDevice.Name.Replace(".", "-"); // Key must not contains dot
-						var ReportDevice = TestReport.AddDevice(DeviceName, DeviceKey);
-						ReportDevice.metadata.Add("Platform", TargetDevice.Platform.ToString());
-						ReportDevice.metadata.Add("Role", RoleType);
-						Phase.deviceKeys.Add(DeviceKey);
+						var ReportDevice = AddRoleDeviceToSessionReport(Item, TestReport);
+						InitPhase.deviceKeys.Add(ReportDevice.key);
+						MainPhase.deviceKeys.Add(ReportDevice.key);
 					}
 				}
+				var MainStream = MainPhase.GetStream();
+				foreach (UnrealTestEvent Event in TestNodeEvents)
+				{
+					MainStream.AddEvent(Event);
+				}
 				// Outcome
-				Phase.SetOutcome(GetTestResult());
+				MainPhase.SetOutcome(GetTestResult());
 				HordeTestReport = TestReport;
 				// Metadata
 				SetReportMetadata(HordeTestReport, new UnrealTestRole[] { GetCachedConfiguration().GetMainRequiredRole() });
@@ -1442,6 +1445,19 @@ namespace Gauntlet
 			return HordeTestReport;
 		}
 
+		private AutomatedTestSessionData.TestDevice AddRoleDeviceToSessionReport(UnrealRoleResult RoleResult, AutomatedTestSessionData TestReport)
+		{
+			ITargetDevice TargetDevice = RoleResult.Artifacts.AppInstance.Device;
+			string DeviceName = TargetDevice.Platform == BuildHostPlatform.Current.Platform ? System.Environment.MachineName : TargetDevice.Name;
+			string DeviceKey = TargetDevice.Name.Replace(".", "-"); // Key must not contains dot
+			string RoleType = RoleResult.Artifacts.SessionRole.RoleType.ToString();
+			AutomatedTestSessionData.TestDevice ReportDevice = TestReport.AddDevice(DeviceName, DeviceKey);
+			ReportDevice.metadata.Add("Platform", TargetDevice.Platform.ToString());
+			ReportDevice.metadata.Add("Role", RoleType);
+
+			return ReportDevice;
+		}
+
 		private SimpleTestReport CreateSimpleReportForHorde_legacy(TestResult Result)
 		{
 			SimpleTestReport SimpleReport = new SimpleTestReport(UnrealTestResult, Context, GetCachedConfiguration());
@@ -1478,23 +1494,74 @@ namespace Gauntlet
 		public virtual ITestReport CreateUnrealEngineTestPassReport(string UnrealAutomatedTestReportPath, string ReportURL)
 		{
 			if (IsSetToRetry()) { return null; }
+
+			BaseHordeReport HordeTestPassResults = null;
+			// Handle initialization phase
+			if (Globals.Params.ParseParam("UseTestDataV2"))
+			{
+				AutomatedTestSessionData SessionReport = new AutomatedTestSessionData(HordeReportTestKey, HordeReportTestName);
+				string ArtifactPathLabel = new DirectoryInfo(ArtifactPath).Name;
+				string HordeArtifactPath = Path.Combine(GetCachedConfiguration().HordeArtifactPath, ArtifactPathLabel);
+				SessionReport.SetOutputArtifactPath(HordeArtifactPath);
+				DateTime SessionTime = SessionStartTime.ToUniversalTime();
+				// Elapse time will be adjusted later if the initialization passed
+				float TimeElapse = (float)(DateTime.Now - SessionStartTime).TotalSeconds;
+				SessionReport.SetSessionTiming(SessionTime, TimeElapse);
+				var InitPhase = SessionReport.AddPhase("Initialization", "Initialization");
+				InitPhase.SetTiming(SessionTime, TimeElapse);
+				if (RoleResults != null)
+				{
+					foreach (UnrealRoleResult Item in RoleResults)
+					{
+						string RoleType = Item.Artifacts.SessionRole.RoleType.ToString();
+						if (Item.ProcessResult == UnrealProcessResult.InitializationFailure)
+						{
+							var EventStream = InitPhase.GetStream();
+							foreach (UnrealTestEvent Event in Item.Events)
+							{
+								EventStream.AddEvent(Event);
+							}
+							EventStream.AddError($"{RoleType}: {Item.Summary}");
+						}
+						// Since we are going through the Role, lets add the related device
+						var ReportDevice = AddRoleDeviceToSessionReport(Item, SessionReport);
+						InitPhase.deviceKeys.Add(ReportDevice.key);
+					}
+				}
+				HordeTestPassResults = SessionReport;
+			}
+
+			// With UE Test Automation, we care only for one role.
+			var MainRole = GetCachedConfiguration().GetMainRequiredRole();
+			var RoleList = new List<UnrealTestRole>() { MainRole };
+			string MainDeviceInstanceName = string.Empty;
+
 			string JsonReportPath = Path.Combine(UnrealAutomatedTestReportPath, "index.json");
-			if (File.Exists(JsonReportPath))
+			bool IsFileReportExist = File.Exists(JsonReportPath);
+			if (IsFileReportExist)
 			{
 				Log.Verbose("Reading json Unreal Automated test report from {Path}", JsonReportPath);
 				UnrealAutomatedTestPassResults JsonTestPassResults = UnrealAutomatedTestPassResults.LoadFromJson(JsonReportPath);
-				// With UE Test Automation, we care only for one role.
-				var MainRole = GetCachedConfiguration().GetMainRequiredRole();
-				var RoleList = new List<UnrealTestRole>() { MainRole };
+				MainDeviceInstanceName = JsonTestPassResults.Devices.Last().Instance;
 				// Convert test results for Horde
-				BaseHordeReport HordeTestPassResults = null;
-				if (Globals.Params.ParseParam("UseTestDataV2"))
-				{
-					string ArtifactPathLabel = new DirectoryInfo(ArtifactPath).Name;
-					string HordeArtifactPath = Path.Combine(GetCachedConfiguration().HordeArtifactPath, ArtifactPathLabel);
-					HordeTestPassResults = AutomatedTestSessionData.FromUnrealAutomatedTests(
-						JsonTestPassResults, HordeReportTestKey, HordeReportTestName, UnrealAutomatedTestReportPath, HordeArtifactPath
-					);
+				if (HordeTestPassResults is AutomatedTestSessionData SessionReport)
+				{ // Populate v2 if set before
+					var InitPhase = SessionReport.GetPhase("Initialization");
+					if (InitPhase != null)
+					{
+						if (JsonTestPassResults.Tests != null && JsonTestPassResults.Tests.Any())
+						{ // Adjust initialization timing since the json report was generated (initialization must have completed successfully)
+							DateTime FirstTestTime = UnrealAutomationEntry.GetTimestampAsDateTime(JsonTestPassResults.Tests.First().DateTime);
+							DateTime SessionTime = SessionStartTime.ToUniversalTime();
+							float TimeElapse = (float)(FirstTestTime - SessionTime).TotalSeconds;
+							InitPhase.timeElapseSec = TimeElapse;
+						}
+						else
+						{
+							InitPhase.GetStream().AddError("No tests were executed!");
+						}
+					}
+					SessionReport.PopulateFromUnrealAutomatedTests(JsonTestPassResults, UnrealAutomatedTestReportPath);
 				}
 				if (HordeTestPassResults == null || Globals.Params.ParseParam("AddTestDataV1"))
 				{
@@ -1524,6 +1591,25 @@ namespace Gauntlet
 					SetReportMetadata(CopyTestPassResults, RoleList); // Set the metadata to the old report
 					HordeTestPassResults.AttachDependencyReport(CopyTestPassResults, GetHordeTestDataKey_legacy());
 				}
+			}
+			else
+			{
+				Log.Warning("Could not find Unreal Automated test report at {FilePath}. Test report will be partial.", JsonReportPath);
+				if (HordeTestPassResults is AutomatedTestSessionData SessionReport)
+				{
+					var InitPhase = SessionReport.GetPhase("Initialization");
+					InitPhase.GetStream().AddWarning("Could not find Unreal Automated test report at {FilePath}. No test report details available.", JsonReportPath);
+					// Check if there is any Engine Test Error found to add more context
+					UnrealRoleResult RelevantRoleResult = RoleResults?.FirstOrDefault(I => I.ProcessResult == UnrealProcessResult.EngineTestError);
+					if (RelevantRoleResult != null)
+					{
+						InitPhase.GetStream().AddError(RelevantRoleResult.Summary);
+					}
+				}
+			}
+
+			if (HordeTestPassResults != null)
+			{ // Populate generic data and artifacts
 				// Metadata
 				SetReportMetadata(HordeTestPassResults, RoleList);
 				// Attached test log artifacts
@@ -1531,12 +1617,17 @@ namespace Gauntlet
 				{
 					foreach (UnrealRoleArtifacts Artifact in SessionArtifacts)
 					{
-						if (Artifact.SessionRole.RoleType == MainRole.Type && !string.IsNullOrEmpty(Artifact.LogPath))
+						string DeviceKey = Artifact.AppInstance.Device.Name.Replace(".", "-"); // Key must not contains dot
+						if (!string.IsNullOrEmpty(Artifact.LogPath))
 						{
-							string LogName = FileUtils.ConvertPathToUri(Path.GetRelativePath(Path.GetFullPath(Globals.Params.ParseParam("UseTestDataV2")? ArtifactPath : Context.Options.LogDir), Path.GetFullPath(Artifact.LogPath)));
+							string LogName = FileUtils.ConvertPathToUri(Path.GetRelativePath(Path.GetFullPath(Globals.Params.ParseParam("UseTestDataV2") ? ArtifactPath : Context.Options.LogDir), Path.GetFullPath(Artifact.LogPath)));
 							if (HordeTestPassResults is BaseHordeReport Report)
 							{
-								Report.AttachDeviceLog(JsonTestPassResults.Devices.Last().Instance, Artifact.LogPath, LogName);
+								Report.AttachDeviceLog(DeviceKey, Artifact.LogPath, LogName);
+								if (Artifact.SessionRole.RoleType == MainRole.Type && !string.IsNullOrEmpty(MainDeviceInstanceName))
+								{ // UE json report has a different Instance name than Gauntlet device
+									Report.AttachDeviceLog(MainDeviceInstanceName, Artifact.LogPath, LogName);
+								}
 							}
 							else
 							{
@@ -1545,13 +1636,9 @@ namespace Gauntlet
 						}
 					}
 				}
-				return HordeTestPassResults;
-			}
-			else
-			{
-				Log.Warning("Could not find Unreal Automated test report at {FilePath}. No Test report will be generated for Horde.", JsonReportPath);
-				return null;
 			}
+
+			return HordeTestPassResults;
 		}
 
 		/// <summary>
@@ -2145,7 +2232,7 @@ namespace Gauntlet
 
 			Log.Info(" #### Role: {0} ({1} {2})", RoleArtifacts.SessionRole.RoleType, RoleArtifacts.SessionRole.Platform, RoleArtifacts.SessionRole.Configuration);
 
-			bool HasFailed = InRoleResult.ProcessResult != UnrealProcessResult.ExitOk;
+			bool HasFailed = InRoleResult.ProcessResult != UnrealProcessResult.ExitOk && InRoleResult.ProcessResult != UnrealProcessResult.EngineTestError;
 			string RoleState = HasFailed ? "failed:" : "completed:";
 			string StatusMessage = string.Format(" #### {0} {1} {2} ({3}, ExitCode={4})", RoleArtifacts.SessionRole.RoleType, RoleState, InRoleResult.Summary, InRoleResult.ProcessResult, InRoleResult.ExitCode);
 			if (HasFailed)
