[Gauntlet] Touch up logging of HordeDeviceReservationService.ReleaseDevices

diff --git a/Engine/Source/Programs/AutomationTool/Gauntlet/Framework/Devices/Gauntlet.DeviceReservationService.cs b/Engine/Source/Programs/AutomationTool/Gauntlet/Framework/Devices/Gauntlet.DeviceReservationService.cs
--- a/Engine/Source/Programs/AutomationTool/Gauntlet/Framework/Devices/Gauntlet.DeviceReservationService.cs
+++ b/Engine/Source/Programs/AutomationTool/Gauntlet/Framework/Devices/Gauntlet.DeviceReservationService.cs
@@ -289,7 +289,7 @@ namespace Gauntlet
 			{
 				// If a user explicitly calls a service's release function with an incorrect device, throw this exception
 				string ExceptionMessage = "Attempted to release the following devices from a service that did not reserve them!";
-				ExceptionMessage += string.Join("\n\t", ThrowDevices.Select(Device => Device.Name));
+				ExceptionMessage += "\n\t" + string.Join("\n\t", ThrowDevices.Select(Device => Device.Name));
 				ExceptionMessage += "\nUse DevicePool.Instance.ReleaseDevices to avoid mistakingly releasing devices reserved from a different service";
 				throw new AutomationException(ExceptionMessage);
 			}
