[Gauntlet] Attempts to join the TestExecutor's starting threads when a cancel is signaled.

diff --git a/Engine/Source/Programs/AutomationTool/Gauntlet/Framework/Gauntlet.TestExecutor.cs b/Engine/Source/Programs/AutomationTool/Gauntlet/Framework/Gauntlet.TestExecutor.cs
--- a/Engine/Source/Programs/AutomationTool/Gauntlet/Framework/Gauntlet.TestExecutor.cs
+++ b/Engine/Source/Programs/AutomationTool/Gauntlet/Framework/Gauntlet.TestExecutor.cs
@@ -7,6 +7,8 @@ using System.Linq;
 using System.Net;
 using System.Text.RegularExpressions;
 using System.Threading;
+using System.Threading.Tasks;
+using AutomationTool;
 using EpicGames.Core;
 using Gauntlet.Utils;
 
@@ -31,6 +33,17 @@ namespace Gauntlet
 
 		[AutoParam(true)]
 		public bool DeferReports;
+
+		/// <summary>
+		/// Maximum number of seconds to wait for test start threads to join when a cancel is signalled
+		/// If this time elapses, the thread is abandoned and expected to close on it's own.
+		/// The most common case of this happening is during an installation to a devkit in UnrealSession
+		/// -------
+		/// TODO: InstallBuild should be given a cancellation token so we're not forced to wait or rely on platform dependent
+		/// implementations to handle cancellations in the middle of a synchronous sdk command call
+		/// </summary>
+		[AutoParam(3 * 60)]
+		public int StartingThreadShutdownLimitSeconds;
 	}
 
 	/// <summary>
@@ -421,18 +434,16 @@ namespace Gauntlet
 						}
 					}
 
-					if (StartingTestThreads.Count > 0)
+					TimeSpan ThreadTimeout = TimeSpan.FromSeconds(Options.StartingThreadShutdownLimitSeconds);
+					Parallel.ForEach(StartingTestThreads, Thread =>
 					{
-						foreach (Thread T in StartingTestThreads)
+						Log.Info("Attempting to join test start thread: {ThreadName} | {ThreadId}", Thread.Name, Thread.ManagedThreadId);
+						if (!Thread.Join(ThreadTimeout))
 						{
-							// SYSLIB0006: 'Thread.Abort()' is obsolete: 'Thread.Abord it not supported and throws PlatformNotSupportedException'
-							//Log.Info("Aborting startup thread");
-							//T.Abort();
-
-							Log.Info($"Thread is still running: {T.Name} {T.ManagedThreadId}");
+							Log.Warning("Failed to join the thread {Thread} | {ThreadId} after {Timeout} seconds.",
+								Thread.Name, Thread.ManagedThreadId, ThreadTimeout.TotalSeconds);
 						}
-						Thread.Sleep(1000);
-					}
+					});
 
 					foreach (TestExecutionInfo TestInfo in StartingTests)
 					{
