Zen workspaces & pak streaming:
- Optional UAT automatic workspace share creation during staging pak builds controlled by:
[/Script/UnrealEd.ProjectPackagingSettings]
bEnablePakStreaming=true
- client side Zen workspace discovery & creation

#rnx

diff --git a/Engine/Source/Programs/AutomationTool/AutomationUtils/ProjectParams.cs b/Engine/Source/Programs/AutomationTool/AutomationUtils/ProjectParams.cs
--- a/Engine/Source/Programs/AutomationTool/AutomationUtils/ProjectParams.cs
+++ b/Engine/Source/Programs/AutomationTool/AutomationUtils/ProjectParams.cs
@@ -564,6 +564,7 @@ namespace AutomationTool
 			bool? IoStore = null,
 			bool? ZenStore = null,
 			string NoZenAutoLaunch = null,
+			string ZenWorkspaceSharePath = null,
 			bool? SkipIoStore = null,
 			bool? GenerateOptimizationData = null,
 			bool? Prereqs = null,
@@ -797,6 +798,7 @@ namespace AutomationTool
 			{
 				this.AdditionalCookerOptions += string.Format(" -NoZenAutoLaunch={0}", this.NoZenAutoLaunch);
 			}
+			this.ZenWorkspaceSharePath = ParseParamValueIfNotSpecified(Command, ZenWorkspaceSharePath, "ZenWorkspaceSharePath", String.Empty);
 			this.GenerateOptimizationData = GetParamValueIfNotSpecified(Command, GenerateOptimizationData, this.GenerateOptimizationData, "makebinaryconfig");
 			
 			this.SkipPak = GetParamValueIfNotSpecified(Command, SkipPak, this.SkipPak, "skippak");
@@ -1476,6 +1478,11 @@ namespace AutomationTool
 		[Help("nozenautolaunch", "URL to a running Zen server")]
 		public string NoZenAutoLaunch { private set; get; }
 
+		/// <summary>
+		/// Shared: Global path to a Zen workspace share to be passed to the client. Used for Zen pak streaming
+		/// </summary>
+		public string ZenWorkspaceSharePath { private set; get; }
+
 		/// <summary>
 		/// Shared: True if optimization data is generated during staging that can improve loadtimes
 		/// </summary>
diff --git a/Engine/Source/Programs/AutomationTool/AutomationUtils/ZenUtils.cs b/Engine/Source/Programs/AutomationTool/AutomationUtils/ZenUtils.cs
new file mode 100644
--- /dev/null
+++ b/Engine/Source/Programs/AutomationTool/AutomationUtils/ZenUtils.cs
@@ -0,0 +1,516 @@
+// Copyright Epic Games, Inc. All Rights Reserved.
+
+using AutomationTool;
+using EpicGames.Core;
+using EpicGames.Horde.Jobs;
+using EpicGames.ProjectStore;
+using EpicGames.Serialization;
+using IdentityModel.OidcClient;
+using Microsoft.Extensions.Logging;
+using System;
+using System.Collections.Generic;
+using System.IO;
+using System.Linq;
+using System.Net.Http;
+using System.Net.WebSockets;
+using System.Runtime.InteropServices;
+using System.Text;
+using System.Text.Json;
+using System.Text.Json.Nodes;
+using System.Threading.Tasks;
+using UnrealBuildBase;
+using UnrealBuildTool;
+
+namespace AutomationUtils
+{
+	public class ZenWorkspace
+	{
+		public string id { get; set; }
+		public string root_path { get; set; }
+		public bool allow_share_creation_from_http { get; set; }
+		public List<ZenWorkspaceShare> shares { get; set; }
+	}
+
+	public class ZenWorkspaceShare
+	{
+		public string id { get; set; }
+		public string share_path { get; set; }
+	}
+
+	public class ZenRunContext
+	{
+		public string Executable { get; set; }
+		public string CommandLineArguments { get; set; }
+		public string DataPath { get; set; }
+
+		public string HostName { get; private set; } = string.Empty;
+		public int HostPort { get; private set; } = 0;
+		public bool IsValid { get; private set; } = false;
+
+		public static ILogger Logger => Log.Logger;
+
+		public static ZenRunContext ReadFromContextFile(FileReference ContextFile)
+		{
+			if (!FileReference.Exists(ContextFile))
+			{
+				return null;
+			}
+
+			string ContextData = FileReference.ReadAllText(ContextFile);
+			JsonSerializerOptions JsonOptions = ZenUtils.GetDefaultJsonSerializerOptions();
+			ZenRunContext Ret = JsonSerializer.Deserialize<ZenRunContext>(ContextData, JsonOptions);
+			if (Ret != null)
+			{
+				Ret.InitializeInstanceData();
+			}
+			return Ret;
+		}
+
+		public static ZenRunContext DiscoverServerContext()
+		{
+			FileReference ContextFile = new (ZenUtils.GetZenRunContextFile());
+			if (!FileReference.Exists(ContextFile))
+			{
+				return null;
+			}
+
+			return ReadFromContextFile(ContextFile);
+		}
+
+		private void InitializeInstanceData()
+		{
+			string LockFilePath = Path.Combine(DataPath, ".lock");
+			FileReference LockFile = new(LockFilePath);
+			if (!FileReference.Exists(LockFile))
+			{
+				return;
+			}
+
+			byte[] CbObjectData;
+			try
+			{
+				using(FileStream Stream = File.Open(LockFilePath, FileMode.Open, FileAccess.Read, FileShare.ReadWrite | FileShare.Delete))
+				{
+					CbObjectData = new byte[Stream.Length];
+					Stream.Read(CbObjectData);
+				}
+			}
+			catch(IOException)
+			{
+				return;
+			}
+
+			if (CbObjectData == null)
+			{
+				return;
+			}
+
+			CbObject LockObject = new CbObject(CbObjectData);
+			int PID = LockObject["pid"].AsInt32();
+			int ServerPort = LockObject["port"].AsInt32();
+
+			IsValid = PID > 0 && (ServerPort > 0 && ServerPort <= 0xffff);
+
+			if (IsValid)
+			{
+				HostPort = ServerPort;
+				HostName = "localhost"; // for now only support localhost
+			}
+		}
+
+		/// <summary>
+		/// Queries the server's health status to check if the the service is running
+		/// </summary>
+		/// <returns></returns>
+		public bool IsServiceRunning()
+		{
+			if (!IsValid)
+			{
+				return false;
+			}
+
+			string Uri = string.Format("http://{0}:{1}/health/ready", HostName, HostPort);
+			using var Request = new HttpRequestMessage(HttpMethod.Get, Uri);
+
+			HttpResponseMessage HttpGetResult = null;
+			HttpClient Client = new HttpClient();
+
+			try
+			{
+				HttpGetResult = Client.Send(Request);
+			}
+			catch
+			{
+				return false;
+			}
+
+			if (HttpGetResult.IsSuccessStatusCode)
+			{
+				return true;
+			}
+
+			return false;
+		}
+
+		/// <summary>
+		/// Checks if the server was run with <c>--workspaces-enabled</c>
+		/// </summary>
+		/// <returns></returns>
+		public bool DoesServerSupportWorkspaces()
+		{
+			if (CommandLineArguments.Contains("--workspaces-enabled"))
+			{
+				return true;
+			}
+
+			return false;
+		}
+
+		/// <summary>
+		/// If there doesn't exists a workspace pointing to BaseDir (or above in the subtree) will create a new zen workspace pointing to BaseDir.
+		/// If a workspace already exists it won't change its Dynamic setting.
+		/// </summary>
+		/// <param name="BaseDir"></param>
+		/// <param name="Dynamic">If the newly created workspace should allow connecting clients to create workspace shares on demand</param>
+		/// <returns>ZenWorkspace if the workspace was created or one already exists. null on failure</returns>
+		public ZenWorkspace GetOrCreateWorkspace(DirectoryReference BaseDir, bool Dynamic = false)
+		{
+			if (!DoesServerSupportWorkspaces())
+			{
+				return null;
+			}
+
+			if (!DirectoryReference.Exists(BaseDir))
+			{
+				Logger.LogError("Can't create zen workspace in {0}. Directory doesn't exist", BaseDir.FullName);
+			}
+
+			ZenWorkspace Workspace = FindWorkspaceForDir(BaseDir);
+			if (Workspace != null)
+			{
+				return Workspace;
+			}
+
+			string Command = String.Format("workspace create {0}", BaseDir.FullName);
+			if (Dynamic)
+			{
+				Command += " --allow-share-create-from-http";
+			}
+
+			if (!ZenUtils.RunZenExe(Command))
+			{
+				return null;
+			}
+
+			Workspace = FindWorkspaceForDir(BaseDir);
+			return Workspace;
+		}
+
+		/// <summary>
+		/// If there doesn't exists a workspace share pointing to ShareDir, will create a new zen workspace share pointing to it.
+		/// </summary>
+		/// <param name="ShareDir"></param>
+		/// <returns></returns>
+		public ZenWorkspaceShare GetOrCreateShare(DirectoryReference ShareDir)
+		{
+			if (!DoesServerSupportWorkspaces())
+			{
+				return null;
+			}
+
+			if (!DirectoryReference.Exists(ShareDir))
+			{
+				Logger.LogError("Can't create zen share in {0}. Directory doesn't exist", ShareDir);
+				return null;
+			}
+
+			ZenWorkspaceShare Result = FindShareForDir(ShareDir);
+			if (Result != null)
+			{
+				return Result;
+			}
+
+			ZenWorkspace Workspace = FindWorkspaceForDir(ShareDir);
+			if (Workspace == null)
+			{
+				Logger.LogError("Can't create zen share in {0}. No zen workspace exists for this directory", ShareDir.FullName);
+				return null;
+			}
+
+			DirectoryReference WorkspaceRef = new(Workspace.root_path);
+			// don't want to have a share pointing to a workspace
+			if (WorkspaceRef == ShareDir)
+			{
+				Logger.LogError("Can't create zen share in {0}. There already exists a workspace pointing to this directory.", ShareDir.FullName);
+				return null;
+			}
+
+			string ShareRelPath = ShareDir.MakeRelativeTo(WorkspaceRef);
+			string Command = String.Format("workspace-share create {0} \"{1}\"", Workspace.id, ShareRelPath);
+			if (!ZenUtils.RunZenExe(Command))
+			{
+				return null;
+			}
+
+			Result = FindShareForDir(ShareDir);
+			return Result;
+		}
+
+		/// <summary>
+		/// Queries the server for a full list of existing workspaces
+		/// </summary>
+		/// <returns></returns>
+		private List<ZenWorkspace> GetWorkspaceList()
+		{
+			if (!IsServiceRunning() || !DoesServerSupportWorkspaces())
+			{
+				return new List<ZenWorkspace>();
+			}
+
+			HttpResponseMessage HttpGetResult = null;
+			HttpClient Client = new HttpClient();
+			string Uri = string.Format("http://{0}:{1}/ws", HostName, HostPort);
+			using var Request = new HttpRequestMessage(HttpMethod.Get, Uri);
+			Request.Headers.Add("Accept", "application/json");
+
+			try
+			{
+				HttpGetResult = Client.Send(Request);
+			}
+			catch
+			{
+				return new List<ZenWorkspace>();
+			}
+
+			if (!HttpGetResult.IsSuccessStatusCode)
+			{
+				return new List<ZenWorkspace>();
+			}
+
+			Task<string> ResultContent = HttpGetResult.Content.ReadAsStringAsync();
+			ResultContent.Wait();
+
+			JsonSerializerOptions JsonOptions = new();
+			JsonOptions.AllowTrailingCommas = true;
+			JsonOptions.PropertyNameCaseInsensitive = true;
+			JsonOptions.ReadCommentHandling = JsonCommentHandling.Skip;
+			JsonNode WorkspaceJsonObject = JsonNode.Parse(ResultContent.Result);
+			JsonArray WorkspaceArray = WorkspaceJsonObject["workspaces"]!.AsArray();
+			List<ZenWorkspace> Result = JsonSerializer.Deserialize<List<ZenWorkspace>>(WorkspaceArray, JsonOptions);
+
+			return Result;
+		}
+
+		/// <summary>
+		/// Finds a zen workspace for the specified directory. Includes workspaces that are above in the directory tree
+		/// </summary>
+		/// <param name="BaseDir"></param>
+		/// <returns>ZenWorkspace on success. null if there's no workspace for in this directory tree</returns>
+		public ZenWorkspace FindWorkspaceForDir(DirectoryReference BaseDir)
+		{
+			var Workspaces = GetWorkspaceList();
+			if (Workspaces == null)
+			{
+				return null;
+			}
+
+			foreach (ZenWorkspace Workspace in Workspaces)
+			{
+				if (BaseDir.IsUnderDirectory(new DirectoryReference(Workspace.root_path)))
+				{
+					return Workspace;
+				}
+			}
+
+			return null;
+		}
+
+		/// <summary>
+		/// Finds a zen workspace share for the specified directory.
+		/// </summary>
+		/// <param name="ShareDir"></param>
+		/// <returns>ZenWorkspaceShare if the share exists. null either if there's no share or there exists no workspace in this dir subtree.</returns>
+		public ZenWorkspaceShare FindShareForDir(DirectoryReference ShareDir)
+		{
+			ZenWorkspace Workspace = FindWorkspaceForDir(ShareDir);
+			if (Workspace == null || Workspace.shares == null)
+			{
+				return null;
+			}
+
+			string RelShare = ShareDir.MakeRelativeTo(new DirectoryReference(Workspace.root_path));
+			foreach (ZenWorkspaceShare ZenShare in Workspace.shares)
+			{
+				if (ZenShare.share_path == RelShare)
+				{
+					return ZenShare;
+				}
+			}
+
+			return null;
+		}
+
+		/// <summary>
+		/// Creates a zen workspace share pointing to ShareDir. If necessary it creates a zen workspace in the parent directory.
+		/// ShareDir must not be a root directory, e.g. d:\
+		/// </summary>
+		/// <param name="ShareDir"></param>
+		/// <returns>
+		/// ZenWorkspace share on success.
+		/// null if the directory doesn't exist or zen server failed to create the workspace
+		/// </returns>
+		public ZenWorkspaceShare CreateWorkspaceAndShare(DirectoryReference ShareDir)
+		{
+			if (!DoesServerSupportWorkspaces())
+			{
+				return null;
+			}
+
+			if (!DirectoryReference.Exists(ShareDir))
+			{
+				Logger.LogError("Can't create workspace and share in {0}. Directory doesn't exist", ShareDir.FullName);
+				return null;
+			}
+
+			if (ShareDir.IsRootDirectory())
+			{
+				Logger.LogError("Can't create share in {0}. Share can't be a root directory", ShareDir);
+				return null;
+			}
+
+			ZenWorkspaceShare Share = FindShareForDir(ShareDir);
+			if (Share != null)
+			{
+				return Share;
+			}
+
+			ZenWorkspace Workspace = GetOrCreateWorkspace(ShareDir.ParentDirectory);
+			if (Workspace == null)
+			{
+				Logger.LogError("Can't create workspace in share's parent directory {0}", ShareDir.ParentDirectory);
+				return null;
+			}
+
+			Share = GetOrCreateShare(ShareDir);
+			if (Share == null)
+			{
+				Logger.LogError("Create zen share in {0} failed", ShareDir);
+				return null;
+			}
+
+			return Share;
+		}
+	}
+
+	public static class ZenUtils
+	{
+		private static string GetServerExecutableName()
+		{
+			string ExecutableName = "zenserver";
+
+			if (RuntimeInformation.IsOSPlatform(OSPlatform.Windows))
+			{
+				ExecutableName += ".exe";
+			}
+
+			return ExecutableName;
+		}
+
+		public static string GetZenInstallPath()
+		{
+			const string EpicProductIdentifier = "UnrealEngine";
+			string LocalAppData = Environment.GetFolderPath(Environment.SpecialFolder.LocalApplicationData);
+			return Path.Combine(LocalAppData, EpicProductIdentifier, "Common", "Zen", "Install");
+		}
+
+		public static string GetZenRunContextFile()
+		{
+			return Path.Combine(GetZenInstallPath(), "zenserver.runcontext");
+		}
+
+		public static FileReference GetZenExeLocation(string ExecutableName)
+		{
+			string PathString = String.Format("{0}/Binaries/{1}/{2}{3}", Unreal.EngineDirectory, HostPlatform.Current.HostEditorPlatform.ToString(), ExecutableName,
+				RuntimeInformation.IsOSPlatform(OSPlatform.Windows) ? ".exe" : "");
+
+			if (Path.IsPathRooted(PathString))
+			{
+				return new FileReference(PathString);
+			}
+			else
+			{
+				return new FileReference(Path.Combine(CommandUtils.CmdEnv.LocalRoot, PathString));
+			}
+		}
+
+		public static bool IsLocalZenServerRunning()
+		{
+			FileReference ContextFileRef = new(GetZenRunContextFile());
+			if (!FileReference.Exists(ContextFileRef))
+			{
+				return false;
+			}
+
+			ZenRunContext Context = ZenRunContext.ReadFromContextFile(ContextFileRef);
+			if (Context == null)
+			{
+				return false;
+			}
+
+			return Context.IsServiceRunning();
+		}
+
+		public static JsonSerializerOptions GetDefaultJsonSerializerOptions()
+		{
+			JsonSerializerOptions Options = new();
+			Options.AllowTrailingCommas = true;
+			Options.PropertyNameCaseInsensitive = true;
+			Options.ReadCommentHandling = JsonCommentHandling.Skip;
+			return Options;
+		}
+
+		public static void ZenLaunch(FileReference SponsorFile)
+		{
+			FileReference ZenLaunchExe = GetZenExeLocation("ZenLaunch");
+
+			string LaunchArguments = String.Format("{0} -SponsorProcessID={1}", CommandUtils.MakePathSafeToUseWithCommandLine(SponsorFile.FullName), Environment.ProcessId);
+
+			CommandUtils.RunAndLog(CommandUtils.CmdEnv, ZenLaunchExe.FullName, LaunchArguments, Options: CommandUtils.ERunOptions.Default);
+		}
+
+		public static bool RunZenExe(string Command)
+		{
+			FileReference ZenExe = GetZenExeLocation("zen");
+			try
+			{
+				CommandUtils.RunAndLog(ZenExe.FullName, Command);
+			}
+			catch
+			{
+				Log.Logger.LogError("{0} failed to run", ZenExe.FullName);
+				return false;
+			}
+
+			return true;
+		}
+
+		public static bool ShouldSetupPakStreaming(ProjectParams Params, DeploymentContext SC)
+		{
+			if (CommandUtils.IsBuildMachine)
+			{
+				return false;
+			}
+
+			if (!Params.UsePak(SC.StageTargetPlatform))
+			{
+				return false;
+			}
+
+			ConfigHierarchy EngineConfig = ConfigCache.ReadHierarchy(ConfigHierarchyType.Game, DirectoryReference.FromFile(Params.RawProjectPath), BuildHostPlatform.Current.Platform, SC.CustomConfig);
+			EngineConfig.GetBool("/Script/UnrealEd.ProjectPackagingSettings", "bUseZenStore", out bool UseZenStore);
+			EngineConfig.GetBool("/Script/UnrealEd.ProjectPackagingSettings", "bEnablePakStreaming", out bool EnablePakStreaming);
+
+			return UseZenStore && EnablePakStreaming;
+		}
+	}
+}
diff --git a/Engine/Source/Programs/AutomationTool/Scripts/CopyBuildToStagingDirectory.Automation.cs b/Engine/Source/Programs/AutomationTool/Scripts/CopyBuildToStagingDirectory.Automation.cs
--- a/Engine/Source/Programs/AutomationTool/Scripts/CopyBuildToStagingDirectory.Automation.cs
+++ b/Engine/Source/Programs/AutomationTool/Scripts/CopyBuildToStagingDirectory.Automation.cs
@@ -25,6 +25,8 @@ using Microsoft.Extensions.Logging;
 using static AutomationTool.CommandUtils;
 using System.Text.Json.Serialization;
 using System.Text.Json.Nodes;
+using AutomationUtils;
+using System.Net;
 
 namespace AutomationScripts
 {
@@ -1019,6 +1021,7 @@ namespace AutomationScripts
 				Options.ReadCommentHandling = JsonCommentHandling.Skip;
 				Options.PropertyNameCaseInsensitive = true;
 				Options.Converters.Add(new JsonStringEnumConverter());
+				Options.DefaultIgnoreCondition = JsonIgnoreCondition.WhenWritingNull;
 				ProjectStoreData ParsedProjectStore = JsonSerializer.Deserialize<ProjectStoreData>(ProjectStoreBytes, Options)!;
 
 				if ((ParsedProjectStore != null) && (ParsedProjectStore.ZenServer != null))
@@ -1043,20 +1046,15 @@ namespace AutomationScripts
 			throw new AutomationException(String.Format("A package store manifest or project store is required when staging to IoStore. Expected to find {0} or {1}. Ensure that cooking was successful.", PackageStoreManifestFile.FullName, ProjectStoreFile.FullName));
 		}
 
-
-
-		/// <summary>
-		/// Creates an updated .projectstore file with a new field containing the stage RemapDirectories, allowing StorageServer to apply these when requesting files
-		/// </summary>
-		/// <returns>New .projectstore file on success, source .projectsource file on failure</returns>
-		private static FileReference MakeProjectStoreWithRemappings(DeploymentContext SC, FileReference SourceProjectStorePath)
+		private static FileReference UpdateProjectStore(ProjectParams Params, DeploymentContext SC, FileReference SourceProjectStorePath)
 		{
 			if (!EpicGames.Core.JsonObject.TryRead(SourceProjectStorePath, out EpicGames.Core.JsonObject ProjectStore))
 			{
 				return SourceProjectStorePath;
 			}
 
-			FileReference UpdatedProjectStorePath = FileReference.Combine(SC.ProjectRoot, "Intermediate", "Staging", SC.StageTargetPlatform.PlatformType.ToString(), "remap_" + SourceProjectStorePath.GetFileName());
+			FileReference UpdatedProjectStorePath = FileReference.Combine(SC.ProjectRoot, "Intermediate", "Staging", SC.StageTargetPlatform.PlatformType.ToString(),
+				"updated_" + SourceProjectStorePath.GetFileName());
 			if (FileReference.Exists(UpdatedProjectStorePath) && !InternalUtils.SafeDeleteFile(UpdatedProjectStorePath.FullName))
 			{
 				return SourceProjectStorePath;
@@ -1067,7 +1065,64 @@ namespace AutomationScripts
 				return SourceProjectStorePath;
 			}
 
-			string MakeStorageServerPath( StagedDirectoryReference Source )
+			if (ZenUtils.ShouldSetupPakStreaming(Params, SC))
+			{
+				AppendWorkspaceToProjectStore(SC, ProjectStore);
+			}
+			else if (SC.RemapDirectories.Any())
+			{
+				MakeProjectStoreWithRemappings(SC, ProjectStore);
+			}
+			File.WriteAllText(UpdatedProjectStorePath.FullName, ProjectStore.ToJsonString());
+
+			return UpdatedProjectStorePath;
+		}
+
+		private static void AppendWorkspaceToProjectStore(DeploymentContext SC, EpicGames.Core.JsonObject ProjectStore)
+		{
+			if (!ProjectStore.TryGetObjectField("zenserver", out EpicGames.Core.JsonObject ZenServerObject))
+			{
+				return;
+			}
+
+			if (!ZenServerObject.TryGetStringField("HostName", out string HostName) || !ZenServerObject.TryGetIntegerField("HostPort", out int HostPort))
+			{
+				return;
+			}
+
+			ZenRunContext ZenServerContext = ZenRunContext.DiscoverServerContext();
+			if (ZenServerContext == null)
+			{
+				return;
+			}
+
+			DirectoryReference StageBaseDir = DirectoryReference.Combine(SC.ProjectRoot, "Saved", "StagedBuilds");
+			ZenWorkspace BaseWorkspace = ZenServerContext.FindWorkspaceForDir(StageBaseDir);
+			DirectoryReference ShareDir = DirectoryReference.Combine(StageBaseDir, SC.PlatformDir);
+			if (BaseWorkspace == null)
+			{
+				return;
+			}
+
+			ZenWorkspaceShare WorkspaceShare = ZenServerContext.FindShareForDir(ShareDir);
+			if (WorkspaceShare != null)
+			{
+				DirectoryReference SharePath = DirectoryReference.Combine(StageBaseDir, WorkspaceShare.share_path);
+				ZenServerObject.AddOrSetFieldValue("workspaceid", BaseWorkspace.id);
+				ZenServerObject.AddOrSetFieldValue("shareid", WorkspaceShare.id);
+				ZenServerObject.AddOrSetFieldValue("workspacesharepath", ShareDir.ToString());
+				ZenServerObject.AddOrSetFieldValue("operatingmode", ZenOperatingMode.Filesystem);
+			}
+
+			ProjectStore.AddOrSetFieldValue("zenserver", ZenServerObject);
+		}
+
+		/// <summary>
+		/// Updates project store json data with a new field containing the stage RemapDirectories, allowing StorageServer to apply these when requesting files
+		/// </summary>
+		private static void MakeProjectStoreWithRemappings(DeploymentContext SC, EpicGames.Core.JsonObject ProjectStore)
+		{
+			string MakeStorageServerPath(StagedDirectoryReference Source)
 			{
 				string Result = Source.ToString();
 				string[] PathFragments = Result.Split(new[] { '\\', '/' }, StringSplitOptions.RemoveEmptyEntries);
@@ -1081,7 +1136,7 @@ namespace AutomationScripts
 					{
 						PathFragments[0] = "{project}";
 					}
-					Result = '/' + string.Join('/',PathFragments) + '/';
+					Result = '/' + string.Join('/', PathFragments) + '/';
 				}
 				return Result;
 			}
@@ -1091,16 +1146,12 @@ namespace AutomationScripts
 			{
 				EpicGames.Core.JsonObject JsonRemapDirectory = new();
 				JsonRemapDirectory.AddOrSetFieldValue("from", MakeStorageServerPath(RemapDirectory.Item1));
-				JsonRemapDirectory.AddOrSetFieldValue("to",   MakeStorageServerPath(RemapDirectory.Item2));
+				JsonRemapDirectory.AddOrSetFieldValue("to", MakeStorageServerPath(RemapDirectory.Item2));
 				JsonRemapDirectories.Add(JsonRemapDirectory);
 			}
 			ProjectStore.AddOrSetFieldValue("remapDirectories", JsonRemapDirectories.ToArray());
-
-			File.WriteAllText(UpdatedProjectStorePath.FullName, ProjectStore.ToJsonString());
-			return UpdatedProjectStorePath;
 		}
 
-
 		private static StagedFileType? GetStagingTypeForCookedFile(FileReference File)
 		{
 			// json files have never been staged
@@ -1801,13 +1852,11 @@ namespace AutomationScripts
 
 							if (CookedFile.HasExtension(".projectstore"))
 							{
-								if (!Params.UsePak(SC.StageTargetPlatform))
+								if (!Params.UsePak(SC.StageTargetPlatform) || ZenUtils.ShouldSetupPakStreaming(Params, SC))
 								{
-									FileReference SrcCookedFile = CookedFile;
-									if (SC.RemapDirectories.Any())
-									{
-										SrcCookedFile = MakeProjectStoreWithRemappings(SC, CookedFile);
-									}
+									FileReference SrcCookedFile = FileReference.Combine(SC.ProjectRoot, "Intermediate", "Staging", SC.StageTargetPlatform.PlatformType.ToString(), "updated_" + CookedFile.GetFileName());
+
+									SrcCookedFile = UpdateProjectStore(Params, SC, CookedFile);
 
 									SC.StageFile(StagedFileType.NonUFS, SrcCookedFile, new StagedFileReference(CookedFile.MakeRelativeTo(SC.PlatformCookDir)));
 								}
@@ -3031,6 +3080,50 @@ namespace AutomationScripts
 			CreatePaks(Params, SC, PakInputs, PakCryptoSettings, CryptoKeysCacheFilename);
 		}
 
+		private static void PrepareZenWorkspaceEnv(ProjectParams Params, DeploymentContext SC)
+		{
+			if (!ZenUtils.ShouldSetupPakStreaming(Params, SC))
+			{
+				return;
+			}
+
+			Logger.LogInformation("Preparing zen workspace environment.");
+
+			ZenUtils.ZenLaunch(SC.RawProjectPath);
+
+			ZenRunContext ZenServerContext = ZenRunContext.DiscoverServerContext();
+			if (ZenServerContext == null)
+			{
+				return;
+			}
+
+			FileReference ZenMarkerFile = FileReference.Combine(SC.PlatformCookDir, "ue.projectstore");
+			System.IO.FileInfo ProjectStoreFileInfo = ZenMarkerFile.ToFileInfo();
+			if (!ProjectStoreFileInfo.Exists)
+			{
+				return;
+			}
+
+			byte[] MarkerFileData = FileReference.ReadAllBytes(ZenMarkerFile);
+			JsonSerializerOptions JsonOptions = new();
+			JsonOptions.AllowTrailingCommas = true;
+			JsonOptions.ReadCommentHandling = JsonCommentHandling.Skip;
+			JsonOptions.PropertyNameCaseInsensitive = true;
+			JsonOptions.Converters.Add(new JsonStringEnumConverter());
+			ProjectStoreData ExtractedProjectStore = JsonSerializer.Deserialize<ProjectStoreData>(MarkerFileData, JsonOptions);
+
+			if (ExtractedProjectStore != null && ExtractedProjectStore.ZenServer != null)
+			{
+				ExtractedProjectStore.ZenServer.OperatingMode = ZenOperatingMode.Stream;
+
+				if (ShouldCreatePak(Params))
+				{
+					DirectoryReference StageBaseDir = DirectoryReference.Combine(SC.ProjectRoot, "Saved", "StagedBuilds", SC.PlatformDir);
+					ZenWorkspaceShare WorkspaceShare = ZenServerContext.CreateWorkspaceAndShare(StageBaseDir);
+				}
+			}
+		}
+
 		/// <summary>
 		/// Creates a standalone pak file for crash reporter
 		/// </summary>
@@ -6066,6 +6159,7 @@ namespace AutomationScripts
 
 						// write out the commandline file now so it can go into the manifest
 						WriteStageCommandline(Params, SC);
+						PrepareZenWorkspaceEnv(Params, SC);
 						CreateStagingManifest(Params, SC);
 						CleanStagingDirectory(Params, SC);
 					}
diff --git a/Engine/Source/Programs/AutomationTool/Scripts/RunProjectCommand.Automation.cs b/Engine/Source/Programs/AutomationTool/Scripts/RunProjectCommand.Automation.cs
--- a/Engine/Source/Programs/AutomationTool/Scripts/RunProjectCommand.Automation.cs
+++ b/Engine/Source/Programs/AutomationTool/Scripts/RunProjectCommand.Automation.cs
@@ -16,6 +16,7 @@ using EpicGames.Core;
 using UnrealBuildBase;
 using Microsoft.Extensions.Logging;
 using EpicGames.Serialization;
+using AutomationUtils;
 
 namespace AutomationScripts
 {
@@ -270,6 +271,27 @@ namespace AutomationScripts
 				ERunOptions ClientRunFlags;
 				string ClientApp;
 				string ClientCmdLine;
+
+				if (!string.IsNullOrEmpty(Params.ZenWorkspaceSharePath))
+				{
+					ZenRunContext ZenServerContext = ZenRunContext.ReadFromContextFile(new FileReference(ZenUtils.GetZenRunContextFile()));
+					if (ZenServerContext != null && ZenServerContext.IsServiceRunning())
+					{
+						DirectoryReference WorkspaceShareDir = new(Params.ZenWorkspaceSharePath);
+						if (SC.Count > 0 && DirectoryReference.Exists(WorkspaceShareDir))
+						{
+							if (ZenServerContext.CreateWorkspaceAndShare(WorkspaceShareDir) == null)
+							{
+								throw new AutomationException("Couldn't createa a Zen workspace at {0}", WorkspaceShareDir.FullName);
+							}
+						}
+					}
+					else
+					{
+						throw new AutomationException("Zen service is not running or wasn't detected. Can't run with ZenWorkspaceSharePath");
+					}
+				}
+
 				SetupClientParams(SC, Params, ClientLogFile, out ClientRunFlags, out ClientApp, out ClientCmdLine);
 				RunClient(SC, ClientLogFile, ClientRunFlags, ClientApp, ClientCmdLine, Params);
 			}
@@ -715,6 +737,11 @@ namespace AutomationScripts
 				TempCmdLine += Params.SessionLabel + " ";
 			}
 
+			if (!String.IsNullOrEmpty(Params.ZenWorkspaceSharePath))
+			{
+				TempCmdLine += "-ZenPak -ZenWorkspaceSharePath=" + Params.ZenWorkspaceSharePath + " ";
+			}
+
 			TempCmdLine += SC.StageTargetPlatform.GetLaunchExtraCommandLine(Params);
 
 			TempCmdLine += "-CrashForUAT ";
diff --git a/Engine/Source/Programs/Shared/EpicGames.ProjectStore/ProjectStoreTypes.cs b/Engine/Source/Programs/Shared/EpicGames.ProjectStore/ProjectStoreTypes.cs
--- a/Engine/Source/Programs/Shared/EpicGames.ProjectStore/ProjectStoreTypes.cs
+++ b/Engine/Source/Programs/Shared/EpicGames.ProjectStore/ProjectStoreTypes.cs
@@ -18,6 +18,17 @@ using UnrealBuildBase;
 
 namespace EpicGames.ProjectStore
 {
+	/// <summary>
+	/// How the client should operate. Whether it should stream from zenserver
+	/// or use the local filesystem
+	/// </summary>
+	[JsonConverter(typeof(JsonStringEnumConverter))]
+	public enum ZenOperatingMode
+	{
+		Stream,
+		Filesystem
+	}
+
 	/// <summary>
 	/// A specification of storage details with zenserver
 	/// </summary>
@@ -34,7 +45,7 @@ namespace EpicGames.ProjectStore
 		/// <summary>
 		/// The list of remote host names for zenserver
 		/// </summary>
-		public List<string> RemoteHostNames { get; } = new List<string>();
+		public List<string> RemoteHostNames { get; set; } = new List<string>();
 		/// <summary>
 		/// The port for zenserver on the host
 		/// </summary>
@@ -47,6 +58,18 @@ namespace EpicGames.ProjectStore
 		/// The oplog identifier for the stored data in zenserver
 		/// </summary>
 		public string? OplogId { get; set; }
+		/// <summary>
+		/// The zen workspace the build's share is available from
+		/// </summary>
+		public string? WorkspaceId { get; set; }
+		/// <summary>
+		/// The zen workspace share pointing to the build
+		/// </summary>
+		public string? ShareId { get; set; }
+		/// <summary>
+		/// Selected mode of operation. If missing, assume ZenOperatingMode.Stream
+		/// </summary>
+		public ZenOperatingMode? OperatingMode { get; set; } = ZenOperatingMode.Stream;
 	}
 	/// <summary>
 	/// A desctiptor of a project store
diff --git a/Engine/Source/Runtime/StorageServerClient/Private/StorageServerConnection.cpp b/Engine/Source/Runtime/StorageServerClient/Private/StorageServerConnection.cpp
--- a/Engine/Source/Runtime/StorageServerClient/Private/StorageServerConnection.cpp
+++ b/Engine/Source/Runtime/StorageServerClient/Private/StorageServerConnection.cpp
@@ -109,6 +109,110 @@ bool FStorageServerConnection::Initialize(TArrayView<const FString> HostAddresse
 	return false;
 }
 
+TIoStatusOr<FStorageServerConnection::Workspaces> FStorageServerConnection::GetWorkspaces()
+{
+	IStorageServerHttpClient::FResult ResultTuple = HttpClient->RequestSync(BaseURI, EStorageServerContentType::CbObject);
+	TIoStatusOr<FIoBuffer> RequestResult = ResultTuple.Get<0>();
+	if (RequestResult.IsOk())
+	{
+		FMemoryReaderView Reader(RequestResult.ValueOrDie().GetView());
+		FCbObject ResponseObj = LoadCompactBinary(Reader).AsObject();
+
+		Workspaces Result;
+		FCbArrayView WorkspacesArray = ResponseObj["workspaces"].AsArrayView();
+		for (FCbFieldView WorkspaceField : WorkspacesArray)
+		{
+			FCbObjectView WorkspaceObject = WorkspaceField.AsObjectView();
+
+			Workspaces::Workspace Workspace = {
+				.Id = *WriteToString<64>(WorkspaceObject["id"].AsObjectId()),
+				.Root = FString(WorkspaceObject["root_path"].AsString()),
+				.AllowShareCreationFromHttp = WorkspaceObject["allow_share_creation_from_http"].AsBool()
+			};
+
+			if (!Workspace.Id.IsEmpty())
+			{
+				FCbArrayView SharesArray = WorkspaceObject["shares"].AsArrayView();
+				for (FCbFieldView ShareField : SharesArray)
+				{
+					FCbObjectView ShareObject = ShareField.AsObjectView();
+
+					Workspaces::Share Share = {
+						.Id = *WriteToString<64>(ShareObject["id"].AsObjectId()),
+						.Path = FString(ShareObject["share_path"].AsString()),
+						.Alias = FString(ShareObject["alias"].AsString())
+					};
+
+					if (!Share.Id.IsEmpty())
+					{
+						Workspace.Shares.Add(Share);
+					}
+				}
+				Result.Workspaces.Add(Workspace);
+			}
+		}
+		return Result;
+	}
+	return RequestResult.Status();
+}
+
+static void PercentEncodeString(FAnsiStringBuilderBase& EncodedBuilder, const FString& String)
+{
+	for (FString::ElementType C : String)
+	{
+		switch (C)
+		{
+		case '!':
+		case '#':
+		case '$':
+		case '&':
+		case '\'':
+		case '(':
+		case ')':
+		case '*':
+		case '+':
+		case ',':
+		case '/':
+		case ':':
+		case ';':
+		case '=':
+		case '?':
+		case '@':
+		case '[':
+		case ']':
+			EncodedBuilder.Append("%02X", C);
+			break;
+		default:
+			EncodedBuilder.Appendf("%c", C);
+		}
+	}
+}
+
+TIoStatusOr<FString> FStorageServerConnection::CreateShare(const FString& WorkspaceId, const FString& SharePath, const FString& Alias)
+{
+	TAnsiStringBuilder<256> ResourceBuilder;
+	ResourceBuilder.Append(BaseURI);
+	ResourceBuilder.Append("/");
+	ResourceBuilder.Append(WorkspaceId);
+	ResourceBuilder.Append("/");
+	ResourceBuilder.Append("000000000000000000000000");
+	ResourceBuilder.Append("?share_path=");
+	PercentEncodeString(ResourceBuilder, SharePath);
+	if (!Alias.IsEmpty())
+	{
+		ResourceBuilder.Append("&alias=");
+		PercentEncodeString(ResourceBuilder, Alias);
+	}
+	IStorageServerHttpClient::FResult ResultTuple = HttpClient->RequestSync(*ResourceBuilder, EStorageServerContentType::Unknown, "PUT");
+	TIoStatusOr<FIoBuffer> Result = ResultTuple.Get<0>();
+	if (Result.IsOk())
+	{
+		FMemoryView TextView(Result.ValueOrDie().GetView());
+		return FString::ConstructFromPtrSize(reinterpret_cast<const ANSICHAR*>(TextView.GetData()), TextView.GetSize());
+	}
+	return Result.Status();
+}
+
 void FStorageServerConnection::PackageStoreRequest(TFunctionRef<void(FPackageStoreEntryResource&&)> Callback)
 {
 	TAnsiStringBuilder<256> ResourceBuilder;
@@ -821,7 +925,7 @@ void FStorageServerConnection::GetDefaultCacheConfiguration(FCacheConfiguration&
 	OutConfiguration.bEnable = bPlatformSupportsCaching && bPlatformEnablesCachingByDefault;
 
 	// always check it first to ensure we disable cache if cmd arg is provided
-	if ((FCString::Strstr(CmdLine, TEXT("-ZenDisableCache")) != nullptr))
+	if ((FCString::Strstr(CmdLine, TEXT("-ZenDisableCache")) != nullptr) || UE::IsUsingZenPakFileStreaming())
 	{
 		OutConfiguration.bEnable = false;
 	}
diff --git a/Engine/Source/Runtime/StorageServerClient/Private/StorageServerConnection.h b/Engine/Source/Runtime/StorageServerClient/Private/StorageServerConnection.h
--- a/Engine/Source/Runtime/StorageServerClient/Private/StorageServerConnection.h
+++ b/Engine/Source/Runtime/StorageServerClient/Private/StorageServerConnection.h
@@ -34,7 +34,7 @@
 #endif
 
 #ifndef HAS_STORAGE_SERVER_RPC_GETCHUNKS_API // zen server 5.5.16 introduced a new API to request chunks
-	#define HAS_STORAGE_SERVER_RPC_GETCHUNKS_API 1
+	#define HAS_STORAGE_SERVER_RPC_GETCHUNKS_API 0 // #jira UE-234828 temporarily disabled until RPC batch requests support gets added to zen workspaces
 #endif
 
 #if !UE_BUILD_SHIPPING
@@ -51,6 +51,29 @@ public:
 
 	bool Initialize(TArrayView<const FString> HostAddresses, const int32 Port, const FAnsiStringView& InBaseURI);
 
+	struct Workspaces
+	{
+		struct Share
+		{
+			FString Id;
+			FString Path;
+			FString Alias;
+		};
+
+		struct Workspace
+		{
+			FString Id;
+			FString Root;
+			bool AllowShareCreationFromHttp = false;
+
+			TArray<Share> Shares;
+		};
+		TArray<Workspace> Workspaces;
+	};
+
+	TIoStatusOr<Workspaces> GetWorkspaces();
+	TIoStatusOr<FString> CreateShare(const FString& WorkspaceId, const FString& SharePath, const FString& Alias);
+
 	void PackageStoreRequest(TFunctionRef<void(FPackageStoreEntryResource&&)> Callback);
 	void FileManifestRequest(TFunctionRef<void(FIoChunkId Id, FStringView Path, int64 RawSize)> Callback);
 	void ChunkInfosRequest(TFunctionRef<void(FIoChunkId Id, FIoHash RawHash, int64 RawSize)> Callback);
diff --git a/Engine/Source/Runtime/StorageServerClient/Private/StorageServerPlatformFile.cpp b/Engine/Source/Runtime/StorageServerClient/Private/StorageServerPlatformFile.cpp
--- a/Engine/Source/Runtime/StorageServerClient/Private/StorageServerPlatformFile.cpp
+++ b/Engine/Source/Runtime/StorageServerClient/Private/StorageServerPlatformFile.cpp
@@ -914,9 +914,114 @@ TUniquePtr<FArchive> FStorageServerPlatformFile::TryFindProjectStoreMarkerFile(I
 	return nullptr;
 }
 
+static FAnsiString FindWorkspaceSharePath(const FString& WorkspaceSharePath, const FStorageServerConnection::Workspaces& Workspaces)
+{
+	const bool IsRelative = FPaths::IsRelative(WorkspaceSharePath);
+	for (const auto& Workspace : Workspaces.Workspaces)
+	{
+		if (!Workspace.Root.IsEmpty())
+		{
+			FString WorkspaceRoot = Workspace.Root;
+			FPaths::NormalizeDirectoryName(WorkspaceRoot);
+			for (const auto& Share : Workspace.Shares)
+			{
+				FString SharePath(Share.Path);
+				FPaths::NormalizeDirectoryName(SharePath);
+				if (FPaths::IsSamePath(IsRelative ? SharePath : FPaths::Combine(WorkspaceRoot, SharePath), WorkspaceSharePath))
+				{
+					TAnsiStringBuilder<256> SharePathBuilder;
+					SharePathBuilder.Append("/ws/");
+					SharePathBuilder.Append(Workspace.Id);
+					SharePathBuilder.Append("/");
+					SharePathBuilder.Append(Share.Id);
+					return SharePathBuilder.ToString();
+				}
+			}
+		}
+	}
+	return {};
+}
+
+static FString GetAsSubPath(const FString& WorkspaceRootPath, const FString& WorkspaceSharePath)
+{
+	FString TestRoot = FPaths::GetPath(WorkspaceSharePath);
+	while (TestRoot.Len() >= WorkspaceRootPath.Len())
+	{
+		if (FPaths::IsSamePath(TestRoot, WorkspaceRootPath))
+		{
+			return WorkspaceSharePath.Mid(TestRoot.Len() + 1);
+		}
+		TestRoot = FPaths::GetPath(TestRoot);
+	}
+	return {};
+}
+
+static FAnsiString CreateWorkspaceShare(FStorageServerConnection& QueryConnection, const FString& WorkspaceSharePath, const FStorageServerConnection::Workspaces& Workspaces)
+{
+	const bool IsRelative = FPaths::IsRelative(WorkspaceSharePath);
+	for (const auto& Workspace : Workspaces.Workspaces)
+	{
+		if (Workspace.AllowShareCreationFromHttp)
+		{
+			FString WorkspaceRoot = Workspace.Root;
+			FPaths::NormalizeDirectoryName(WorkspaceRoot);
+			if (!Workspace.Root.IsEmpty())
+			{
+				FString SharePath = IsRelative ? WorkspaceSharePath : GetAsSubPath(WorkspaceRoot, WorkspaceSharePath);
+				FPaths::NormalizeDirectoryName(SharePath);
+				if (!SharePath.IsEmpty())
+				{
+					TIoStatusOr<FString> CreateResult = QueryConnection.CreateShare(Workspace.Id, SharePath, "");
+					if (CreateResult.IsOk())
+					{
+						TAnsiStringBuilder<256> SharePathBuilder;
+						SharePathBuilder.Append("/ws/");
+						SharePathBuilder.Append(Workspace.Id);
+						SharePathBuilder.Append("/");
+						SharePathBuilder.Append(CreateResult.ValueOrDie());
+						return SharePathBuilder.ToString();
+					}
+				}
+			}
+		}
+	}
+	return {};
+}
+
 FAnsiString FStorageServerPlatformFile::MakeBaseURI()
 {
 	TAnsiStringBuilder<256> BaseURIBuilder;
+	if (UE::IsUsingZenPakFileStreaming() && !WorkspaceSharePath.IsEmpty())
+	{
+		FPaths::NormalizeDirectoryName(WorkspaceSharePath);
+		FStorageServerConnection QueryConnection;
+		if (QueryConnection.Initialize(HostAddrs, HostPort, "/ws"))
+		{
+			TIoStatusOr<FStorageServerConnection::Workspaces> WorkspacesResponse = QueryConnection.GetWorkspaces();
+			if (WorkspacesResponse.IsOk())
+			{
+				const FStorageServerConnection::Workspaces& Workspaces = WorkspacesResponse.ValueOrDie();
+
+				if (FAnsiString ExistingShareURI = FindWorkspaceSharePath(WorkspaceSharePath, Workspaces); !ExistingShareURI.IsEmpty())
+				{
+					return ExistingShareURI;
+				}
+				if (FAnsiString NewShareURI = CreateWorkspaceShare(QueryConnection, WorkspaceSharePath, Workspaces); !NewShareURI.IsEmpty())
+				{
+					return NewShareURI;
+				}
+				UE_LOG(LogStorageServerPlatformFile, Error, TEXT("Failed to to resolve or create workspace share path '%s' from %s"), *WorkspaceSharePath, *FString(QueryConnection.GetHostAddr()));
+			}
+			else
+			{
+				UE_LOG(LogStorageServerPlatformFile, Error, TEXT("Failed to get list of workspaces from %s. Status: %s"), *FString(QueryConnection.GetHostAddr()), *WorkspacesResponse.Status().ToString());
+			}
+		}
+		else
+		{
+			UE_LOG(LogStorageServerPlatformFile, Error, TEXT("Failed to connect to %s to get list of workspace shares"), *FString(QueryConnection.GetHostAddr()));
+		}
+	}
 	if (!BaseURI.IsEmpty())
 	{
 		BaseURIBuilder.Append(BaseURI);
@@ -954,6 +1059,8 @@ bool FStorageServerPlatformFile::ShouldBeUsed(IPlatformFile* Inner, const TCHAR*
 	if (FParse::Param(FCommandLine::Get(), TEXT("SkipZenStore")))
 		return false;
 
+	bool bPreferFileSystem = false;
+
 #if WITH_COTF
 	UE::Cook::ICookOnTheFlyModule& CookOnTheFlyModule = FModuleManager::LoadModuleChecked<UE::Cook::ICookOnTheFlyModule>(TEXT("CookOnTheFly"));
 	TSharedPtr<UE::Cook::ICookOnTheFlyServerConnection> DefaultConnection = CookOnTheFlyModule.GetDefaultServerConnection();
@@ -975,6 +1082,13 @@ bool FStorageServerPlatformFile::ShouldBeUsed(IPlatformFile* Inner, const TCHAR*
 			if (ProjectStoreObject->TryGetObjectField(TEXT("zenserver"), ZenServerObjectPtr) && (ZenServerObjectPtr != nullptr))
 			{
 				const TSharedPtr<FJsonObject>& ZenServerObject = *ZenServerObjectPtr;
+
+				FString FilesystemOperatingMode;
+				if (ZenServerObject->TryGetStringField(TEXT("operatingmode"), FilesystemOperatingMode) && (FilesystemOperatingMode == "Filesystem"))
+				{
+					bPreferFileSystem = true;
+				}
+
 #if PLATFORM_DESKTOP || PLATFORM_ANDROID
 				FString HostName;
 				if (ZenServerObject->TryGetStringField(TEXT("hostname"), HostName) && !HostName.IsEmpty())
@@ -998,7 +1112,7 @@ bool FStorageServerPlatformFile::ShouldBeUsed(IPlatformFile* Inner, const TCHAR*
 									// As this is the fastest connection when on USB-C, set this as the first to test
 									// TODO: what about when using USB2?? Should we detect device type?
 									HostAddrs.Insert(MacHostName + TEXT(".local"), 0);
-								
+
 									// Some macs drop the ".local", so try that as well, but as a last resort
 									HostAddrs.Add(MacHostName);
 								}
@@ -1039,7 +1153,13 @@ bool FStorageServerPlatformFile::ShouldBeUsed(IPlatformFile* Inner, const TCHAR*
 	{
 		UE_LOG(LogStorageServerPlatformFile, Display, TEXT("Using connection settings from command line: -ZenStorePort='%d'"), HostPort);
 	}
-	return HostAddrs.Num() > 0;
+
+	if (!bPreferFileSystem || UE::IsUsingZenPakFileStreaming())
+	{
+		return HostAddrs.Num() > 0;
+	}
+
+	return false;
 }
 
 bool FStorageServerPlatformFile::Initialize(IPlatformFile* Inner, const TCHAR* CmdLine)
@@ -1093,6 +1213,10 @@ bool FStorageServerPlatformFile::Initialize(IPlatformFile* Inner, const TCHAR* C
 					{
 						BaseURI.Empty();
 					}
+					if (!ZenServerObject->TryGetStringField(TEXT("workspacesharepath"), WorkspaceSharePath))
+					{
+						WorkspaceSharePath.Empty();
+					}
 					UE_LOG(LogStorageServerPlatformFile, Display, TEXT("Using settings from ue.projectstore: ServerProject='%s' and ServerPlatform='%s'"), *ServerProject, *ServerPlatform);
 				}
 
@@ -1123,6 +1247,10 @@ bool FStorageServerPlatformFile::Initialize(IPlatformFile* Inner, const TCHAR* C
 		{
 			UE_LOG(LogStorageServerPlatformFile, Display, TEXT("Using settings from command line: -ZenStoreBaseURI='%s'"), *BaseURI);
 		}
+		if (FParse::Value(CmdLine, TEXT("-ZenWorkspaceSharePath="), WorkspaceSharePath))
+		{
+			UE_LOG(LogStorageServerPlatformFile, Display, TEXT("Using settings from command line: -ZenWorkspaceSharePath='%s'"), *WorkspaceSharePath);
+		}
 
 		if (UE::IsUsingZenPakFileStreaming())
 		{
diff --git a/Engine/Source/Runtime/StorageServerClient/Private/StorageServerPlatformFile.h b/Engine/Source/Runtime/StorageServerClient/Private/StorageServerPlatformFile.h
--- a/Engine/Source/Runtime/StorageServerClient/Private/StorageServerPlatformFile.h
+++ b/Engine/Source/Runtime/StorageServerClient/Private/StorageServerPlatformFile.h
@@ -163,6 +163,7 @@ private:
 	FString BaseURI;
 	FString AbsProjectDir;
 	FString AbsEngineDir;
+	FString WorkspaceSharePath;
 	TDirectoryTree<FString> RemapDirectoriesTree;
 	mutable TArray<FString> HostAddrs;
 	mutable uint16 HostPort = 8558;
diff --git a/Engine/Source/Runtime/StorageServerClientDebug/Private/StorageServerConnectionDebug.cpp b/Engine/Source/Runtime/StorageServerClientDebug/Private/StorageServerConnectionDebug.cpp
--- a/Engine/Source/Runtime/StorageServerClientDebug/Private/StorageServerConnectionDebug.cpp
+++ b/Engine/Source/Runtime/StorageServerClientDebug/Private/StorageServerConnectionDebug.cpp
@@ -85,7 +85,7 @@ bool FStorageServerConnectionDebug::OnTick(float)
 		if (GZenShowStats && GEngine)
 		{
 			FString ZenConnectionDebugMsg;
-			ZenConnectionDebugMsg = FString::Printf(TEXT("ZenServer streaming from %s [%.2fMbps]"), *HostAddress, Throughput);
+			ZenConnectionDebugMsg = FString::Printf(TEXT("ZenServer %s from %s [%.2fMbps]"), UE::IsUsingZenPakFileStreaming() ? TEXT("pak streaming") : TEXT("streaming"), * HostAddress, Throughput);
 			GEngine->AddOnScreenDebugMessage((uint64)this, 86400.0f, FColor::White, ZenConnectionDebugMsg, false);
 		}
 		
