Update ReportItemizedExecutableCode to output the Total in the sections report.

#rnx
#tests local

diff --git a/Engine/Source/Programs/AutomationTool/Scripts/ReportItemizedExecutableCode.cs b/Engine/Source/Programs/AutomationTool/Scripts/ReportItemizedExecutableCode.cs
--- a/Engine/Source/Programs/AutomationTool/Scripts/ReportItemizedExecutableCode.cs
+++ b/Engine/Source/Programs/AutomationTool/Scripts/ReportItemizedExecutableCode.cs
@@ -378,6 +378,7 @@ namespace ReportItemizedExecutableCode.Automation
 				throw new AutomationException("CSV output ({0}) does not match expectation ({1})", Lines[0], ExpectedColumns);
 			}
 
+			uint TotalVMSize = 0;
 			StringBuilder Builder = new StringBuilder("Section,Megabytes,Bytes\n");
 			foreach (string Line in Lines.Skip(1))
 			{
@@ -385,10 +386,13 @@ namespace ReportItemizedExecutableCode.Automation
 				string SectionName = SanitizeModuleName(LineValues[0]).Replace("-", "");
 
 				uint VMSize = uint.Parse(LineValues[1]);
+				TotalVMSize += VMSize;
 
 				Builder.AppendLine(string.Format("{0},{1:F2},{2}", SectionName, VMSize / (1024.0 * 1024.0), VMSize));
 			}
 
+			Builder.AppendLine(string.Format("{0},{1:F2},{2}", "Total", TotalVMSize / (1024.0 * 1024.0), TotalVMSize));
+
 			foreach (string DestPath in DestPaths)
 			{
 				System.Console.WriteLine($"\nWriting sections report to: {DestPath}");
