[Gauntlet] Delete an UnrealTestNode's artifact directory if it already exists when it reserves that path.
This prevents the artifact path in CreateReport from picking up artifacts from previous runs of the test

diff --git a/Engine/Source/Programs/AutomationTool/Gauntlet/Unreal/Base/Gauntlet.UnrealTestNode.cs b/Engine/Source/Programs/AutomationTool/Gauntlet/Unreal/Base/Gauntlet.UnrealTestNode.cs
--- a/Engine/Source/Programs/AutomationTool/Gauntlet/Unreal/Base/Gauntlet.UnrealTestNode.cs
+++ b/Engine/Source/Programs/AutomationTool/Gauntlet/Unreal/Base/Gauntlet.UnrealTestNode.cs
@@ -12,6 +12,7 @@ using Microsoft.Extensions.Logging;
 using AutomationUtils.Matchers;
 using static Gauntlet.HordeReport;
 using static Gauntlet.UnrealSessionInstance;
+using Gauntlet.Utils;
 
 namespace Gauntlet
 {
@@ -695,10 +696,11 @@ namespace Gauntlet
 		}
 
 		/// <summary>
-		/// Generate an unique path for the test artifacts and reserve it
+		/// Generates a unique path for the test artifacts and reserves it
+		/// This will create the target directory, deleting it first if it exists already
 		/// </summary>
-		/// <returns></returns>
-		protected String ReserveArtifactPath()
+		/// <returns>The path to the created artifact directory</returns>
+		protected string ReserveArtifactPath()
 		{
 			// Either use the ArtifactName param or name of this test
 			string TestFolder = string.IsNullOrEmpty(Context.Options.ArtifactName) ? this.ToString() : Context.Options.ArtifactName;
@@ -756,8 +758,16 @@ namespace Gauntlet
 
 			} while (ArtifactPathIsTaken);
 
-			ReservedArtifcactPaths.Add(ArtifactPath);
+			DirectoryInfo ArtifactDirectory = new DirectoryInfo(ArtifactPath);
+			if (ArtifactDirectory.Exists)
+			{
+				SystemHelpers.Delete(ArtifactDirectory, true, true);
+			}
+
+			Log.Info("Creating artifact path {Path}", ArtifactPath);
+			Directory.CreateDirectory(ArtifactPath);
 
+			ReservedArtifcactPaths.Add(ArtifactPath);
 			return ArtifactPath;
 		}
 
@@ -790,11 +800,6 @@ namespace Gauntlet
 
 			ReserveArtifactPath();
 
-			// We need to create this directory at the start of the test rather than the end of the test - we are running into instances where multiple A/B tests
-			// on the same build are seeing the directory as non-existent and thinking it is safe to write to.
-			Log.Info("UnrealTestNode.StartTest Calling CreateDirectory for artifacts at {Path}", ArtifactPath);
-			Directory.CreateDirectory(ArtifactPath);
-
 			// Launch the test
 			TestInstance = UnrealApp.LaunchSession();
 			if (TestInstance == null)
@@ -999,8 +1004,6 @@ namespace Gauntlet
 			LastActiveHeartbeatTime = DateTime.MinValue;
 
 			ReserveArtifactPath();
-			Log.Info("UnrealTestNode.ReStartTest Calling CreateDirectory for artifacts at {Path}", ArtifactPath);
-			Directory.CreateDirectory(ArtifactPath);
 
 			// Relaunch the test
 			TestInstance = UnrealApp.RestartSession();
