Changing CustomRetryLogic to CheckShouldRetry to make it more explicit, resetting Retries count if tests started successfully

#rnx

diff --git a/Engine/Source/Programs/AutomationTool/Gauntlet/Unreal/Base/Gauntlet.UnrealTestNode.cs b/Engine/Source/Programs/AutomationTool/Gauntlet/Unreal/Base/Gauntlet.UnrealTestNode.cs
--- a/Engine/Source/Programs/AutomationTool/Gauntlet/Unreal/Base/Gauntlet.UnrealTestNode.cs
+++ b/Engine/Source/Programs/AutomationTool/Gauntlet/Unreal/Base/Gauntlet.UnrealTestNode.cs
@@ -273,44 +273,32 @@ namespace Gauntlet
 		public virtual bool CanRetry() { return Retries < MaxRetries; }
 		public virtual bool SetToRetryIfPossible()
 		{
-			if (!CanRetry()) { return false; }
+			if (!CanRetry()) 
+			{
+				Log.Warning($"Max retries have been exceeded, attempted retry {Retries}/{MaxRetries} times.");
+				return false; 
+			}
 			if (IsSetToRetry()) { return true; }
+
 			Retries++;
 			SetUnrealTestResult(TestResult.WantRetry);
+			Log.Warning($"Test node is set to retry, test will now retry (Attempt {Retries}/{MaxRetries})");
 			return true;
 		}
 
 		/// <summary>
 		/// PLEASE READ BEFORE OVERRIDING:
-		/// This should only be used in very specific instances to retry for user defined dependency/service failure strings
-		/// False - Do not retry, user defined retry logic did not find a reason to restart tests
-		/// True - User defined logic found a reason to start tests, perform a retry
+		/// This checks a condition defined in a test node at the end of the test to determine whether or not the test should run again.
+		/// This should only be used in very specific instances to retry for user defined dependency/service failure string.
 		/// Example: Useful for unexpected unattended log outs during testing
 		/// </summary>
-		protected virtual bool CustomRetryLogic()
+		protected virtual bool CheckShouldRetry()
 		{
 			return false;
 		}
 
 		public virtual bool IsSetToRetry()
 		{
-			if (CustomRetryLogic())
-			{
-				if (CanRetry())
-				{
-					Retries++;
-					SetUnrealTestResult(TestResult.WantRetry);
-					// This is logged as an error to raise visibility of retry attempts for blocked tests
-					Log.Error($"CustomRetryLogic detected an issue, test will now retry (Attempt {Retries}/{MaxRetries})");
-					return true;
-				}
-				else
-				{
-					// This is logged as an error to raise visibility of retry attempts for blocked tests
-					Log.Error($"Max retries triggered by CustomRetryLogic has exceeded, issues have been found {Retries}/{MaxRetries} times. Creating report and ending the Gauntlet test run.");
-					return false;
-				}
-			}
 			return GetTestResult() == TestResult.WantRetry;
 		}
 
@@ -1288,7 +1276,7 @@ namespace Gauntlet
 		/// <returns></returns>
 		public virtual ITestReport CreateReport(TestResult Result, UnrealTestContext Context, UnrealBuildSource Build, IEnumerable<UnrealRoleResult> RoleResults, string ArtifactPath)
 		{
-			if (IsSetToRetry()) { return null; }
+			if ((CheckShouldRetry() && SetToRetryIfPossible()) || IsSetToRetry()) { return null; }
 			if (GetCachedConfiguration().WriteTestResultsForHorde)
 			{
 				// write test report for Horde
