Add support for builds API filesystem snapshot descriptors to ZenImportSnapshotTask.

[RN] none

diff --git a/Engine/Source/Programs/AutomationTool/BuildGraph/Tasks/ZenImportSnapshotTask.cs b/Engine/Source/Programs/AutomationTool/BuildGraph/Tasks/ZenImportSnapshotTask.cs
--- a/Engine/Source/Programs/AutomationTool/BuildGraph/Tasks/ZenImportSnapshotTask.cs
+++ b/Engine/Source/Programs/AutomationTool/BuildGraph/Tasks/ZenImportSnapshotTask.cs
@@ -107,6 +107,10 @@ namespace AutomationTool.Tasks
 			/// Snapshot stored as a file on disk.
 			/// </summary>
 			File,
+			/// <summary>
+			/// Snapshot stored in Unreal Cloud builds API.
+			/// </summary>
+			Builds,
 		}
 
 		/// <summary>
@@ -159,6 +163,12 @@ namespace AutomationTool.Tasks
 			/// </summary>
 			public string Key { get; set; }
 
+			/// <summary>
+			/// For builds snapshots, the builds ID that identifies them.
+			/// </summary>
+			[JsonPropertyName("builds-id")]
+			public string BuildsId { get; set; }
+
 			/// <summary>
 			/// For file snapshots, the directory it is stored in.
 			/// </summary>
@@ -454,6 +464,12 @@ namespace AutomationTool.Tasks
 					oplogImportCommandLine.AppendFormat(" --file {0}", descriptor.Directory);
 					oplogImportCommandLine.AppendFormat(" --name {0}", descriptor.Filename);
 					break;
+				case SnapshotStorageType.Builds:
+					oplogImportCommandLine.AppendFormat(" --builds {0}", descriptor.Host);
+					oplogImportCommandLine.AppendFormat(" --builds-id {0}", descriptor.BuildsId);
+					oplogImportCommandLine.AppendFormat(" --namespace {0}", descriptor.Namespace);
+					oplogImportCommandLine.AppendFormat(" --bucket {0}", descriptor.Bucket);
+					break;
 				default:
 					break;
 				}
