Fix wrong stage directory when creating a workspace

#rnx

diff --git a/Engine/Source/Programs/AutomationTool/Scripts/CopyBuildToStagingDirectory.Automation.cs b/Engine/Source/Programs/AutomationTool/Scripts/CopyBuildToStagingDirectory.Automation.cs
--- a/Engine/Source/Programs/AutomationTool/Scripts/CopyBuildToStagingDirectory.Automation.cs
+++ b/Engine/Source/Programs/AutomationTool/Scripts/CopyBuildToStagingDirectory.Automation.cs
@@ -1098,7 +1098,7 @@ namespace AutomationScripts
 
 			DirectoryReference StageBaseDir = DirectoryReference.Combine(SC.ProjectRoot, "Saved", "StagedBuilds");
 			ZenWorkspace BaseWorkspace = ZenServerContext.FindWorkspaceForDir(StageBaseDir);
-			DirectoryReference ShareDir = DirectoryReference.Combine(StageBaseDir, SC.PlatformDir);
+			DirectoryReference ShareDir = DirectoryReference.Combine(StageBaseDir, SC.FinalCookPlatform);
 			if (BaseWorkspace == null)
 			{
 				return;
@@ -3118,8 +3118,7 @@ namespace AutomationScripts
 
 				if (ShouldCreatePak(Params))
 				{
-					DirectoryReference StageBaseDir = DirectoryReference.Combine(SC.ProjectRoot, "Saved", "StagedBuilds", SC.PlatformDir);
-					ZenWorkspaceShare WorkspaceShare = ZenServerContext.CreateWorkspaceAndShare(StageBaseDir);
+					ZenWorkspaceShare WorkspaceShare = ZenServerContext.CreateWorkspaceAndShare(SC.StageDirectory);
 				}
 			}
 		}
