logging:

added -NODEBUGOUTPUT -SuppressConsoleOutputVerboseLogging to the server cmdlines which helps perf especially when verbose+ logging is turned on

added -logdebug and -logdebugv to easily turn on a bunch of our relevant categories to log to file Verbose or VeryVerbose

diff --git a/Engine/Source/Programs/AutomationTool/Scripts/LaunchMultiServer.Automation.cs b/Engine/Source/Programs/AutomationTool/Scripts/LaunchMultiServer.Automation.cs
--- a/Engine/Source/Programs/AutomationTool/Scripts/LaunchMultiServer.Automation.cs
+++ b/Engine/Source/Programs/AutomationTool/Scripts/LaunchMultiServer.Automation.cs
@@ -426,7 +426,7 @@ namespace AutomationTool
 			var CommonArgs = (Params.Cook ? "" : DeployContext.ProjectArgForCommandLines + " ");
 			if (LaunchType == ProcessLaunchType.Server || LaunchType == ProcessLaunchType.ProxyServer)
 			{
-				CommonArgs += Params.MapToRun + " -server -log ";
+				CommonArgs += Params.MapToRun + " -server -log -NODEBUGOUTPUT -SuppressConsoleOutputVerboseLogging ";
 
 				// Here are some examples of adding parameters for debugging purposes (note the trailing space within the quotes)
 				// CommonArgs += "-dpcvars=net.MaxConstructedPartialBunchSizeBytes=131071,dstm.SimpleCoinBots=0 ";
@@ -443,6 +443,15 @@ namespace AutomationTool
 				{
 					CommonArgs += " -nodebugoutput -noconsole ";
 				}
+				
+				var Cats = new[] { "LogDstm", "LogDstmObject", "LogDstmDemo", "LogRemoteExec", "LogRemoteObject", "LogDstmPushMigration" };
+
+				foreach (var Cat in Cats)
+				{
+					if (ParseParam("logdebug")) { CommonArgs += string.Format("-LogCmds=\"{0} Verbose\" ", Cat); }
+					else if (ParseParam("logdebugv")) { CommonArgs += string.Format("-LogCmds=\"{0} VeryVerbose\" ", Cat); }
+				}
+
 			}
 
 			CommonArgs += " -messaging ";
