LaunchMultiServer: Disable sound when running unattended clients

diff --git a/Engine/Source/Programs/AutomationTool/Scripts/LaunchMultiServer.Automation.cs b/Engine/Source/Programs/AutomationTool/Scripts/LaunchMultiServer.Automation.cs
--- a/Engine/Source/Programs/AutomationTool/Scripts/LaunchMultiServer.Automation.cs
+++ b/Engine/Source/Programs/AutomationTool/Scripts/LaunchMultiServer.Automation.cs
@@ -369,7 +369,7 @@ namespace AutomationTool
 								for (int BotIndex = 0; BotIndex < NumBotsToSpawn; BotIndex++)
 								{
 									++TotalBotCount;
-									var BotClientAdditionalParams = String.Format("127.0.0.1:{0} -log={2}ProxyBot-{1}.log -log -nullrhi -bot -unattended -nosplash ", ProxyListeningPort, TotalBotCount, ProjectName);
+									var BotClientAdditionalParams = String.Format("127.0.0.1:{0} -log={2}ProxyBot-{1}.log -log -nullrhi -bot -nosound -unattended -nosplash ", ProxyListeningPort, TotalBotCount, ProjectName);
 									BotClientAdditionalParams += String.Format("-ConsoleTitle=\"{0} Proxy Bot Client {1}\" -SessionName=\"{0} Proxy Bot Client {1}\" ", ProjectName, TotalBotCount);
 									if (!ParseParam("nonewconsole"))
 									{
