Avoid using callstack property and use Gauntlet_TestEvent to Gauntlet log output to avoid noise in Horde log

#rnx

diff --git a/Engine/Source/Programs/AutomationTool/Gauntlet/Framework/Gauntlet.TestExecutor.cs b/Engine/Source/Programs/AutomationTool/Gauntlet/Framework/Gauntlet.TestExecutor.cs
--- a/Engine/Source/Programs/AutomationTool/Gauntlet/Framework/Gauntlet.TestExecutor.cs
+++ b/Engine/Source/Programs/AutomationTool/Gauntlet/Framework/Gauntlet.TestExecutor.cs
@@ -604,7 +604,7 @@ namespace Gauntlet
 			}
 			catch (Exception Ex)
 			{
-				Log.Error(KnownLogEvents.Gauntlet_TestEvent, "Test {Name} threw an exception during launch. Skipping test. Ex: {Exception}\n{Callstack}", Name, Ex.Message, Ex.StackTrace);
+				Log.Error(KnownLogEvents.Gauntlet_TestEvent, "Test {Name} threw an exception during launch. Skipping test. Ex: {Exception}\n{StackTrace}", Name, Ex.Message, Ex.StackTrace);
 			}			
 
 			return false;			
@@ -868,12 +868,12 @@ namespace Gauntlet
 				{
 					if (TestIsRunning)
 					{
-						Log.Warning(KnownLogEvents.Gauntlet_TestEvent, "Cancelled Test {Name} threw an exception while stopping. Ex: {Exception}\n{Callstack}", 
+						Log.Warning(KnownLogEvents.Gauntlet_TestEvent, "Cancelled Test {Name} threw an exception while stopping. Ex: {Exception}\n{StackTrace}", 
 							TestInfo.TestNode.Name, ex.Message, ex.StackTrace);
 					}
 					else
 					{
-						Log.Error(KnownLogEvents.Gauntlet_TestEvent, "Test {Name} threw an exception while stopping. Ex: {Exception}\n{Callstack}",
+						Log.Error(KnownLogEvents.Gauntlet_TestEvent, "Test {Name} threw an exception while stopping. Ex: {Exception}\n{StackTrace}",
 							TestInfo.TestNode.Name, ex.Message, ex.StackTrace);
 					}
 
diff --git a/Engine/Source/Programs/AutomationTool/Gauntlet/Platform/Windows/Gauntlet.TargetDeviceWindows.cs b/Engine/Source/Programs/AutomationTool/Gauntlet/Platform/Windows/Gauntlet.TargetDeviceWindows.cs
--- a/Engine/Source/Programs/AutomationTool/Gauntlet/Platform/Windows/Gauntlet.TargetDeviceWindows.cs
+++ b/Engine/Source/Programs/AutomationTool/Gauntlet/Platform/Windows/Gauntlet.TargetDeviceWindows.cs
@@ -414,7 +414,7 @@ namespace Gauntlet
 						}
 						if (Stack.Any())
 						{
-							Log.Error(KnownLogEvents.Gauntlet_TestEvent, "Stack analysis produced with 'cdb {Args}':\n{Callstack}", Args, string.Join('\n', Stack));
+							Log.Error(KnownLogEvents.Gauntlet_TestEvent, "Stack analysis produced with 'cdb {Args}':\n{StackTrace}", Args, string.Join('\n', Stack));
 						}
 						else
 						{
diff --git a/Engine/Source/Programs/AutomationTool/Gauntlet/Unreal/Base/Gauntlet.UnrealTestNode.cs b/Engine/Source/Programs/AutomationTool/Gauntlet/Unreal/Base/Gauntlet.UnrealTestNode.cs
--- a/Engine/Source/Programs/AutomationTool/Gauntlet/Unreal/Base/Gauntlet.UnrealTestNode.cs
+++ b/Engine/Source/Programs/AutomationTool/Gauntlet/Unreal/Base/Gauntlet.UnrealTestNode.cs
@@ -275,14 +275,14 @@ namespace Gauntlet
 		{
 			if (!CanRetry()) 
 			{
-				Log.Warning($"Max retries have been exceeded, attempted retry {Retries}/{MaxRetries} times.");
+				Log.Warning(KnownLogEvents.Gauntlet_TestEvent, $"Max retries have been exceeded, attempted retry {Retries}/{MaxRetries} times.");
 				return false; 
 			}
 			if (IsSetToRetry()) { return true; }
 
 			Retries++;
 			SetUnrealTestResult(TestResult.WantRetry);
-			Log.Warning($"Test node is set to retry, test will now retry (Attempt {Retries}/{MaxRetries})");
+			Log.Warning(KnownLogEvents.Gauntlet_TestEvent, $"Test node is set to retry, test will now retry (Attempt {Retries}/{MaxRetries})");
 			return true;
 		}
 
