Initial commit for Zen Streaming tests for win64

This runs the tests using lyra on win64 against the shared Zen server.

The buildcommand SetOIDCToken is a productivity enhancer for local dev/test, since the token expires after a brief period of time.

Followup commits will add support for additional platforms.

#rnx

diff --git a/Engine/Source/Programs/AutomationTool/Gauntlet/Unreal/Automation/UE.ZenStreaming.cs b/Engine/Source/Programs/AutomationTool/Gauntlet/Unreal/Automation/UE.ZenStreaming.cs
new file mode 100644
--- /dev/null
+++ b/Engine/Source/Programs/AutomationTool/Gauntlet/Unreal/Automation/UE.ZenStreaming.cs
@@ -0,0 +1,72 @@
+// Copyright Epic Games, Inc. All Rights Reserved.
+
+using System;
+using System.IO;
+using Gauntlet;
+using System.Linq;
+using Microsoft.Extensions.Logging;
+using AutomationTool;
+
+namespace UE
+{
+	public class ZenStreaming: UnrealTestNode<UnrealTestConfiguration>
+	{
+		private const string StreamingStatsFile = "ZenStreamingClient.log";
+		private UnrealLogStreamParser ClientLogParser;
+		private readonly string ReplayFile;
+		private string ZenArgs;
+
+		public ZenStreaming(UnrealTestContext InContext) : base(InContext)
+		{
+			File.WriteAllText(StreamingStatsFile, string.Empty);
+			ClientLogParser = new UnrealLogStreamParser();
+
+			ZenArgs = Globals.Params.ParseValue("zenargs", string.Empty);
+			ReplayFile = Globals.Params.ParseValue("replayfile", string.Empty);
+		}
+
+		public override UnrealTestConfiguration GetConfiguration()
+		{
+			UnrealTestConfiguration Config = base.GetConfiguration();
+			UnrealTestRole ClientRole = Config.RequireRole(UnrealTargetRole.Client);
+
+			string command = @$" {ZenArgs} -execcmds=""DemoPlay {ReplayFile}"" -exitafterreplay";
+			ClientRole.CommandLineParams.Add(command);
+			
+			return Config;
+		}
+
+		public override void StopTest(StopReason InReason)
+		{
+			ReadClientLogStream();
+			ParseStreamingStats();
+
+			base.StopTest(InReason);
+		}
+
+		private void ReadClientLogStream()
+		{
+			IAppInstance ClientApp = TestInstance.ClientApps.FirstOrDefault();
+
+			if (ClientApp is null)
+			{
+				return;
+			}
+
+			ClientLogParser.ReadStream(ClientApp.StdOut);
+		}
+
+		private void ParseStreamingStats()
+		{
+			string[] LogZenLines = ClientLogParser.GetLogFromChannel("LogStreaming").ToArray();
+
+			if (!LogZenLines.Any())
+			{
+				return;
+			}
+
+			string LogLinesToSave = string.Join(Environment.NewLine, LogZenLines);
+			File.AppendAllText(StreamingStatsFile, LogLinesToSave);
+		}
+	}
+}
\ No newline at end of file
