Decouple -snapshot from -cook and add -incremental parameter to BuildCookRun.

[RN] none

diff --git a/Engine/Source/Programs/AutomationTool/AutomationUtils/ProjectParams.cs b/Engine/Source/Programs/AutomationTool/AutomationUtils/ProjectParams.cs
--- a/Engine/Source/Programs/AutomationTool/AutomationUtils/ProjectParams.cs
+++ b/Engine/Source/Programs/AutomationTool/AutomationUtils/ProjectParams.cs
@@ -30,7 +30,8 @@ namespace AutomationTool
 	[ParamHelp("run", "Run the game after it is built (including server, if -server)", ParamType = typeof(bool))]
 	[ParamHelp("zenstore", "Save cooked output data to the Zen storage server", ParamType = typeof(bool))]
 	[ParamHelp("iterate", "Uses the iterative cooking/deploy", ParamType = typeof(bool))]
-	[ParamHelp("snapshot", "Uses incremental cook on top of the most recently published snapshot", ParamType = typeof(bool))]
+	[ParamHelp("snapshot", "Imports the most recently published project snapshot (if cooking, this forces the use of -cookincremental)", ParamType = typeof(bool))]
+	[ParamHelp("cookincremental", "Cooks incrementally using previously-cooked data or a published snapshot as a base", ParamType = typeof(bool))]
 	public interface IProjectParamsHelpers
 	{
 	}
@@ -401,6 +402,7 @@ namespace AutomationTool
 			this.KeepFileOpenLog = InParams.KeepFileOpenLog;
 			this.FastCook = InParams.FastCook;
 			this.Snapshot = InParams.Snapshot;
+			this.IncrementalCook = InParams.IncrementalCook;
 			this.Devices = InParams.Devices;
 			this.DeviceNames = InParams.DeviceNames;
 			this.ServerDevice = InParams.ServerDevice;
@@ -612,6 +614,7 @@ namespace AutomationTool
 			bool? IterativeDeploy = null,
 			bool? FastCook = null,
 			bool? Snapshot = null,
+			bool? IncrementalCook = null,
 			bool? IgnoreCookErrors = null,
 			bool? KeepFileOpenLog = null,
 			bool? CodeSign = null,
@@ -965,6 +968,9 @@ namespace AutomationTool
 			this.IterativeDeploy = GetParamValueIfNotSpecified(Command, IterativeDeploy, this.IterativeDeploy, new string[] {"iterativedeploy", "iterate" } );
 			this.FastCook = GetParamValueIfNotSpecified(Command, FastCook, this.FastCook, "FastCook");
 			this.Snapshot = GetParamValueIfNotSpecified(Command, Snapshot, this.Snapshot, "Snapshot");
+
+			// -snapshot implies the use of -cookincremental for cooking.
+			this.IncrementalCook = this.Snapshot || GetParamValueIfNotSpecified(Command, IncrementalCook, this.IncrementalCook, "cookincremental");
 			this.IgnoreCookErrors = GetParamValueIfNotSpecified(Command, IgnoreCookErrors, this.IgnoreCookErrors, "IgnoreCookErrors");
 			this.KeepFileOpenLog = GetParamValueIfNotSpecified(Command, KeepFileOpenLog, this.KeepFileOpenLog, "KeepFileOpenLog");
 
@@ -2054,11 +2060,17 @@ namespace AutomationTool
 		public bool FastCook;
 
 		/// <summary>
-		/// Cook: Retrieve the most recently published snapshot and cook incrementally on top of that
+		/// Imports the most recently published project snapshot (if cooking, this forces the use of -cookincremental)
 		/// </summary>
-		[Help("Snapshot", "Uses incremental cooking on top of the most recently published snapshot for the project")]
+		[Help("snapshot", "Imports the most recently published project snapshot (if cooking, this forces the use of -cookincremental)")]
 		public bool Snapshot;
 
+		/// <summary>
+		/// Cook: Cooks incrementally using previously-cooked data or a published snapshot as a base
+		/// </summary>
+		[Help("cookincremental", "Cooks incrementally using previously-cooked data or a published snapshot as a base")]
+		public bool IncrementalCook;
+
 		/// <summary>
 		/// Cook: Ignores cook errors and continues with packaging etc.
 		/// </summary>
diff --git a/Engine/Source/Programs/AutomationTool/Scripts/CookCommand.Automation.cs b/Engine/Source/Programs/AutomationTool/Scripts/CookCommand.Automation.cs
--- a/Engine/Source/Programs/AutomationTool/Scripts/CookCommand.Automation.cs
+++ b/Engine/Source/Programs/AutomationTool/Scripts/CookCommand.Automation.cs
@@ -55,7 +55,7 @@ namespace AutomationScripts
 			{
 				CommandletParams += " -FastCook";
 			}
-			if (Params.Snapshot)
+			if (Params.IncrementalCook)
 			{
 				CommandletParams += " -cookincremental";
 			}
@@ -101,7 +101,7 @@ namespace AutomationScripts
 			return CommandletParams;
 		}
 
-		public static void GetCookByTheBookCommandletParams(ProjectParams Params, out string CommandletParams, out string[] Maps, out string[] DirectoriesToCook, out string InternationalizationPreset, out string[] CulturesToCook, out string PlatformsToCook)
+		private static HashSet<string> GetPlatformsToCook(ProjectParams Params)
 		{
 			var PlatformsToCookSet = new HashSet<string>();
 			if (!Params.NoClient)
@@ -125,6 +125,12 @@ namespace AutomationScripts
 				}
 			}
 
+			return PlatformsToCookSet;
+		}
+
+		public static void GetCookByTheBookCommandletParams(ProjectParams Params, out string CommandletParams, out string[] Maps, out string[] DirectoriesToCook, out string InternationalizationPreset, out string[] CulturesToCook, out string PlatformsToCook)
+		{
+			var PlatformsToCookSet = GetPlatformsToCook(Params);
 			PlatformsToCook = CombineCommandletParams(PlatformsToCookSet);
 
 			if (Params.Clean.HasValue && Params.Clean.Value && !Params.IterativeCooking)
@@ -244,6 +250,18 @@ namespace AutomationScripts
 
 		public static void Cook(ProjectParams Params)
 		{
+			if (Params.Snapshot)
+			{
+				Logger.LogInformation("********** SNAPSHOT IMPORT STARTED **********");
+
+				var ImportStartTime = DateTime.UtcNow;
+				string PlatformsToCook = CombineCommandletParams(GetPlatformsToCook(Params));
+				ImportSnapshots(Params, PlatformsToCook);
+
+				Logger.LogInformation("Snapshot import time: {0:0.00} s", (DateTime.UtcNow - ImportStartTime).TotalMilliseconds / 1000);
+				Logger.LogInformation("********** SNAPSHOT IMPORT COMPLETED **********");
+			}
+
 			if ((!Params.Cook && !(Params.CookOnTheFly && !Params.SkipServer)) || Params.SkipCook)
 			{
 				return;
@@ -297,11 +315,6 @@ namespace AutomationScripts
 					// get the parameters to CBTB
 					GetCookByTheBookCommandletParams(Params, out CommandletParams, out MapsToCook, out DirectoriesToCook, out InternationalizationPreset, out CulturesToCook, out PlatformsToCook);
 
-					if (Params.Snapshot)
-					{
-						ImportSnapshots(Params, PlatformsToCook);
-					}
-
 					// run a blocking cook
 					CookCommandlet(Params.RawProjectPath, Params.UnrealExe, MapsToCook, DirectoriesToCook, InternationalizationPreset, CulturesToCook, PlatformsToCook, CommandletParams);
 				}
