ZenLoader - Quick Launch Automation (Part 1)

1. The Gauntlet test is a simple client boot test (additional checks will be added in the following commits).
2. The workaround is implemented until the bug UE-254908 is fixed
3. The buildgraph implements 15 combinations of packaging parameters and runs the Gauntlet test
4. A Platform is Win64

#rnx

diff --git a/Engine/Source/Programs/AutomationTool/Gauntlet/Unreal/Automation/UE.ZenLoaderTest.cs b/Engine/Source/Programs/AutomationTool/Gauntlet/Unreal/Automation/UE.ZenLoaderTest.cs
new file mode 100644
--- /dev/null
+++ b/Engine/Source/Programs/AutomationTool/Gauntlet/Unreal/Automation/UE.ZenLoaderTest.cs
@@ -0,0 +1,39 @@
+// Copyright Epic Games, Inc. All Rights Reserved.
+
+using System.IO;
+using Gauntlet;
+using UnrealBuildTool;
+
+namespace UE
+{
+	public class ZenLoaderTest : UnrealTestNode<UnrealTestConfiguration>
+	{
+		public ZenLoaderTest(UnrealTestContext InContext) : base(InContext)
+		{
+			string ContentPath = Path.Combine(Context.Options.Build, GetTargetPlatformName(), Context.Options.Project, "Content");
+
+			if (!Directory.Exists(ContentPath))
+			{
+				Directory.CreateDirectory(ContentPath); // TODO remove after the bug UE-254908 is fixed
+			}
+		}
+
+		public override UnrealTestConfiguration GetConfiguration()
+		{
+			UnrealTestConfiguration Config = base.GetConfiguration();
+			UnrealTestRole ClientRole = Config.RequireRole(UnrealTargetRole.Client);
+
+			ClientRole.CommandLine += @" -log -ExecCmds=""Automation SoftQuit""";
+
+			return Config;
+		}
+
+		private string GetTargetPlatformName()
+		{
+			UnrealTargetPlatform? TargetPlatform = Context.Constraint.Platform ?? UnrealTargetPlatform.Win64;
+			string TargetPlatformName = TargetPlatform == UnrealTargetPlatform.Win64 ? "Windows" : TargetPlatform.ToString();
+
+			return TargetPlatformName;
+		}
+	}
+}
