[Gauntlet] Adds new UnrealTestRole option to skip the clean of the device cache between test runs

diff --git a/Engine/Source/Programs/AutomationTool/Gauntlet/Unreal/Base/Gauntlet.UnrealSession.cs b/Engine/Source/Programs/AutomationTool/Gauntlet/Unreal/Base/Gauntlet.UnrealSession.cs
--- a/Engine/Source/Programs/AutomationTool/Gauntlet/Unreal/Base/Gauntlet.UnrealSession.cs
+++ b/Engine/Source/Programs/AutomationTool/Gauntlet/Unreal/Base/Gauntlet.UnrealSession.cs
@@ -152,6 +152,11 @@ namespace Gauntlet
 		/// </summary>
 		public bool CreateGifFromScreenshots { get; set; }
 
+		/// <summary>
+		/// Whether or not this role should skip the cleaning of its device's artifact between each test run
+		/// </summary>
+		public bool SkipCleanDeviceArtifacts { get; set; }
+
 		/// <summary>
 		/// Is this role Null?
 		/// </summary>
@@ -1635,7 +1640,10 @@ namespace Gauntlet
 
 				try
 				{
-					Device.CleanArtifacts();
+					if (!Role.SkipCleanDeviceArtifacts)
+					{
+						Device.CleanArtifacts();
+					}
 
 					if (RoleNeedsInstall(Role))
 					{
diff --git a/Engine/Source/Programs/AutomationTool/Gauntlet/Unreal/Base/Gauntlet.UnrealTestConfiguration.cs b/Engine/Source/Programs/AutomationTool/Gauntlet/Unreal/Base/Gauntlet.UnrealTestConfiguration.cs
--- a/Engine/Source/Programs/AutomationTool/Gauntlet/Unreal/Base/Gauntlet.UnrealTestConfiguration.cs
+++ b/Engine/Source/Programs/AutomationTool/Gauntlet/Unreal/Base/Gauntlet.UnrealTestConfiguration.cs
@@ -505,6 +505,11 @@ namespace Gauntlet
 		/// Whether this role will create a gif from screenshots produced as an artifact
 		/// </summary>
 		public bool CreateGifFromScreenshots;
+
+		/// <summary>
+		/// Whether or not this role should skip the cleaning of its device's artifact between each test run
+		/// </summary>
+		public bool SkipCleanDeviceArtifacts { get; set; }
 	}
 
 	/// <summary>
diff --git a/Engine/Source/Programs/AutomationTool/Gauntlet/Unreal/Base/Gauntlet.UnrealTestNode.cs b/Engine/Source/Programs/AutomationTool/Gauntlet/Unreal/Base/Gauntlet.UnrealTestNode.cs
--- a/Engine/Source/Programs/AutomationTool/Gauntlet/Unreal/Base/Gauntlet.UnrealTestNode.cs
+++ b/Engine/Source/Programs/AutomationTool/Gauntlet/Unreal/Base/Gauntlet.UnrealTestNode.cs
@@ -613,6 +613,7 @@ namespace Gauntlet
 					SessionRole.CommandLineParams = TestRole.CommandLineParams;
 					SessionRole.RoleModifier = TestRole.RoleType;
 					SessionRole.Constraint = UseContextConstraint ? Context.Constraint : new UnrealDeviceTargetConstraint(SessionPlatform);
+					SessionRole.SkipCleanDeviceArtifacts = TestRole.SkipCleanDeviceArtifacts;
 					Log.Verbose("Created SessionRole {Role} from RoleContext {Context} (RoleType={Type})", SessionRole, RoleContext, TypesToRoles.Key);
 
 					// TODO - this can all / mostly go into UnrealTestConfiguration.ApplyToConfig
