fix for attempting to stage multiple architectures on Windows

#rnx

diff --git a/Engine/Source/Programs/AutomationTool/Scripts/CopyBuildToStagingDirectory.Automation.cs b/Engine/Source/Programs/AutomationTool/Scripts/CopyBuildToStagingDirectory.Automation.cs
--- a/Engine/Source/Programs/AutomationTool/Scripts/CopyBuildToStagingDirectory.Automation.cs
+++ b/Engine/Source/Programs/AutomationTool/Scripts/CopyBuildToStagingDirectory.Automation.cs
@@ -5891,6 +5891,29 @@ namespace AutomationScripts
 		//@todo move this
 		public static List<DeploymentContext> CreateDeploymentContext(ProjectParams Params, bool InDedicatedServer, bool DoCleanStage = false)
 		{
+			static IEnumerable<Tuple<string, UnrealTargetConfiguration, UnrealArchitectures>> ExpandDeploymentArchitectures(string Target, UnrealTargetConfiguration Configuration, UnrealArchitectures Architectures, UnrealTargetPlatform Platform )
+			{
+				List<Tuple<string, UnrealTargetConfiguration, UnrealArchitectures>> Result = new();
+
+				UnrealArchitectureConfig ArchConfig = UnrealArchitectureConfig.ForPlatform(Platform);
+				if (ArchConfig.Mode == UnrealArchitectureMode.OneTargetPerArchitecture && Architectures != null)
+				{
+					foreach (UnrealArch Architecture in Architectures.Architectures)
+					{
+						Result.Add( new(Target, Configuration, new(Architecture)));
+					}
+				}
+				else
+				{
+					Result.Add( new(Target, Configuration, Architectures));
+				}
+					
+				return Result;
+			};
+
+
+
+
 			ParamList<string> ListToProcess;
 			if (Params.IsProgramTarget)
 			{
@@ -5906,31 +5929,6 @@ namespace AutomationScripts
 			} 
 			var ConfigsToProcess = InDedicatedServer && (Params.Cook || Params.CookOnTheFly) ? Params.ServerConfigsToBuild : Params.ClientConfigsToBuild;
 
-			List<Tuple<string, UnrealTargetConfiguration, UnrealArchitectures>> TargetAndConfigPairs = new();
-
-			foreach (var Target in ListToProcess)
-			{
-				// If we are staging a client and have been asked to include editor targets, we currently only want to
-				// include a single Development editor target. Ideally, we should have shipping editor configs and then
-				// just include the requested configs for all targets
-				if (Params.HasEditorTargets && Params.EditorTargets.Contains(Target))
-				{
-					TargetAndConfigPairs.Add(new Tuple<string, UnrealTargetConfiguration, UnrealArchitectures>(Target, UnrealTargetConfiguration.Development, Params.EditorArchitecture));
-				}
-				else if (Params.HasProgramTargets && Params.ProgramTargets.Contains(Target))
-				{
-					TargetAndConfigPairs.Add(new Tuple<string, UnrealTargetConfiguration, UnrealArchitectures>(Target, UnrealTargetConfiguration.Development, Params.ProgramArchitecture));
-				}
-				else
-				{
-					UnrealArchitectures Arches = InDedicatedServer ? Params.ServerArchitecture : Params.ClientArchitecture;
-					foreach (var Config in ConfigsToProcess)
-					{
-						TargetAndConfigPairs.Add(new Tuple<string, UnrealTargetConfiguration, UnrealArchitectures>(Target, Config, Arches));
-					}
-				}
-			}
-
 			List<TargetPlatformDescriptor> PlatformsToStage = Params.ClientTargetPlatforms;
 			if (InDedicatedServer && (Params.Cook || Params.CookOnTheFly))
 			{
@@ -5940,6 +5938,31 @@ namespace AutomationScripts
 			List<DeploymentContext> DeploymentContexts = new List<DeploymentContext>();
 			foreach (var StagePlatform in PlatformsToStage)
 			{
+				List<Tuple<string, UnrealTargetConfiguration, UnrealArchitectures>> TargetAndConfigPairs = new();
+				foreach (var Target in ListToProcess)
+				{
+					// If we are staging a client and have been asked to include editor targets, we currently only want to
+					// include a single Development editor target. Ideally, we should have shipping editor configs and then
+					// just include the requested configs for all targets
+					if (Params.HasEditorTargets && Params.EditorTargets.Contains(Target))
+					{
+						TargetAndConfigPairs.AddRange( ExpandDeploymentArchitectures(Target, UnrealTargetConfiguration.Development, Params.EditorArchitecture, StagePlatform.Type ));
+					}
+					else if (Params.HasProgramTargets && Params.ProgramTargets.Contains(Target))
+					{
+						TargetAndConfigPairs.AddRange( ExpandDeploymentArchitectures(Target, UnrealTargetConfiguration.Development, Params.ProgramArchitecture, StagePlatform.Type ));
+					}
+					else
+					{
+						UnrealArchitectures Arches = InDedicatedServer ? Params.ServerArchitecture : Params.ClientArchitecture;
+						foreach (var Config in ConfigsToProcess)
+						{
+							TargetAndConfigPairs.AddRange( ExpandDeploymentArchitectures(Target, UnrealTargetConfiguration.Development, Arches, StagePlatform.Type ));
+						}
+					}
+				}
+
+
 				// Get the platform to get cooked data from, may differ from the stage platform
 				TargetPlatformDescriptor CookedDataPlatform = Params.GetCookedDataPlatformForClientTarget(StagePlatform);
 
