Stage VC redist not UE prereq + add arm64 VC redist

#rnx

diff --git a/Engine/Source/Programs/AutomationTool/Win/WinPlatform.Automation.cs b/Engine/Source/Programs/AutomationTool/Win/WinPlatform.Automation.cs
--- a/Engine/Source/Programs/AutomationTool/Win/WinPlatform.Automation.cs
+++ b/Engine/Source/Programs/AutomationTool/Win/WinPlatform.Automation.cs
@@ -168,7 +168,9 @@ public class Win64Platform : Platform
 
 		if (Params.Prereqs)
 		{
-			SC.StageFile(StagedFileType.NonUFS, FileReference.Combine(SC.EngineRoot, "Extras", "Redist", "en-us", "UEPrereqSetup_x64.exe"));
+			//SC.StageFile(StagedFileType.NonUFS, FileReference.Combine(SC.EngineRoot, "Extras", "Redist", "en-us", "UEPrereqSetup_x64.exe"));
+			SC.StageFile(StagedFileType.NonUFS, FileReference.Combine(SC.EngineRoot, "Extras", "Redist", "en-us", "vc_redist.x64.exe"));
+			SC.StageFile(StagedFileType.NonUFS, FileReference.Combine(SC.EngineRoot, "Extras", "Redist", "en-us", "vc_redist.arm64.exe")); // always staging arm64 redist in case we are run on an arm64 device in emulation mode
 		}
 
 		if (!string.IsNullOrWhiteSpace(Params.AppLocalDirectory))
