Fixed issues with RetrieveCloudArtifact:
* Environment variable should be used before the ini values
* Fixes issue with env. vars not using the correct format for linux and mac
* Handle if host is a semicolon seperate list in the ini
* Improved error message when the host value is not a valid uri

diff --git a/Engine/Source/Programs/AutomationTool/BuildGraph/Tasks/RetrieveCloudArtifactTask.cs b/Engine/Source/Programs/AutomationTool/BuildGraph/Tasks/RetrieveCloudArtifactTask.cs
--- a/Engine/Source/Programs/AutomationTool/BuildGraph/Tasks/RetrieveCloudArtifactTask.cs
+++ b/Engine/Source/Programs/AutomationTool/BuildGraph/Tasks/RetrieveCloudArtifactTask.cs
@@ -187,54 +187,57 @@ namespace AutomationTool.Tasks
 				}
 			}
 
+			string? cloudHostEnvVar = Environment.GetEnvironmentVariable(OperatingSystem.IsWindows() ? "UE-CloudPublishHost" : "UE_CloudPublishHost");
+			
 			string host;
 			if (!String.IsNullOrEmpty(_parameters.Host))
 			{
 				host = _parameters.Host;
 			}
-			else if (foundConfig && cloudConfig != null && !String.IsNullOrEmpty(cloudConfig.Value.Host))
+			else if (!String.IsNullOrEmpty(cloudHostEnvVar))
 			{
-				host = cloudConfig.Value.Host;
+				host = cloudHostEnvVar;
 			}
-			else
+			else if (foundConfig && cloudConfig != null && !String.IsNullOrEmpty(cloudConfig.Value.Host))
 			{
-				string? cloudHostEnvVar = Environment.GetEnvironmentVariable("UE-CloudPublishHost");
-				if (!String.IsNullOrEmpty(cloudHostEnvVar))
-				{
-					host = cloudHostEnvVar;
-				}
-				else
+				string cloudConfigHost = cloudConfig.Value.Host;
+				if (cloudConfigHost.Contains(';', StringComparison.OrdinalIgnoreCase))
 				{
-					throw new AutomationException("Missing UE-CloudPublishHost environment variable; unable to determine cloud host. Specify the environment variable or define it in your StorageServer section in your engine ini");
+					// if its a list pick the first element
+					cloudConfigHost = cloudConfigHost.Split(";").First();
 				}
+
+				host = cloudConfigHost;
+			}
+			else
+			{
+				throw new AutomationException("Missing UE-CloudPublishHost environment variable; unable to determine cloud host. Specify the environment variable or define it in your StorageServer section in your engine ini");
 			}
 
+
 			string ns;
+			string? cloudDefaultNamespaceEnv = Environment.GetEnvironmentVariable(OperatingSystem.IsWindows() ? "UE-CloudPublishNamespace" : "UE_CloudPublishNamespace");
+				
 			if (!String.IsNullOrEmpty(_parameters.Namespace))
 			{
 				ns = _parameters.Namespace;
 			}
+			else if (!String.IsNullOrEmpty(cloudDefaultNamespaceEnv))
+			{
+				ns = cloudDefaultNamespaceEnv;
+			}
 			else if (foundConfig && cloudConfig != null && !String.IsNullOrEmpty(cloudConfig.Value.BuildsNamespace))
 			{
 				ns = cloudConfig.Value.BuildsNamespace;
 			}
 			else
 			{
-				string? cloudDefaultNamespace = Environment.GetEnvironmentVariable("UE-CloudPublishNamespace");
-				if (!String.IsNullOrEmpty(cloudDefaultNamespace))
-				{
-					ns = cloudDefaultNamespace;
-				}
-				else
-				{
-					throw new AutomationException("Missing UE-CloudPublishNamespace environment variable; unable to default namespace please specify it in the task or define it in your StorageServer section in your engine ini.");
-				}
+				throw new AutomationException("Missing UE-CloudPublishNamespace environment variable; unable to default namespace please specify it in the task or define it in your StorageServer section in your engine ini.");
 			}
-
 			string httpVersion = _parameters.HttpVersion;
 			if (String.IsNullOrEmpty(httpVersion) || httpVersion == "None")
 			{
-				string? httpVersionEnvironmentVariable = Environment.GetEnvironmentVariable("UE-CloudPublishHttpVersion");
+				string? httpVersionEnvironmentVariable = Environment.GetEnvironmentVariable(OperatingSystem.IsWindows() ? "UE-CloudPublishHttpVersion": "UE_CloudPublishHttpVersion");
 				if (!String.IsNullOrEmpty(httpVersionEnvironmentVariable))
 				{
 					httpVersion = httpVersionEnvironmentVariable;
@@ -248,7 +251,7 @@ namespace AutomationTool.Tasks
 			}
 			else
 			{
-				string? cloudAccessTokenEnvVar = Environment.GetEnvironmentVariable("UE-CloudDataCacheAccessToken");
+				string? cloudAccessTokenEnvVar = Environment.GetEnvironmentVariable(OperatingSystem.IsWindows() ? "UE-CloudDataCacheAccessToken" : "UE_CloudDataCacheAccessToken");
 				if (!String.IsNullOrEmpty(cloudAccessTokenEnvVar))
 				{
 					accessToken = cloudAccessTokenEnvVar;
@@ -381,7 +384,15 @@ namespace AutomationTool.Tasks
 		static HttpClient BuildHttpClient(string host, string ns, string bucketId, string accessToken)
 		{
 			HttpClient httpClient = new HttpClient();
-			httpClient.BaseAddress = new Uri($"{host}/api/v2/builds/{ns}/{bucketId}/");
+			string url = $"{host}/api/v2/builds/{ns}/{bucketId}/";
+			try
+			{
+				httpClient.BaseAddress = new Uri(url);
+			}
+			catch (UriFormatException)
+			{
+				throw new AutomationException($"{url} is not a valid url. Make sure host, namespace and bucket have been specified");
+			}
 
 			CbObjectId sessionId = CbObjectId.NewObjectId();
 			Logger.LogInformation("Using SessionId {SessionId} please include this in any error reports.", sessionId);
