Gauntlet tweaks to support async RPCs

diff --git a/Engine/Source/Programs/AutomationTool/Gauntlet/Framework/Gauntlet.Exception.cs b/Engine/Source/Programs/AutomationTool/Gauntlet/Framework/Gauntlet.Exception.cs
--- a/Engine/Source/Programs/AutomationTool/Gauntlet/Framework/Gauntlet.Exception.cs
+++ b/Engine/Source/Programs/AutomationTool/Gauntlet/Framework/Gauntlet.Exception.cs
@@ -8,6 +8,12 @@ namespace Gauntlet
 				: base(Msg)
 		{
 		}
+
+		public TestException(string Message, System.Exception InnerException)
+			: base(Message, InnerException)
+		{
+		}
+
 		public TestException(string Format, params object[] Args)
 				: base(string.Format(Format, Args))
 		{
@@ -20,6 +26,12 @@ namespace Gauntlet
 				: base(Msg)
 		{
 		}
+
+		public DeviceException(string Message, System.Exception InnerException)
+			: base(Message, InnerException)
+		{
+		}
+
 		public DeviceException(string Format, params object[] Args)
 				: base(Format, Args)
 		{
diff --git a/Engine/Source/Programs/AutomationTool/Gauntlet/Framework/RpcFramework/Gauntlet.RpcExecutorAsync.cs b/Engine/Source/Programs/AutomationTool/Gauntlet/Framework/RpcFramework/Gauntlet.RpcExecutorAsync.cs
--- a/Engine/Source/Programs/AutomationTool/Gauntlet/Framework/RpcFramework/Gauntlet.RpcExecutorAsync.cs
+++ b/Engine/Source/Programs/AutomationTool/Gauntlet/Framework/RpcFramework/Gauntlet.RpcExecutorAsync.cs
@@ -17,72 +17,77 @@ public partial class RpcExecutor
 	/// <summary>
 	/// Asynchronously send an RPC to the target.
 	/// </summary>
-	/// <param name="Target">
+	/// <param name="InTarget">
 	/// The <see cref="RpcTarget"/> to send the RPC to.
 	/// </param>
-	/// <param name="RpcName">
+	/// <param name="InRpcName">
 	/// The name of the RPC to send.
 	/// </param>
-	/// <param name="Args">
+	/// <param name="InArgs">
 	/// A map of parameters
 	/// </param>
-	/// <param name="CancellationToken">
+	/// <param name="InCancellationToken">
 	/// The <see cref="CancellationToken"/> to monitor for a cancellation request.
 	/// </param>
 	/// <returns>
 	/// A <see cref="Task{TResult}"/> representing the target's response.
 	/// </returns>
 	/// <exception cref="ArgumentException">
-	/// The <paramref name="Target"/>, <paramref name="RpcName"/>, or <paramref name="Args"/> parameters are invalid.
+	/// The <paramref name="InTarget"/>, <paramref name="InRpcName"/>, or <paramref name="InArgs"/> parameters are invalid.
 	/// </exception>
 	public static Task<HttpResponseMessage> CallRpcAsync(
-		RpcTarget Target,
-		string RpcName,
-		IReadOnlyDictionary<string,object> Args,
-		CancellationToken CancellationToken = default)
+		RpcTarget InTarget,
+		string InRpcName,
+		IReadOnlyDictionary<string, object> InArgs,
+		CancellationToken InCancellationToken = default)
 	{
-		ArgumentNullException.ThrowIfNull(Target);
-		ArgumentException.ThrowIfNullOrEmpty(RpcName);
-		ArgumentNullException.ThrowIfNull(Args);
+		ArgumentNullException.ThrowIfNull(InTarget);
+		ArgumentException.ThrowIfNullOrEmpty(InRpcName);
+		ArgumentNullException.ThrowIfNull(InArgs);
 
-		if (!TryGetRpcDefinition(Target, RpcName, out var RpcEntry))
+		if (!TryGetRpcDefinition(InTarget, InRpcName, out var RpcEntry))
 		{
-			throw new ArgumentException(@"Failed to find RPC definition.", nameof(RpcName));
+			throw new ArgumentException(@"Failed to find RPC definition.", nameof(InRpcName));
 		}
 
-		if (!TryGetRpcUrl(Target, RpcName, out var RpcUrl))
+		if (!TryGetRpcUrl(InTarget, InRpcName, out var RpcUrl))
 		{
 			// There is no reason this would ever hit as the prior exception would have hit
-			throw new ArgumentException(@"Failed to create RPC url.", nameof(RpcName));
+			throw new ArgumentException(@"Failed to create RPC url.", nameof(InRpcName));
 		}
 
-		ValidateRpcArguments(RpcEntry, Args);
+		ValidateRpcArguments(RpcEntry, InArgs);
 
 		var MethodToUse = new HttpMethod(RpcEntry.Verb);
 
 		var RpcRequest = new HttpRequestMessage(MethodToUse, RpcUrl);
 
-		RpcRequest.Headers.Add("rpcname", RpcName);
+		RpcRequest.Headers.Add("rpcname", InRpcName);
 
-		if (Args.Count > 0)
+		if (InArgs.Count > 0)
 		{
-			var JsonBody = JsonConvert.SerializeObject(Args);
+			var JsonBody = JsonConvert.SerializeObject(InArgs);
 
 			RpcRequest.Content = new StringContent(JsonBody);
 		}
 
-		Log.VeryVerbose($"calling RPC {RpcName} on target {Target.TargetName} at {RpcUrl}");
+		Log.VeryVerbose($"calling RPC {InRpcName} on target {InTarget.TargetName} at {RpcUrl}");
 
-		return ExecutorClient.SendAsync(RpcRequest, CancellationToken);
+		return ExecutorClient.SendAsync(RpcRequest, InCancellationToken);
 	}
 
-	private static void ValidateRpcArguments(RpcDefinition RpcDefinition, IReadOnlyDictionary<string, object> Args)
+	private static void ValidateRpcArguments(RpcDefinition InRpcDefinition, IReadOnlyDictionary<string, object> InArgs)
 	{
-		// Should this be a custom exception type. RpcArgumentException or something?
+		// Should this be a custom exception type? RpcArgumentException or something?
 
-		foreach (var RpcArgument in RpcDefinition.Args)
+		if (InRpcDefinition.Args is null)
 		{
-			if (!Args.TryGetValue(RpcArgument.Name, out var Value))
+			return;
+		}
+
+		foreach (var RpcArgument in InRpcDefinition.Args)
+		{
+			if (!InArgs.TryGetValue(RpcArgument.Name, out var Value))
 			{
 				if (!RpcArgument.Optional)
 				{
@@ -123,17 +128,17 @@ public partial class RpcExecutor
 		{ "double", typeof(double) },
 	};
 
-	private static bool TryGetRpcDefinition(RpcTarget Target, string Name, [NotNullWhen(true)] out RpcDefinition? Definition)
+	private static bool TryGetRpcDefinition(RpcTarget InTarget, string InName, [NotNullWhen(true)] out RpcDefinition? OutDefinition)
 	{
-		Definition = Target.GetRpc(Name);
+		OutDefinition = InTarget.GetRpc(InName);
 
-		return Definition is not null;
+		return OutDefinition is not null;
 	}
 
-	private static bool TryGetRpcUrl(RpcTarget Target, string Name, [NotNullWhen(true)] out string? Url)
+	private static bool TryGetRpcUrl(RpcTarget InTarget, string InName, [NotNullWhen(true)] out string? OutUrl)
 	{
-		Url = Target.CreateRpcUrl(Name);
+		OutUrl = InTarget.CreateRpcUrl(InName);
 
-		return !string.IsNullOrEmpty(Url);
+		return !string.IsNullOrEmpty(OutUrl);
 	}
 }
\ No newline at end of file
