fix non-Development staging issue

#rnx

diff --git a/Engine/Source/Programs/AutomationTool/Scripts/CopyBuildToStagingDirectory.Automation.cs b/Engine/Source/Programs/AutomationTool/Scripts/CopyBuildToStagingDirectory.Automation.cs
--- a/Engine/Source/Programs/AutomationTool/Scripts/CopyBuildToStagingDirectory.Automation.cs
+++ b/Engine/Source/Programs/AutomationTool/Scripts/CopyBuildToStagingDirectory.Automation.cs
@@ -5896,7 +5896,7 @@ namespace AutomationScripts
 				List<Tuple<string, UnrealTargetConfiguration, UnrealArchitectures>> Result = new();
 
 				UnrealArchitectureConfig ArchConfig = UnrealArchitectureConfig.ForPlatform(Platform);
-				if (ArchConfig.Mode == UnrealArchitectureMode.OneTargetPerArchitecture && Architectures != null)
+				if (ArchConfig.Mode == UnrealArchitectureMode.OneTargetPerArchitecture && Architectures != null && Architectures.bIsMultiArch)
 				{
 					foreach (UnrealArch Architecture in Architectures.Architectures)
 					{
@@ -5957,7 +5957,7 @@ namespace AutomationScripts
 						UnrealArchitectures Arches = InDedicatedServer ? Params.ServerArchitecture : Params.ClientArchitecture;
 						foreach (var Config in ConfigsToProcess)
 						{
-							TargetAndConfigPairs.AddRange( ExpandDeploymentArchitectures(Target, UnrealTargetConfiguration.Development, Arches, StagePlatform.Type ));
+							TargetAndConfigPairs.AddRange( ExpandDeploymentArchitectures(Target, Config, Arches, StagePlatform.Type ));
 						}
 					}
 				}
