Add Gauntlet argument to keep artifacts with binary overlay (-dev) run

#tests Run locally, Horde ID - 67ed960139205539ce1c251d

diff --git a/Engine/Source/Programs/AutomationTool/Gauntlet/Framework/Gauntlet.Utils.cs b/Engine/Source/Programs/AutomationTool/Gauntlet/Framework/Gauntlet.Utils.cs
--- a/Engine/Source/Programs/AutomationTool/Gauntlet/Framework/Gauntlet.Utils.cs
+++ b/Engine/Source/Programs/AutomationTool/Gauntlet/Framework/Gauntlet.Utils.cs
@@ -44,6 +44,7 @@ namespace Gauntlet
 		public static bool CancelSignalled { get; private set; }
 
 		public static bool IsRunningDev => Params.ParseParam("dev");
+		public static bool IsForcingCopyArtifacts => Params.ParseParam("ForceCopyArtifacts");
 
 		/// <summary>
 		/// Returns the device pool id 
diff --git a/Engine/Source/Programs/AutomationTool/Gauntlet/Unreal/Base/Gauntlet.UnrealSession.cs b/Engine/Source/Programs/AutomationTool/Gauntlet/Unreal/Base/Gauntlet.UnrealSession.cs
--- a/Engine/Source/Programs/AutomationTool/Gauntlet/Unreal/Base/Gauntlet.UnrealSession.cs
+++ b/Engine/Source/Programs/AutomationTool/Gauntlet/Unreal/Base/Gauntlet.UnrealSession.cs
@@ -1062,9 +1062,12 @@ namespace Gauntlet
 			bool IsBuildMachine = CommandUtils.IsBuildMachine;
 			bool SkipArchivingAssets = Globals.IsRunningDev || (IsEditorBuild && (IsBuildMachine == false || InRunningRole.AppInstance.ExitCode == 0));
 			bool bRetainArtifacts = InContext.TestParams.ParseParam("RetainDeviceArtifacts");
+			bool bForceArchivingArtifacts = Globals.IsForcingCopyArtifacts;
 
+			bool bShouldArchiving = (bForceArchivingArtifacts || !SkipArchivingAssets);
+			
 			// Check if we should copy artifacts
-			if (!SkipArchivingAssets)
+			if (bShouldArchiving)
 			{
 				if (SourceDirectory.Exists)
 				{
