[WinArm64] AppLocalDependencies are stored per-architecture.

- also updated prereq readme instructions to include arm64.
- note: cannot stage both x64 and arm64 AppLocalDependencies at the same time yet.

#rnx

diff --git a/Engine/Source/Programs/AutomationTool/Win/WinPlatform.Automation.cs b/Engine/Source/Programs/AutomationTool/Win/WinPlatform.Automation.cs
--- a/Engine/Source/Programs/AutomationTool/Win/WinPlatform.Automation.cs
+++ b/Engine/Source/Programs/AutomationTool/Win/WinPlatform.Automation.cs
@@ -175,7 +175,17 @@ public class Win64Platform : Platform
 
 		if (!string.IsNullOrWhiteSpace(Params.AppLocalDirectory))
 		{
-			StageAppLocalDependencies(Params, SC, "Win64");
+			IEnumerable<string> ArchitectureNames = SC.StageTargets.SelectMany( Target => Target.Receipt.Architectures.Architectures.Select( Arch => Arch.WindowsToolChain) ).Distinct();
+			if (ArchitectureNames.Count() > 1)
+			{
+				Logger.LogWarning("Currently cannot use AppLocalDirectory with multiple staged architectures - only x64 AppLocalDependencies will be staged");
+				ArchitectureNames = ["x64"];
+			}
+
+			foreach (string ArchitectureName in ArchitectureNames)
+			{
+				StageAppLocalDependencies(Params, SC, "Win64", ArchitectureName);
+			}
 		}
 	}
 
@@ -361,7 +371,7 @@ public class Win64Platform : Platform
 		return true;
 	}
 
-	public void StageAppLocalDependencies(ProjectParams Params, DeploymentContext SC, string SourcePlatformDir)
+	public void StageAppLocalDependencies(ProjectParams Params, DeploymentContext SC, string SourcePlatformDir, string SourceArchitecture)
 	{
 		Dictionary<string, string> PathVariables = new Dictionary<string, string>();
 		PathVariables["EngineDir"] = SC.EngineRoot.FullName;
@@ -373,7 +383,15 @@ public class Win64Platform : Platform
 		{
 			string ExpandedAppLocalDir = Utils.ExpandVariables(AppLocalDirectory, PathVariables);
 
-			DirectoryReference BaseAppLocalDependenciesPath = Path.IsPathRooted(ExpandedAppLocalDir) ? new DirectoryReference(CombinePaths(ExpandedAppLocalDir, SourcePlatformDir)) : DirectoryReference.Combine(SC.ProjectRoot, ExpandedAppLocalDir, SourcePlatformDir);
+			DirectoryReference BaseAppLocalDependenciesPath = Path.IsPathRooted(ExpandedAppLocalDir) ? new DirectoryReference(CombinePaths(ExpandedAppLocalDir, SourcePlatformDir, SourceArchitecture)) : DirectoryReference.Combine(SC.ProjectRoot, ExpandedAppLocalDir, SourcePlatformDir, SourceArchitecture);
+			if (!DirectoryReference.Exists(BaseAppLocalDependenciesPath))
+			{
+				BaseAppLocalDependenciesPath = BaseAppLocalDependenciesPath.ParentDirectory!; // no architecture subfolder - must be user-supplied dependency folder in need of upgrade
+				if (DirectoryReference.Exists(BaseAppLocalDependenciesPath))
+				{
+					Logger.LogWarning("{path} should have {arch} subfolder with the redist DLLs inside. Falling back to use {path} directly", BaseAppLocalDependenciesPath, SourceArchitecture, BaseAppLocalDependenciesPath);
+				}
+			}
 			if (DirectoryReference.Exists(BaseAppLocalDependenciesPath))
 			{
 				StageAppLocalDependenciesToDir(SC, BaseAppLocalDependenciesPath, StagedDirectoryReference.Combine("Engine", "Binaries", PlatformType.ToString()));
diff --git a/Engine/Source/Programs/PrereqInstaller/Resources/VCRedist/README.txt b/Engine/Source/Programs/PrereqInstaller/Resources/VCRedist/README.txt
--- a/Engine/Source/Programs/PrereqInstaller/Resources/VCRedist/README.txt
+++ b/Engine/Source/Programs/PrereqInstaller/Resources/VCRedist/README.txt
@@ -4,11 +4,16 @@ Update version numbers to match in:
 Engine\Source\Programs\Windows\BootstrapPackagedGame\Private\BootstrapPackagedGame.cpp
 
 Copy redist files from:
-C:\Program Files\Microsoft Visual Studio\2022\Professional\VC\Redist\MSVC\<version>\x64\Microsoft.VC143.CRT
+C:\Program Files\Microsoft Visual Studio\2022\Professional\VC\Redist\MSVC\<version>\x64\Microsoft.VC143.CRT\
+C:\Program Files\Microsoft Visual Studio\2022\Professional\VC\Redist\MSVC\<version>\arm64\Microsoft.VC143.CRT\
 to:
-Engine\Binaries\ThirdParty\AppLocalDependencies\Win64\Microsoft.VC.CRT
+Engine\Binaries\ThirdParty\AppLocalDependencies\Win64\x64\Microsoft.VC.CRT\
+Engine\Binaries\ThirdParty\AppLocalDependencies\Win64\arm64\Microsoft.VC.CRT\
 
 Copy ucrt files from:
 C:\Program Files (x86)\Windows Kits\10\Redist\<version>\ucrt\DLLs\x64\
+C:\Program Files (x86)\Windows Kits\10\Redist\<version>\ucrt\DLLs\arm64\
 to:
-Engine\Binaries\ThirdParty\AppLocalDependencies\Win64\UCRT\
+Engine\Binaries\ThirdParty\AppLocalDependencies\Win64\UCRT\x64\
+Engine\Binaries\ThirdParty\AppLocalDependencies\Win64\UCRT\arm64\
+
