[Gauntlet] Adds an UnrealAppConfig parameter to ITargetDevice.CleanArtifacts. This isn't required for all platforms, but several platforms do require information about the Build/Configuration to reliably determine which directories need to be cleaned.
Also updates all ITargetDevice implementations to account for this new parameter.

diff --git a/Engine/Source/Programs/AutomationTool/Gauntlet/Framework/Base/Gauntlet.TargetDevice.cs b/Engine/Source/Programs/AutomationTool/Gauntlet/Framework/Base/Gauntlet.TargetDevice.cs
--- a/Engine/Source/Programs/AutomationTool/Gauntlet/Framework/Base/Gauntlet.TargetDevice.cs
+++ b/Engine/Source/Programs/AutomationTool/Gauntlet/Framework/Base/Gauntlet.TargetDevice.cs
@@ -101,7 +101,7 @@ namespace Gauntlet
 		/// <summary>
 		/// Deletes artifacts and other loose files associated with UE processes
 		/// </summary>
-		void CleanArtifacts();
+		void CleanArtifacts(UnrealAppConfig AppConfiguration = null);
 
 		/// <summary>
 		/// Local cache path
diff --git a/Engine/Source/Programs/AutomationTool/Gauntlet/Framework/Base/Gauntlet.TargetDeviceDesktopCommon.cs b/Engine/Source/Programs/AutomationTool/Gauntlet/Framework/Base/Gauntlet.TargetDeviceDesktopCommon.cs
--- a/Engine/Source/Programs/AutomationTool/Gauntlet/Framework/Base/Gauntlet.TargetDeviceDesktopCommon.cs
+++ b/Engine/Source/Programs/AutomationTool/Gauntlet/Framework/Base/Gauntlet.TargetDeviceDesktopCommon.cs
@@ -100,7 +100,7 @@ namespace Gauntlet
 			CleanArtifacts();
 		}
 
-		public void CleanArtifacts()
+		public void CleanArtifacts(UnrealAppConfig AppConfig = null)
 		{
 			if (Directory.Exists(UserDir))
 			{
@@ -119,10 +119,21 @@ namespace Gauntlet
 					}
 				}
 			}
-			if (InstallCache != null)
+
+			if (AppConfig == null)
+			{
+				Log.Warning("{Platform} expects an non-null AppConfig value to reliably determine which artifact path to clear! Skipping clean", Platform);
+				return;
+			}
+
+			if (AppConfig.Build is not StagedBuild Build)
 			{
-				CleanArtifactDirectory(GetInstallArtifactPath());
+				return;
 			}
+
+			string BuildDirectory = GetStagedBuildDirectory(AppConfig, Build);
+			string ArtifactDirectory = Path.Combine(BuildDirectory, AppConfig.ProjectName, "Saved");
+			CleanArtifactDirectory(ArtifactDirectory);
 		}
 
 		public virtual void InstallBuild(UnrealAppConfig AppConfig)
@@ -181,7 +192,6 @@ namespace Gauntlet
 					throw new AutomationException("{0} is an invalid build type for {1}!", Build.GetType().Name, Platform);
 			}
 
-			InstallCache = Install;
 			return Install;
 		}
 
@@ -249,7 +259,18 @@ namespace Gauntlet
 
 		protected abstract IAppInstall CreateEditorInstall(UnrealAppConfig AppConfig, EditorBuild Build);
 
-		protected abstract string GetInstallArtifactPath();
+		protected string GetStagedBuildDirectory(UnrealAppConfig AppConfig, StagedBuild Build)
+		{
+			string BuildDir = Build.BuildPath;
+			if (SystemHelpers.IsNetworkPath(BuildDir))
+			{
+				string SubDir = string.IsNullOrEmpty(AppConfig.Sandbox) ? AppConfig.ProjectName : AppConfig.Sandbox;
+				string InstallDir = Path.Combine(InstallRoot, SubDir, AppConfig.ProcessType.ToString());
+				BuildDir = InstallDir;
+			}
+
+			return BuildDir;
+		}
 
 		private void CleanArtifactDirectory(string ArtifactDirectory)
 		{
diff --git a/Engine/Source/Programs/AutomationTool/Gauntlet/Platform/Android/Gauntlet.TargetDeviceAndroid.cs b/Engine/Source/Programs/AutomationTool/Gauntlet/Platform/Android/Gauntlet.TargetDeviceAndroid.cs
--- a/Engine/Source/Programs/AutomationTool/Gauntlet/Platform/Android/Gauntlet.TargetDeviceAndroid.cs
+++ b/Engine/Source/Programs/AutomationTool/Gauntlet/Platform/Android/Gauntlet.TargetDeviceAndroid.cs
@@ -349,17 +349,27 @@ namespace Gauntlet
 			}
 		}
 
-		public void CleanArtifacts()
+		public void CleanArtifacts(UnrealAppConfig AppConfig = null)
 		{
-			if(!string.IsNullOrEmpty(DeviceExternalStorageSavedPath))
+			if (AppConfig == null)
 			{
-				RunAdbDeviceCommand(string.Format("shell rm -r {0}/*", DeviceExternalStorageSavedPath));
+				Log.Warning("{Platform} expects an non-null AppConfig value to reliably determine which artifact path to clear! Skipping clean", Platform);
+				return;
 			}
 
-			if (!string.IsNullOrEmpty(DeviceExternalFilesSavedPath))
+			if (AppConfig.Build is not AndroidBuild Build)
 			{
-				RunAdbDeviceCommand(string.Format("shell rm -r {0}/*", DeviceExternalFilesSavedPath));
+				return;
 			}
+
+			string ExternalStoragePath = StoragePath + "/UnrealGame/" + AppConfig.ProjectName;
+			string ExternalFilesPath = StoragePath + "/Android/data/" + Build.AndroidPackageName + "/files/UnrealGame/" + AppConfig.ProjectName;
+
+			string ExternalAppStorage = string.Format("{0}/{1}/Saved", ExternalStoragePath, AppConfig.ProjectName);
+			string ExternalAppFiles = string.Format("{0}/{1}/Saved", ExternalFilesPath, AppConfig.ProjectName);
+
+			RunAdbDeviceCommand(string.Format("shell rm -r {0}/*", ExternalAppStorage));
+			RunAdbDeviceCommand(string.Format("shell rm -r {0}/*", ExternalAppFiles));
 		}
 
 		public void InstallBuild(UnrealAppConfig AppConfig)
diff --git a/Engine/Source/Programs/AutomationTool/Gauntlet/Platform/IOS/Gauntlet.TargetDeviceIOS.cs b/Engine/Source/Programs/AutomationTool/Gauntlet/Platform/IOS/Gauntlet.TargetDeviceIOS.cs
--- a/Engine/Source/Programs/AutomationTool/Gauntlet/Platform/IOS/Gauntlet.TargetDeviceIOS.cs
+++ b/Engine/Source/Programs/AutomationTool/Gauntlet/Platform/IOS/Gauntlet.TargetDeviceIOS.cs
@@ -397,18 +397,23 @@ namespace Gauntlet
 			}
 		}
 
-		public void CleanArtifacts()
-		{
-			if(InstallCache == null)
-			{
-				throw new AutomationException("Could not deduce Bundle ID and project name due to invalid installation cache. " +
-					"Either specify the Bundle ID and Project Name using this function's overload OR ensure CreateAppInstall has been called first.");
-			}
-
-			CleanArtifacts(InstallCache.PackageName, InstallCache.ProjectName);
-		}
+        public void CleanArtifacts(UnrealAppConfig AppConfig = null)
+        {
+            if (AppConfig == null)
+            {
+                Log.Warning("{Platform} expects an non-null AppConfig value to reliably determine which artifact path to clear! Skipping clean", Platform);
+                return;
+            }
+
+            if (AppConfig.Build is not IOSBuild Build)
+            {
+                return;
+            }
+
+            CleanArtifacts(Build.PackageName, AppConfig.ProjectName);
+        }
 
-		public void CleanArtifacts(string BundleID, string ProjectName)
+        public void CleanArtifacts(string BundleID, string ProjectName)
 		{
 			string EngineSavedDirectory = "/Documents/Engine/Saved";
 			string ProjectSavedDirectory = string.Format("/Documents/{0}/Saved", ProjectName);
diff --git a/Engine/Source/Programs/AutomationTool/Gauntlet/Platform/Linux/Gauntlet.TargetDeviceLinux.cs b/Engine/Source/Programs/AutomationTool/Gauntlet/Platform/Linux/Gauntlet.TargetDeviceLinux.cs
--- a/Engine/Source/Programs/AutomationTool/Gauntlet/Platform/Linux/Gauntlet.TargetDeviceLinux.cs
+++ b/Engine/Source/Programs/AutomationTool/Gauntlet/Platform/Linux/Gauntlet.TargetDeviceLinux.cs
@@ -111,13 +111,7 @@ namespace Gauntlet
 
 		protected override IAppInstall CreateStagedInstall(UnrealAppConfig AppConfig, StagedBuild Build)
 		{
-			string BuildDir = Build.BuildPath;
-			if (SystemHelpers.IsNetworkPath(BuildDir))
-			{
-				string SubDir = string.IsNullOrEmpty(AppConfig.Sandbox) ? AppConfig.ProjectName : AppConfig.Sandbox;
-				string InstallDir = Path.Combine(InstallRoot, SubDir, AppConfig.ProcessType.ToString());
-				BuildDir = InstallDir;
-			}
+			string BuildDir = GetStagedBuildDirectory(AppConfig, Build);
 
 			PopulateDirectoryMappings(Path.Combine(BuildDir, AppConfig.ProjectName));
 
@@ -146,11 +140,6 @@ namespace Gauntlet
 			return LinuxApp;
 		}
 
-		protected override string GetInstallArtifactPath()
-		{
-			return (InstallCache as DesktopCommonAppInstall<TargetDeviceLinux>).ArtifactPath;
-		}
-
 		#region Legacy Implementations
 
 		protected override IAppInstall InstallNativeStagedBuild(UnrealAppConfig AppConfig, NativeStagedBuild InBuild)
diff --git a/Engine/Source/Programs/AutomationTool/Gauntlet/Platform/Mac/Gauntlet.TargetDeviceMac.cs b/Engine/Source/Programs/AutomationTool/Gauntlet/Platform/Mac/Gauntlet.TargetDeviceMac.cs
--- a/Engine/Source/Programs/AutomationTool/Gauntlet/Platform/Mac/Gauntlet.TargetDeviceMac.cs
+++ b/Engine/Source/Programs/AutomationTool/Gauntlet/Platform/Mac/Gauntlet.TargetDeviceMac.cs
@@ -92,7 +92,6 @@ namespace Gauntlet
 					throw new AutomationException("{0} is an invalid build type for {1}!", Build.GetType().Name, Platform);
 			}
 
-			InstallCache = Install;
 			return Install;
 		}
 
@@ -262,11 +261,6 @@ namespace Gauntlet
 			return BuildDir;
 		}
 
-		protected override string GetInstallArtifactPath()
-		{
-			return (InstallCache as DesktopCommonAppInstall<TargetDeviceMac>).ArtifactPath;
-		}
-
 		private string GetTargetBuildDirectory(string BasePath, string Sandbox, string Project, string ProcessType)
 		{
 			string TargetDir = BasePath;
diff --git a/Engine/Source/Programs/AutomationTool/Gauntlet/Platform/Null/Gauntlet.TargetDeviceNull.cs b/Engine/Source/Programs/AutomationTool/Gauntlet/Platform/Null/Gauntlet.TargetDeviceNull.cs
--- a/Engine/Source/Programs/AutomationTool/Gauntlet/Platform/Null/Gauntlet.TargetDeviceNull.cs
+++ b/Engine/Source/Programs/AutomationTool/Gauntlet/Platform/Null/Gauntlet.TargetDeviceNull.cs
@@ -167,7 +167,7 @@ namespace Gauntlet
 
 		}
 
-		public void CleanArtifacts()
+		public void CleanArtifacts(UnrealAppConfig AppConfig = null)
 		{
 
 		}
diff --git a/Engine/Source/Programs/AutomationTool/Gauntlet/Platform/Windows/Gauntlet.TargetDeviceWindows.cs b/Engine/Source/Programs/AutomationTool/Gauntlet/Platform/Windows/Gauntlet.TargetDeviceWindows.cs
--- a/Engine/Source/Programs/AutomationTool/Gauntlet/Platform/Windows/Gauntlet.TargetDeviceWindows.cs
+++ b/Engine/Source/Programs/AutomationTool/Gauntlet/Platform/Windows/Gauntlet.TargetDeviceWindows.cs
@@ -68,7 +68,6 @@ namespace Gauntlet
 					throw new AutomationException("{0} is an invalid build type for {1}!", Build.GetType().Name, Platform);
 			}
 
-			InstallCache = Install;
 			return Install;
 		}
 
@@ -137,13 +136,7 @@ namespace Gauntlet
 
 		protected override IAppInstall CreateStagedInstall(UnrealAppConfig AppConfig, StagedBuild Build)
 		{
-			string BuildDir = Build.BuildPath;
-			if (SystemHelpers.IsNetworkPath(BuildDir))
-			{
-				string SubDir = string.IsNullOrEmpty(AppConfig.Sandbox) ? AppConfig.ProjectName : AppConfig.Sandbox;
-				string InstallDir = Path.Combine(InstallRoot, SubDir, AppConfig.ProcessType.ToString());
-				BuildDir = InstallDir;
-			}
+			string BuildDir = GetStagedBuildDirectory(AppConfig, Build);
 
 			PopulateDirectoryMappings(Path.Combine(BuildDir, AppConfig.ProjectName));
 
@@ -179,11 +172,6 @@ namespace Gauntlet
 			return WinApp;
 		}
 
-		protected override string GetInstallArtifactPath()
-		{
-			return (InstallCache as DesktopCommonAppInstall<TargetDeviceWindows>).ArtifactPath;
-		}
-
 		#region Legacy Implementations
 		public override IAppInstall InstallApplication(UnrealAppConfig AppConfig)
 		{
diff --git a/Engine/Source/Programs/AutomationTool/Gauntlet/Unreal/Base/Gauntlet.UnrealSession.cs b/Engine/Source/Programs/AutomationTool/Gauntlet/Unreal/Base/Gauntlet.UnrealSession.cs
--- a/Engine/Source/Programs/AutomationTool/Gauntlet/Unreal/Base/Gauntlet.UnrealSession.cs
+++ b/Engine/Source/Programs/AutomationTool/Gauntlet/Unreal/Base/Gauntlet.UnrealSession.cs
@@ -1389,7 +1389,10 @@ namespace Gauntlet
 						try
 						{
 							// Make sure the artifact paths are clean before installation
-							Device.CleanArtifacts();
+							if(!Role.SkipCleanDeviceArtifacts)
+							{
+                                Device.CleanArtifacts(AppConfig);
+                            }
 
 							Install = Device.InstallApplication(AppConfig);
 							IDeviceUsageReporter.RecordEnd(Device.Name, Device.Platform, IDeviceUsageReporter.EventType.Install, IDeviceUsageReporter.EventState.Success);
@@ -1643,11 +1646,6 @@ namespace Gauntlet
 
 				try
 				{
-					if (!Role.SkipCleanDeviceArtifacts)
-					{
-						Device.CleanArtifacts();
-					}
-
 					if (RoleNeedsInstall(Role))
 					{
 						RolesToConfigs.Remove(Role);
@@ -1660,13 +1658,17 @@ namespace Gauntlet
 						AppConfig.CommandLineParams.AddUnique("DeviceTag", Device.Name);
 						RolesToConfigs.Add(Role, AppConfig);
 
-						// Install the build
 						if (AppConfig.FullClean)
 						{
 							Log.Info("Fully cleaning device before install...");
 							Device.FullClean();
 						}
 
+						if (!Role.SkipCleanDeviceArtifacts)
+						{
+							Device.CleanArtifacts(AppConfig);
+						}
+
 						if (AppConfig.SkipInstall)
 						{
 							Log.Info("Skipping install due to SkipInstall");
@@ -1697,7 +1699,11 @@ namespace Gauntlet
 					else
 					{
 						AppConfig = RolesToConfigs[Role];
-					}
+                        if (!Role.SkipCleanDeviceArtifacts)
+                        {
+                            Device.CleanArtifacts(AppConfig);
+                        }
+                    }
 
 					if (RolesToConfigs[Role].Build.SupportsAdditionalFileCopy)
 					{
