[Gauntlet] Reverts change that required explicit platform information to be included in executable name. Instead, more granularly search platform folders for exact platform names for individual build sources

diff --git a/Engine/Source/Programs/AutomationTool/Gauntlet/Unreal/BuildSource/Gauntlet.StagedBuildSource.cs b/Engine/Source/Programs/AutomationTool/Gauntlet/Unreal/BuildSource/Gauntlet.StagedBuildSource.cs
--- a/Engine/Source/Programs/AutomationTool/Gauntlet/Unreal/BuildSource/Gauntlet.StagedBuildSource.cs
+++ b/Engine/Source/Programs/AutomationTool/Gauntlet/Unreal/BuildSource/Gauntlet.StagedBuildSource.cs
@@ -296,12 +296,6 @@ namespace Gauntlet
 					UnrealTargetRole Role = ConfigInfo.RoleType;
 					string Flavor = ConfigInfo.Flavor;
 
-					if (ConfigInfo.Platform.HasValue && InPlatform != ConfigInfo.Platform.Value)
-					{
-						Log.Verbose("{PackageName} platform {Platform} does not match, skipping.", App.Name, ConfigInfo.Platform.Value);
-						continue;
-					}
-
 					// if we have a verificatio BuildSource object, verify the build is usable before we even create the build
 					if (VerificationBuildSource == null || VerificationBuildSource.ShouldMakeBuildAvailable(App.FullName))
 					{
@@ -366,30 +360,27 @@ namespace Gauntlet
 			get { return AutomationTool.Platform.GetExeExtension(Platform); }
 		}
 
+		protected virtual string[] ValidPlatformFolders => [PlatformFolderPrefix];
+
+
 		public virtual List<IBuild> GetBuildsAtPath(string InProjectName, string InPath, int MaxRecursion = 3)
 		{
 			List<DirectoryInfo> AllDirs = new List<DirectoryInfo>();
-
 			List<IBuild> Builds = new List<IBuild>();
-
-			// c:\path\to\build
 			DirectoryInfo PathDI = new DirectoryInfo(InPath);
 
 			if (PathDI.Exists)
 			{
-				// If the path is directly to a platform folder, add us
-				if (PathDI.Name.IndexOf(PlatformFolderPrefix, StringComparison.OrdinalIgnoreCase) >= 0)
+				// Find all directories that begin with the platform name
+				foreach (string PlatformFolder in ValidPlatformFolders)
 				{
-					AllDirs.Add(PathDI);
+					if (PathDI.Name.IndexOf(PlatformFolder, StringComparison.OrdinalIgnoreCase) >= 0)
+					{
+						AllDirs.Add(PathDI);
+					}
+					AllDirs.AddRange(PathDI.GetDirectories(PlatformFolder + "*", SearchOption.TopDirectoryOnly));
 				}
 
-				// Assume it's a folder of all build types so find all sub directories that begin with the foldername for this platform
-				IEnumerable<DirectoryInfo> MatchingDirs = PathDI.GetDirectories("*", SearchOption.TopDirectoryOnly);
-
-				MatchingDirs = MatchingDirs.Where(D => D.Name.StartsWith(PlatformFolderPrefix, StringComparison.OrdinalIgnoreCase)).ToArray();
-
-				AllDirs.AddRange(MatchingDirs);
-
 				List<DirectoryInfo> DirsToRecurse = AllDirs;
 
 				// now get subdirs
