Cloud Artifacts - Made verify a parameter which is disabled by default - we have not caught any issues with this for the last few weeks so its not worth the added walltime. Also made verbose a parameter but kept it enabled by default. Most of the verbose logging comes from the verification and not the upload

diff --git a/Engine/Source/Programs/AutomationTool/BuildGraph/Tasks/CreateCloudArtifactTask.cs b/Engine/Source/Programs/AutomationTool/BuildGraph/Tasks/CreateCloudArtifactTask.cs
--- a/Engine/Source/Programs/AutomationTool/BuildGraph/Tasks/CreateCloudArtifactTask.cs
+++ b/Engine/Source/Programs/AutomationTool/BuildGraph/Tasks/CreateCloudArtifactTask.cs
@@ -145,6 +145,18 @@ namespace AutomationTool.Tasks
 		/// </summary>
 		[TaskParameter(Optional = true)]
 		public string HttpVersion { get; set; } = "None";
+
+		/// <summary>
+		/// Verify the contents after upload, essentially downloads the entire build to make sure all data was uploaded correctly. Takes a bit of time and is generally not needed but can be useful to detect issues in the upload.
+		/// </summary>
+		[TaskParameter(Optional = true)]
+		public bool Verify { get; set; } = false;
+
+		/// <summary>
+		/// Verbose logging output, can be useful to find out more information when an issue has occured.
+		/// </summary>
+		[TaskParameter(Optional = true)]
+		public bool Verbose { get; set; } = true;
 	}
 
 	/// <summary>
@@ -458,7 +470,7 @@ namespace AutomationTool.Tasks
 				Branch = streamId,
 				BaselineBranch = baselineBranch,
 				Platform = platform
-			}, allowMultipart: _parameters.AllowMultipart, assumeHttp2: String.Equals(httpVersion, "http2-only", StringComparison.OrdinalIgnoreCase), forceClean: _parameters.Force);
+			}, allowMultipart: _parameters.AllowMultipart, assumeHttp2: String.Equals(httpVersion, "http2-only", StringComparison.OrdinalIgnoreCase), forceClean: _parameters.Force, verify: _parameters.Verify, verbose: _parameters.Verbose);
 
 			Uri buildUri = new Uri($"{descriptorHost}/api/v2/builds/{ns}/{bucketId}/{buildId}");
 
@@ -552,7 +564,7 @@ namespace AutomationTool.Tasks
 		}
 		public static void UploadBuild(string host, DirectoryReference baseDir, string namespaceId, string bucketId, CbObjectId buildId, string partName,
 			CbObjectId partId, Dictionary<string, object> metadata, string accessToken, BuildInfo info, bool allowMultipart = false, bool assumeHttp2 = false,
-			bool forceClean = false)
+			bool forceClean = false, bool verify = false, bool verbose = false)
 		{
 			string metadataJsonPath = Path.Combine(Path.GetTempPath(), Path.GetTempFileName());
 			try
@@ -591,8 +603,10 @@ namespace AutomationTool.Tasks
 
 				string metadataString = File.ReadAllText(metadataJsonPath);
 				string http2Options = assumeHttp2 ? "--assume-http2" : "";
+				string verifyOption = verify ? "--verify" : "";
+				string verboseOption = verbose ? "--verbose" : "";
 				// pass the access token via environment variable to avoid showing it on the cli
-				string cmdline = $"builds upload --url {host} --namespace {namespaceId} --bucket {bucketId} --build-id {buildId} --build-part-name \"{partName}\" --build-part-id {partId} --metadata-path \"{metadataJsonPath}\" --access-token-env UE-CloudDataCacheAccessToken --create-build --plain-progress --allow-multipart={allowMultipart} {http2Options} --clean={forceClean} --verify --verbose \"{baseDir}\"";
+				string cmdline = $"builds upload --url {host} --namespace {namespaceId} --bucket {bucketId} --build-id {buildId} --build-part-name \"{partName}\" --build-part-id {partId} --metadata-path \"{metadataJsonPath}\" --access-token-env UE-CloudDataCacheAccessToken --create-build --plain-progress --allow-multipart={allowMultipart} {http2Options} --clean={forceClean} {verifyOption} {verboseOption} \"{baseDir}\"";
 				Log.Logger.LogInformation("Using metadata '{Metadata}'", metadataString);
 				CommandUtils.RunAndLog(CommandUtils.CmdEnv, zenExe.FullName,cmdline , Options: CommandUtils.ERunOptions.Default, EnvVars: new Dictionary<string, string> {{"UE-CloudDataCacheAccessToken", accessToken}});
 			}
