[Gauntlet] Passes through SessionRole's MaxDuration to UnrealAppConfig created within the context of an UnrealSession

diff --git a/Engine/Source/Programs/AutomationTool/Gauntlet/Unreal/Base/Gauntlet.UnrealAppConfig.cs b/Engine/Source/Programs/AutomationTool/Gauntlet/Unreal/Base/Gauntlet.UnrealAppConfig.cs
--- a/Engine/Source/Programs/AutomationTool/Gauntlet/Unreal/Base/Gauntlet.UnrealAppConfig.cs
+++ b/Engine/Source/Programs/AutomationTool/Gauntlet/Unreal/Base/Gauntlet.UnrealAppConfig.cs
@@ -118,6 +118,11 @@ namespace Gauntlet
 		/// Use additional debug memory on target device if available. 
 		/// </summary>
 		public bool UsePlatformExtraDebugMemory { get; set; }
+		
+		/// <summary>
+		/// Defines the expected duration of the app process. May not be set depending on the context of this object's creation
+		/// </summary>
+		public int MaxDuration { get; set; }
 
 		/// <summary>
 		/// Constructor that sets some required values to defaults
diff --git a/Engine/Source/Programs/AutomationTool/Gauntlet/Unreal/BuildSource/Gauntlet.UnrealBuildSource.cs b/Engine/Source/Programs/AutomationTool/Gauntlet/Unreal/BuildSource/Gauntlet.UnrealBuildSource.cs
--- a/Engine/Source/Programs/AutomationTool/Gauntlet/Unreal/BuildSource/Gauntlet.UnrealBuildSource.cs
+++ b/Engine/Source/Programs/AutomationTool/Gauntlet/Unreal/BuildSource/Gauntlet.UnrealBuildSource.cs
@@ -10,6 +10,7 @@ using System.Linq;
 using EpicGames.Core;
 using System.Text.Json.Serialization;
 using System.Text.Json;
+using UnrealGame;
 
 namespace Gauntlet
 {
@@ -614,6 +615,8 @@ namespace Gauntlet
 				Issues.ForEach(S => Log.Error(S));
 				return null;
 			}
+			
+			UnrealTestConfiguration TestConfig = Role.Options as UnrealTestConfiguration;
 
 			UnrealAppConfig Config = new UnrealAppConfig();
 
@@ -625,6 +628,7 @@ namespace Gauntlet
 			Config.Configuration = Role.Configuration;
 			Config.CommandLineParams = new GauntletCommandLine(Role.CommandLineParams);
             Config.FilesToCopy = new List<UnrealFileToCopy>();
+            Config.MaxDuration = (int)(TestConfig?.MaxDuration ?? 0);
 
 			// new system of retrieving and encapsulating the info needed to install/launch. Android & Mac
 			Config.Build = GetMatchingBuilds(Role.RoleType, Role.Platform, Role.Configuration, Role.RequiredBuildFlags, Role.RequiredFlavor).OrderBy(B => B.PreferenceOrder).FirstOrDefault();
