Localization:
- Fixed a bug with Engine plugin localization config files failing to generate.
- Removed the Audit_InCook collection from being included in default generated localization config files

diff --git a/Engine/Source/Programs/AutomationTool/Localization/LocalizationConfigFileParams.cs b/Engine/Source/Programs/AutomationTool/Localization/LocalizationConfigFileParams.cs
--- a/Engine/Source/Programs/AutomationTool/Localization/LocalizationConfigFileParams.cs
+++ b/Engine/Source/Programs/AutomationTool/Localization/LocalizationConfigFileParams.cs
@@ -89,8 +89,6 @@ namespace EpicGames.Localization
 			gatherTextFromAssets.ExcludePathFilters.Add(Path.Combine(localizationTargetRootDirectory, "Content", "*Test", "*"));
 			
 			gatherTextFromAssets.PackageFileNameFilters= new List<string>() { "*.umap", "*.uasset" };
-			//@TODOLocalization: Some licensees might not have a need for an Audit_InCook file...should remove this from defaults after config file generation with templates is implemented
-			gatherTextFromAssets.CollectionFilters = new List<string>() { "Audit_InCook" };
 			gatherTextFromAssets.WorldCollectionFilters = new List<string>() { "Audit_InCook" };
 
 			gatherTextFromAssets.bShouldExcludeDerivedClasses = false;
diff --git a/Engine/Source/Programs/AutomationTool/Scripts/Localisation.Automation.cs b/Engine/Source/Programs/AutomationTool/Scripts/Localisation.Automation.cs
--- a/Engine/Source/Programs/AutomationTool/Scripts/Localisation.Automation.cs
+++ b/Engine/Source/Programs/AutomationTool/Scripts/Localisation.Automation.cs
@@ -992,8 +992,8 @@ class Localize : BuildCommand
 					{
 						PluginLocalizationConfigFiles = Generator.GenerateDefaultSettingsConfigFiles(GeneratorParams);
 					}
-					
-					if (PluginLocalizationConfigFiles.Count != TemplateLocalizationConfigFiles.Count)
+					// We only count as an error when generation of template localization config files fail 
+					if ((TemplateLocalizationConfigFiles.Count > 0) && (PluginLocalizationConfigFiles.Count != TemplateLocalizationConfigFiles.Count))
 					{
 						Logger.LogError($"1 or more localization config files failed to be auto generated for localization target '{GeneratorParams.LocalizationTargetName}'");
 						continue;
