Multi-architecture support for BootstrapPackagedGame

- bootstrapper will select and run the most appropriate executable for the host architecture.
- note: bootstrapper is only x64 and runs under emulation on arm64 devices.

#rnx

diff --git a/Engine/Source/Programs/AutomationTool/Win/WinPlatform.Automation.cs b/Engine/Source/Programs/AutomationTool/Win/WinPlatform.Automation.cs
--- a/Engine/Source/Programs/AutomationTool/Win/WinPlatform.Automation.cs
+++ b/Engine/Source/Programs/AutomationTool/Win/WinPlatform.Automation.cs
@@ -120,14 +120,38 @@ public class Win64Platform : Platform
 		// Stage the bootstrap executable
 		if (!Params.NoBootstrapExe)
 		{
-			foreach(StageTarget Target in SC.StageTargets)
+			// one bootstrap per configuration, per target
+			foreach (UnrealTargetConfiguration Configuration in SC.StageTargetConfigurations)
 			{
-				BuildProduct Executable = Target.Receipt.BuildProducts.FirstOrDefault(x => x.Type == BuildProductType.Executable);
-				if(Executable != null)
+				foreach (string TargetName in SC.StageTargets.Select( T => T.Receipt.TargetName ).Distinct())
 				{
-					// only create bootstraps for executables
-					List<StagedFileReference> StagedFiles = SC.FilesToStage.NonUFSFiles.Where(x => x.Value == Executable.Path).Select(x => x.Key).ToList();
-					if (StagedFiles.Count > 0 && Executable.Path.HasExtension(".exe"))
+					Dictionary<UnrealArch,StagedFileReference> RemappedStagedFiles = [];
+					BuildProduct PrimaryExecutable = null;
+					StageTarget? PrimaryTarget = null;
+
+					// collected the staged executable files for each staged architecture
+					foreach (StageTarget Target in SC.StageTargets.Where( T => T.Receipt.Configuration == Configuration && T.Receipt.TargetName == TargetName))
+					{
+						BuildProduct Executable = Target.Receipt.BuildProducts.FirstOrDefault(x => x.Type == BuildProductType.Executable);
+						if (Executable != null)
+						{
+							// only create bootstraps for executables
+							List<StagedFileReference> StagedFiles = SC.FilesToStage.NonUFSFiles.Where(x => x.Value == Executable.Path).Select(x => x.Key).ToList();
+							if (StagedFiles.Count > 0 && Executable.Path.HasExtension(".exe"))
+							{
+								UnrealArch Architecture = Target.Receipt.Architectures.SingleArchitecture;
+								RemappedStagedFiles[Architecture] = DeploymentContext.ApplyDirectoryRemap(SC, StagedFiles.First());
+
+								if (PrimaryExecutable == null || Architecture.bIsX64) // prefer x64 for the BootstrapExeName because it is undecorated
+								{
+									PrimaryExecutable = Executable;
+									PrimaryTarget = Target;
+								}
+							}
+						}
+					}
+					
+					if (RemappedStagedFiles.Any())
 					{
 						string BootstrapArguments = "";
 						if (!ShouldStageCommandLine(Params, SC))
@@ -145,21 +169,24 @@ public class Win64Platform : Platform
 						string BootstrapExeName;
 						if(SC.StageTargetConfigurations.Count > 1)
 						{
-							BootstrapExeName = Executable.Path.GetFileName();
+							BootstrapExeName = PrimaryExecutable.Path.GetFileName();
 						}
 						else if(Params.IsCodeBasedProject)
 						{
-							BootstrapExeName = Target.Receipt.TargetName + ".exe";
+							BootstrapExeName = PrimaryTarget.Value.Receipt.TargetName + ".exe";
 						}
 						else
 						{
 							BootstrapExeName = SC.ShortProjectName + ".exe";
 						}
 
-						foreach (StagedFileReference StagePath in StagedFiles)
+						if (RemappedStagedFiles.Count == 1)
 						{
-							StagedFileReference RemappedStagePath = DeploymentContext.ApplyDirectoryRemap(SC, StagePath);
-							StageBootstrapExecutable(SC, BootstrapExeName, Executable.Path, RemappedStagePath, BootstrapArguments);
+							StageBootstrapExecutable(SC, BootstrapExeName, PrimaryExecutable.Path, RemappedStagedFiles.Values.First(), BootstrapArguments);
+						}
+						else
+						{
+							StageBootstrapExecutable(SC, BootstrapExeName, PrimaryExecutable.Path, null, BootstrapArguments, RemappedStagedFiles);
 						}
 					}
 				}
@@ -199,7 +226,7 @@ public class Win64Platform : Platform
 		CommandUtils.CopyFile(SourceFile.FullName, LocalFile);
 	}
 
-	void StageBootstrapExecutable(DeploymentContext SC, string ExeName, FileReference TargetFile, StagedFileReference StagedRelativeTargetPath, string StagedArguments)
+	void StageBootstrapExecutable(DeploymentContext SC, string ExeName, FileReference TargetFile, StagedFileReference StagedRelativeTargetPath, string StagedArguments, Dictionary<UnrealArch,StagedFileReference> PerArchitectureStagedRelativeTargetPath = null)
 	{
 		UnrealTargetPlatform BootstrapPlatform = (BootstrapExePlatform ?? SC.StageTargetPlatform.PlatformType);
 		FileReference InputFile = FileReference.Combine(SC.LocalRoot, "Engine", "Binaries", BootstrapPlatform.ToString(), String.Format("BootstrapPackagedGame-{0}-Shipping.exe", BootstrapPlatform));
@@ -234,9 +261,38 @@ public class Win64Platform : Platform
 					const int IconResourceId = 101;
 					if(GroupIcon != null) Update.SetIcons(IconResourceId, GroupIcon);
 
-					const int ExecFileResourceId = 201;
-					Update.SetData(ExecFileResourceId, ResourceType.RawData, Encoding.Unicode.GetBytes(StagedRelativeTargetPath.ToString().Replace('/', '\\') + "\0"));
+					if (PerArchitectureStagedRelativeTargetPath == null)
+					{
+						const int ExecFileResourceId = 201;
+						Update.SetData(ExecFileResourceId, ResourceType.RawData, Encoding.Unicode.GetBytes(StagedRelativeTargetPath.ToString().Replace('/', '\\') + "\0"));
+					}
+					else
+					{
+						foreach (KeyValuePair<UnrealArch,StagedFileReference> ExtraTargetPath in PerArchitectureStagedRelativeTargetPath)
+						{
+							if (ExtraTargetPath.Key == UnrealArch.X64)
+							{
+								const int ExecFileResourceId_X64 = 201;
+								Update.SetData(ExecFileResourceId_X64, ResourceType.RawData, Encoding.Unicode.GetBytes(ExtraTargetPath.Value.ToString().Replace('/', '\\') + "\0"));
+							}
+							else if (ExtraTargetPath.Key == UnrealArch.Arm64)
+							{
+								const int ExecFileResourceId_ARM64 = 203;
+								Update.SetData(ExecFileResourceId_ARM64, ResourceType.RawData, Encoding.Unicode.GetBytes(ExtraTargetPath.Value.ToString().Replace('/', '\\') + "\0"));
 
+							}
+							else if (ExtraTargetPath.Key == UnrealArch.Arm64ec)
+							{
+								const int ExecFileResourceId_ARM64EC = 204;
+								Update.SetData(ExecFileResourceId_ARM64EC, ResourceType.RawData, Encoding.Unicode.GetBytes(ExtraTargetPath.Value.ToString().Replace('/', '\\') + "\0"));
+							}
+							else
+							{
+								throw new AutomationException("Unexpected architecture {arch} for bootstrap executable", ExtraTargetPath.Key);
+							}
+						}
+					}
+					
 					const int ExecArgsResourceId = 202;
 					Update.SetData(ExecArgsResourceId, ResourceType.RawData, Encoding.Unicode.GetBytes(StagedArguments + "\0"));
 				}
diff --git a/Engine/Source/Programs/Windows/BootstrapPackagedGame/Private/BootstrapPackagedGame.cpp b/Engine/Source/Programs/Windows/BootstrapPackagedGame/Private/BootstrapPackagedGame.cpp
--- a/Engine/Source/Programs/Windows/BootstrapPackagedGame/Private/BootstrapPackagedGame.cpp
+++ b/Engine/Source/Programs/Windows/BootstrapPackagedGame/Private/BootstrapPackagedGame.cpp
@@ -4,6 +4,8 @@
 
 #define IDI_EXEC_FILE 201
 #define IDI_EXEC_ARGS 202
+#define IDI_EXEC_FILE_ARM64 203
+#define IDI_EXEC_FILE_ARM64EC 204
 
 struct VersionInfo
 {
@@ -367,6 +369,41 @@ int SpawnTarget(WCHAR* CmdLine)
 	return (int)ExitCode;
 }
 
+WCHAR* FindBestExecFile(HINSTANCE hInstance, const WCHAR* BaseDirectory)
+{
+	bool bIsARM64HostPlatform = IsARM64HostPlatform();
+	if (bIsARM64HostPlatform)
+	{
+		WCHAR ExePath[MAX_PATH];
+
+		WCHAR* ExecFileARM64 = ReadResourceString(hInstance, MAKEINTRESOURCE(IDI_EXEC_FILE_ARM64));
+		if (ExecFileARM64 != nullptr)
+		{
+			wsprintf(ExePath, L"%s\\%s", BaseDirectory, ExecFileARM64);
+			if(GetFileAttributes(ExecFileARM64) != INVALID_FILE_ATTRIBUTES)
+			{
+				return ExecFileARM64;
+			}
+
+			delete[] ExecFileARM64;
+		}
+
+		WCHAR* ExecFileARM64EC = ReadResourceString(hInstance, MAKEINTRESOURCE(IDI_EXEC_FILE_ARM64EC));
+		if (ExecFileARM64EC != nullptr)
+		{
+			wsprintf(ExePath, L"%s\\%s", BaseDirectory, ExecFileARM64EC);
+			if(GetFileAttributes(ExecFileARM64EC) != INVALID_FILE_ATTRIBUTES)
+			{
+				return ExecFileARM64EC;
+			}
+
+			delete[] ExecFileARM64EC;
+		}
+	}
+
+	return ReadResourceString(hInstance, MAKEINTRESOURCE(IDI_EXEC_FILE));
+}
+
 int WINAPI wWinMain(_In_ HINSTANCE hInstance, _In_opt_ HINSTANCE hPrevInstance, _In_ LPWSTR CmdLine, _In_ int ShowCmd)
 {
 	(void)hPrevInstance;
@@ -382,7 +419,7 @@ int WINAPI wWinMain(_In_ HINSTANCE hInstance, _In_opt_ HINSTANCE hPrevInstance,
 	PathRemoveFileSpec(BaseDirectory);
 
 	// Get the executable to run
-	WCHAR* ExecFile = ReadResourceString(hInstance, MAKEINTRESOURCE(IDI_EXEC_FILE));
+	WCHAR* ExecFile = FindBestExecFile(hInstance, BaseDirectory);
 	if(ExecFile == NULL)
 	{
 		MessageBoxW(NULL, L"This program is used for packaged games and is not meant to be run directly.", NULL, MB_OK);
