Cloud Artifacts - Simplified engine ini reading to avoid triggering uat compilation of build rules when we just want the inis parsed thus resolving issues with plugin compilation that we do not need have happen.

diff --git a/Engine/Source/Programs/AutomationTool/BuildGraph/Tasks/CreateCloudArtifactTask.cs b/Engine/Source/Programs/AutomationTool/BuildGraph/Tasks/CreateCloudArtifactTask.cs
--- a/Engine/Source/Programs/AutomationTool/BuildGraph/Tasks/CreateCloudArtifactTask.cs
+++ b/Engine/Source/Programs/AutomationTool/BuildGraph/Tasks/CreateCloudArtifactTask.cs
@@ -215,8 +215,8 @@ namespace AutomationTool.Tasks
 			CloudConfiguration? cloudConfig = null;
 			if (projectPath != null)
 			{
-				ProjectProperties properties = ProjectUtils.GetProjectProperties(projectPath);
-				ConfigHierarchy config = properties.EngineConfigs[HostPlatform.Current.HostEditorPlatform];
+				(Dictionary<UnrealTargetPlatform, ConfigHierarchy> engineConfigs, Dictionary<UnrealTargetPlatform, ConfigHierarchy> _) = GetIniConfigs(projectPath);
+				ConfigHierarchy config = engineConfigs[HostPlatform.Current.HostEditorPlatform];
 				foundConfig = config.TryGetValueGeneric("StorageServers", "Cloud", out CloudConfiguration foundCloudConfig);
 				cloudConfig = foundCloudConfig;
 			}
@@ -484,6 +484,24 @@ namespace AutomationTool.Tasks
 			}
 		}
 
+		internal static (Dictionary<UnrealTargetPlatform, ConfigHierarchy>, Dictionary<UnrealTargetPlatform, ConfigHierarchy>) GetIniConfigs(FileReference projectPath)
+		{
+			Logger.LogDebug("Loading ini files for {RawProjectPath}", projectPath);
+
+			Dictionary<UnrealTargetPlatform, ConfigHierarchy> engineConfigs = new Dictionary<UnrealTargetPlatform, ConfigHierarchy>();
+			Dictionary<UnrealTargetPlatform, ConfigHierarchy> gameConfigs = new Dictionary<UnrealTargetPlatform, ConfigHierarchy>();
+			foreach (UnrealTargetPlatform targetPlatformType in UnrealTargetPlatform.GetValidPlatforms())
+			{
+				ConfigHierarchy engineConfig = ConfigCache.ReadHierarchy(ConfigHierarchyType.Engine, projectPath.Directory, targetPlatformType);
+				engineConfigs.Add(targetPlatformType, engineConfig);
+
+				ConfigHierarchy gameConfig = ConfigCache.ReadHierarchy(ConfigHierarchyType.Game, projectPath.Directory, targetPlatformType); 
+				gameConfigs.Add(targetPlatformType, gameConfig);
+			}
+
+			return (engineConfigs, gameConfigs);
+		}
+
 		internal static string SanitizeBucketValue(string s)
 		{
 			// trim starting + and then replace any + and . with -
diff --git a/Engine/Source/Programs/AutomationTool/BuildGraph/Tasks/RetrieveCloudArtifactTask.cs b/Engine/Source/Programs/AutomationTool/BuildGraph/Tasks/RetrieveCloudArtifactTask.cs
--- a/Engine/Source/Programs/AutomationTool/BuildGraph/Tasks/RetrieveCloudArtifactTask.cs
+++ b/Engine/Source/Programs/AutomationTool/BuildGraph/Tasks/RetrieveCloudArtifactTask.cs
@@ -162,8 +162,8 @@ namespace AutomationTool.Tasks
 
 			if (projectPath != null)
 			{
-				ProjectProperties properties = ProjectUtils.GetProjectProperties(projectPath);
-				ConfigHierarchy config = properties.EngineConfigs[HostPlatform.Current.HostEditorPlatform];
+				(Dictionary<UnrealTargetPlatform, ConfigHierarchy> engineConfigs, Dictionary<UnrealTargetPlatform, ConfigHierarchy> _) = CreateCloudArtifactTask.GetIniConfigs(projectPath);
+				ConfigHierarchy config = engineConfigs[HostPlatform.Current.HostEditorPlatform];
 				foundConfig = config.TryGetValueGeneric("StorageServers", "Cloud", out CloudConfiguration foundCloudConfig);
 				cloudConfig = foundCloudConfig;
 			}
