Specifying '-noraytracing' if raytracing is not enabled to force disabling of raytracing flows within the Engine

#rnx

diff --git a/Engine/Source/Programs/AutomationTool/Gauntlet/Unreal/Automation/UE.Automation.cs b/Engine/Source/Programs/AutomationTool/Gauntlet/Unreal/Automation/UE.Automation.cs
--- a/Engine/Source/Programs/AutomationTool/Gauntlet/Unreal/Automation/UE.Automation.cs
+++ b/Engine/Source/Programs/AutomationTool/Gauntlet/Unreal/Automation/UE.Automation.cs
@@ -274,7 +274,7 @@ namespace UE
 				}
 				else
 				{
-					AppConfig.CommandLine += " -dpcvars=r.RayTracing=0";
+					AppConfig.CommandLine += " -noraytracing";
 				}
 				
 				// Options specific to Desktop platforms
