Make pak order file specs more flexible:
- Add support for adding arbitrary suffixes to order files (ie don't throw an exception if an underscore doesn't preceed a number)
- Remove the implicit logic that makes order files without wildcards required. This leads to all files having wildcards added which causes issues when we don't want to match these patterns

#tests tested via preflight; validated the output of UnrealPak

diff --git a/Engine/Source/Programs/AutomationTool/Scripts/CopyBuildToStagingDirectory.Automation.cs b/Engine/Source/Programs/AutomationTool/Scripts/CopyBuildToStagingDirectory.Automation.cs
--- a/Engine/Source/Programs/AutomationTool/Scripts/CopyBuildToStagingDirectory.Automation.cs
+++ b/Engine/Source/Programs/AutomationTool/Scripts/CopyBuildToStagingDirectory.Automation.cs
@@ -67,8 +67,11 @@ namespace AutomationScripts
 				if (index != -1)
 				{
 					index++;
-					string Number = FileName.Substring(index, FileName.Length - index);
-					AppendOrder = Convert.ToInt32(Number);
+					string Suffix = FileName.Substring(index, FileName.Length - index);
+					if (int.TryParse(Suffix, out int Order))
+					{
+						AppendOrder = Order;
+					}
 				}
 
 				if(OrderType == OrderFileType.Cooker && AppendOrder == -1)
@@ -3778,7 +3781,7 @@ namespace AutomationScripts
 						}
 					}
 
-					OrderFileSpecs.Add(new OrderFileSpec { FileNamePattern = FilePattern, OrderType = OrderFile.OrderFileType.Custom, Required = FilePattern.Contains("*") == false, Priority = Priority });
+					OrderFileSpecs.Add(new OrderFileSpec { FileNamePattern = FilePattern, OrderType = OrderFile.OrderFileType.Custom, Required = false, Priority = Priority });
 				}
 			}
 
