Cloud Artifacts - Exposed the boost workers option to create and retrieve artifacts, disabled by default

diff --git a/Engine/Source/Programs/AutomationTool/BuildGraph/Tasks/CreateCloudArtifactTask.cs b/Engine/Source/Programs/AutomationTool/BuildGraph/Tasks/CreateCloudArtifactTask.cs
--- a/Engine/Source/Programs/AutomationTool/BuildGraph/Tasks/CreateCloudArtifactTask.cs
+++ b/Engine/Source/Programs/AutomationTool/BuildGraph/Tasks/CreateCloudArtifactTask.cs
@@ -157,6 +157,12 @@ namespace AutomationTool.Tasks
 		/// </summary>
 		[TaskParameter(Optional = true)]
 		public bool Verbose { get; set; } = true;
+
+		/// <summary>
+		/// Increase the number of worker threads used by zen, may cause machine to be less responsive but will generally improve upload times
+		/// </summary>
+		[TaskParameter(Optional = true)]
+		public bool BoostWorkers { get; set; } = false;
 	}
 
 	/// <summary>
@@ -470,7 +476,7 @@ namespace AutomationTool.Tasks
 				Branch = streamId,
 				BaselineBranch = baselineBranch,
 				Platform = platform
-			}, allowMultipart: _parameters.AllowMultipart, assumeHttp2: String.Equals(httpVersion, "http2-only", StringComparison.OrdinalIgnoreCase), forceClean: _parameters.Force, verify: _parameters.Verify, verbose: _parameters.Verbose);
+			}, allowMultipart: _parameters.AllowMultipart, assumeHttp2: String.Equals(httpVersion, "http2-only", StringComparison.OrdinalIgnoreCase), forceClean: _parameters.Force, verify: _parameters.Verify, verbose: _parameters.Verbose, boostWorkers: _parameters.BoostWorkers);
 
 			Uri buildUri = new Uri($"{descriptorHost}/api/v2/builds/{ns}/{bucketId}/{buildId}");
 
@@ -580,9 +586,10 @@ namespace AutomationTool.Tasks
 			public string Branch { get; set; }
 			public string BaselineBranch { get; set; }
 		}
-		public static void UploadBuild(string host, DirectoryReference baseDir, string namespaceId, string bucketId, CbObjectId buildId, string partName,
-			CbObjectId partId, Dictionary<string, object> metadata, string accessToken, BuildInfo info, bool allowMultipart = false, bool assumeHttp2 = false,
-			bool forceClean = false, bool verify = false, bool verbose = false)
+		public static void UploadBuild(string host, DirectoryReference baseDir, string namespaceId, string bucketId, CbObjectId buildId,
+			string partName, CbObjectId partId, Dictionary<string, object> metadata, string accessToken, BuildInfo info,
+			bool allowMultipart = false, bool assumeHttp2 = false, bool forceClean = false, bool verify = false, bool verbose = false,
+			bool boostWorkers = false)
 		{
 			string metadataJsonPath = Path.Combine(Path.GetTempPath(), Path.GetTempFileName());
 			try
@@ -623,8 +630,9 @@ namespace AutomationTool.Tasks
 				string http2Options = assumeHttp2 ? "--assume-http2" : "";
 				string verifyOption = verify ? "--verify" : "";
 				string verboseOption = verbose ? "--verbose" : "";
+				string boostWorkersOption = boostWorkers ? "--boost-workers" : "";
 				// pass the access token via environment variable to avoid showing it on the cli
-				string cmdline = $"builds upload --url {host} --namespace {namespaceId} --bucket {bucketId} --build-id {buildId} --build-part-name \"{partName}\" --build-part-id {partId} --metadata-path \"{metadataJsonPath}\" --access-token-env UE-CloudDataCacheAccessToken --create-build --plain-progress --allow-multipart={allowMultipart} {http2Options} --clean={forceClean} {verifyOption} {verboseOption} \"{baseDir}\"";
+				string cmdline = $"builds upload --url {host} --namespace {namespaceId} --bucket {bucketId} --build-id {buildId} --build-part-name \"{partName}\" --build-part-id {partId} --metadata-path \"{metadataJsonPath}\" --access-token-env UE-CloudDataCacheAccessToken --create-build --plain-progress --allow-multipart={allowMultipart} {http2Options} {boostWorkersOption} --clean={forceClean} {verifyOption} {verboseOption} \"{baseDir}\"";
 				Log.Logger.LogInformation("Using metadata '{Metadata}'", metadataString);
 				CommandUtils.RunAndLog(CommandUtils.CmdEnv, zenExe.FullName,cmdline , Options: CommandUtils.ERunOptions.Default, EnvVars: new Dictionary<string, string> {{"UE-CloudDataCacheAccessToken", accessToken}});
 			}
diff --git a/Engine/Source/Programs/AutomationTool/BuildGraph/Tasks/RetrieveCloudArtifactTask.cs b/Engine/Source/Programs/AutomationTool/BuildGraph/Tasks/RetrieveCloudArtifactTask.cs
--- a/Engine/Source/Programs/AutomationTool/BuildGraph/Tasks/RetrieveCloudArtifactTask.cs
+++ b/Engine/Source/Programs/AutomationTool/BuildGraph/Tasks/RetrieveCloudArtifactTask.cs
@@ -117,6 +117,12 @@ namespace AutomationTool.Tasks
 		/// </summary>
 		[TaskParameter(Optional = true)]
 		public string HttpVersion { get; set; } = "None";
+
+		/// <summary>
+		/// Increase the number of worker threads used by zen, may cause machine to be less responsive but will generally improve download times
+		/// </summary>
+		[TaskParameter(Optional = true)]
+		public bool BoostWorkers { get; set; } = false;
 	}
 
 	/// <summary>
@@ -377,7 +383,7 @@ namespace AutomationTool.Tasks
 			CbObjectId artifact = artifactIds.First();
 			Logger.LogInformation("Found artifact {ArtifactId}", artifact);
 
-			BuildDownload.DownloadBuild(host, outputDir, ns, bucketId, artifact.ToString(), accessToken, allowMultipart: _parameters.AllowMultipart, assumeHttp2: String.Equals(httpVersion, "http2-only", StringComparison.OrdinalIgnoreCase));
+			BuildDownload.DownloadBuild(host, outputDir, ns, bucketId, artifact.ToString(), accessToken, allowMultipart: _parameters.AllowMultipart, assumeHttp2: String.Equals(httpVersion, "http2-only", StringComparison.OrdinalIgnoreCase), boostWorkers: _parameters.BoostWorkers);
 		}
 
 		static HttpClient BuildHttpClient(string host, string ns, string bucketId, string accessToken)
@@ -451,14 +457,15 @@ namespace AutomationTool.Tasks
 
 	static class BuildDownload
 	{
-		public static void DownloadBuild(string host, DirectoryReference targetDir, string namespaceId, string bucketId, string buildId, string accessToken, bool allowMultipart = false, bool assumeHttp2 = false)
+		public static void DownloadBuild(string host, DirectoryReference targetDir, string namespaceId, string bucketId, string buildId, string accessToken, bool allowMultipart = false, bool assumeHttp2 = false, bool boostWorkers = false)
 		{
 			FileInfo zenExe = new FileInfo("Engine/Binaries/Win64/zen.exe");
 
 			string http2Options = assumeHttp2 ? "--assume-http2" : "";
+			string boostWorkersOption = boostWorkers ? "--boost-workers" : "";
 			// pass the access token via environment variable to avoid showing it on the cli
 			string cmdline =
-				$"builds download --url {host} --namespace {namespaceId} --bucket {bucketId} --access-token-env UE-CloudDataCacheAccessToken --plain-progress --allow-multipart={allowMultipart} {http2Options} \"{targetDir}\" {buildId}";
+				$"builds download --url {host} --namespace {namespaceId} --bucket {bucketId} --access-token-env UE-CloudDataCacheAccessToken --plain-progress --allow-multipart={allowMultipart} {http2Options} {boostWorkersOption} \"{targetDir}\" {buildId}";
 			CommandUtils.RunAndLog(CommandUtils.CmdEnv, zenExe.FullName, cmdline, Options: CommandUtils.ERunOptions.Default,
 				EnvVars: new Dictionary<string, string> { { "UE-CloudDataCacheAccessToken", accessToken } });
 		}
