Cloud Artifacts - Pre-register the artifact so that Horde is aware it will be created

diff --git a/Engine/Source/Programs/AutomationTool/BuildGraph/BgScriptReader.cs b/Engine/Source/Programs/AutomationTool/BuildGraph/BgScriptReader.cs
--- a/Engine/Source/Programs/AutomationTool/BuildGraph/BgScriptReader.cs
+++ b/Engine/Source/Programs/AutomationTool/BuildGraph/BgScriptReader.cs
@@ -1779,6 +1779,10 @@ namespace AutomationTool
 				{
 					AddArtifactFromTask(element, info.Arguments);
 				}
+				if (info.Name.Equals("CreateCloudArtifact", StringComparison.OrdinalIgnoreCase))
+				{
+					AddCloudArtifactFromTask(element, info.Arguments);
+				}
 			}
 		}
 
@@ -1808,6 +1812,37 @@ namespace AutomationTool
 			_graph.Artifacts.Add(artifact);
 		}
 
+		void AddCloudArtifactFromTask(BgScriptElement element, Dictionary<string, string> arguments)
+		{
+			arguments.TryGetValue(nameof(CreateCloudArtifactTaskParameters.Name), out string? name);
+			arguments.TryGetValue(nameof(CreateCloudArtifactTaskParameters.HordeArtifactName), out string? hordeArtifactName);
+			arguments.TryGetValue(nameof(CreateCloudArtifactTaskParameters.Type), out string? type);
+			arguments.TryGetValue(nameof(CreateCloudArtifactTaskParameters.Description), out string? description);
+			arguments.TryGetValue(nameof(CreateCloudArtifactTaskParameters.Metadata), out string? metadata);
+
+			string? artifactName = String.IsNullOrEmpty(hordeArtifactName) ? name : hordeArtifactName;
+			if (String.IsNullOrEmpty(artifactName))
+			{
+				LogError(element, "Missing artifact name");
+				return;
+			}
+			if (String.IsNullOrEmpty(type))
+			{
+				LogError(element, "Missing artifact type");
+				return;
+			}
+
+			string[] metadataArray = (metadata ?? String.Empty).Split(';', StringSplitOptions.RemoveEmptyEntries | StringSplitOptions.TrimEntries);
+			List<string> metadataList =
+			[
+				..metadataArray,
+				"Backend=Zen"
+			];
+
+			BgArtifactDef artifact = new BgArtifactDef(artifactName, type, description, null, _enclosingNode!.Name, null, Array.Empty<string>(), metadataList.ToArray());
+			_graph.Artifacts.Add(artifact);
+		}
+
 		/// <summary>
 		/// Reads the definition for an email notifier
 		/// </summary>
