Remove reference to WindowsDevicePortalWrapper.dll in Gauntlet script
#rnx

diff --git a/Engine/Source/Programs/AutomationTool/Gauntlet/Gauntlet.Automation.csproj b/Engine/Source/Programs/AutomationTool/Gauntlet/Gauntlet.Automation.csproj
--- a/Engine/Source/Programs/AutomationTool/Gauntlet/Gauntlet.Automation.csproj
+++ b/Engine/Source/Programs/AutomationTool/Gauntlet/Gauntlet.Automation.csproj
@@ -43,9 +43,6 @@
       <HintPath>..\..\..\..\Binaries\ThirdParty\Newtonsoft\Newtonsoft.Json.dll</HintPath>
       <Private>True</Private>
     </Reference>
-    <Reference Include="WindowsDevicePortalWrapper">
-      <HintPath>..\..\..\..\Binaries\ThirdParty\WindowsDevicePortalWrapper\WindowsDevicePortalWrapper.dll</HintPath>
-    </Reference>
   </ItemGroup>
 
     <!-- The following will glob all platform extensions' Gauntlet source
