[Automation] Correct exception thrown by TargetDeviceIOS CleanArtifacts when a bundle doesn't exist

diff --git a/Engine/Source/Programs/AutomationTool/Gauntlet/Platform/IOS/Gauntlet.TargetDeviceIOS.cs b/Engine/Source/Programs/AutomationTool/Gauntlet/Platform/IOS/Gauntlet.TargetDeviceIOS.cs
--- a/Engine/Source/Programs/AutomationTool/Gauntlet/Platform/IOS/Gauntlet.TargetDeviceIOS.cs
+++ b/Engine/Source/Programs/AutomationTool/Gauntlet/Platform/IOS/Gauntlet.TargetDeviceIOS.cs
@@ -427,7 +427,12 @@ namespace Gauntlet
 			
 			IProcessResult Result = ExecuteIOSDeployCommand(ListEngineSaved);
 			Result.WaitForExit();
-			if (Result.ExitCode != 0)
+
+			if (Result.ExitCode == 253) // Exit Code 253 - no bundle exists meaning there is nothing to clean
+			{
+				return;
+			}
+			else if (Result.ExitCode != 0)
 			{
 				throw new DeviceException("Listing Engine artifacts for {0} exited with code {1}:\n{2}", BundleID, Result.ExitCode, Result.Output);
 			}
