Align error fingerprint redaction process with Horde issue fingerprint

#rnx

diff --git a/Engine/Source/Programs/AutomationTool/Gauntlet/Framework/Utils/Gauntlet.HordeReport.cs b/Engine/Source/Programs/AutomationTool/Gauntlet/Framework/Utils/Gauntlet.HordeReport.cs
--- a/Engine/Source/Programs/AutomationTool/Gauntlet/Framework/Utils/Gauntlet.HordeReport.cs
+++ b/Engine/Source/Programs/AutomationTool/Gauntlet/Framework/Utils/Gauntlet.HordeReport.cs
@@ -1601,8 +1601,7 @@ namespace Gauntlet
 				// Redact numbers and timestamp like things
 				static Regex Regex_numbertimestampstring = new Regex(@"\d[\d.,:]*", RegexOptions.Compiled);
 				// Redact path
-				static Regex Regex_pathstring = new Regex(@"(?<![A-Z])(?:[A-Z]:|/)[^ :]+[/\\]SYNC[/\\]", RegexOptions.Compiled);
-				static Regex Regex_pathstring2 = new Regex(@"(?:^|\s)/[\w/.:-]+", RegexOptions.Compiled);
+				static Regex Regex_pathstring = new Regex(@"(?:(?<![A-Z])|HTTPS*:)(?:[A-Z]:|[/\\])[\w/\\.:+?=%-]+", RegexOptions.Compiled);
 				// Redact hex strings
 				static Regex Regex_hexstring = new Regex(@"(?:0X|#)[0-9A-F]+", RegexOptions.Compiled);
 
@@ -1667,7 +1666,6 @@ namespace Gauntlet
 				{
 					Text = Text.Trim().ToUpperInvariant();
 					Text = Regex_pathstring.Replace(Text, "");
-					Text = Regex_pathstring2.Replace(Text, "");
 					Text = Regex_hexstring.Replace(Text, "H");
 					Text = Regex_numbertimestampstring.Replace(Text, "n");
 
