Rename conflicting command line argument from -Device= to -DeviceName=
-Device= being used by the DevicePool to add device to the unprovisioned(local) pool

#rnx

diff --git a/Engine/Source/Programs/AutomationTool/Gauntlet/Unreal/RunUnreal.cs b/Engine/Source/Programs/AutomationTool/Gauntlet/Unreal/RunUnreal.cs
--- a/Engine/Source/Programs/AutomationTool/Gauntlet/Unreal/RunUnreal.cs
+++ b/Engine/Source/Programs/AutomationTool/Gauntlet/Unreal/RunUnreal.cs
@@ -37,7 +37,10 @@ namespace Gauntlet
 	[ParamHelp("NoTimeout", "No maximum timeout", ParamType = typeof(bool))]
 	[ParamHelp("TestIterations", "Number of iterations to repeat this test", ParamType = typeof(int), DefaultValue = 1)]
 	[ParamHelp("CookedEditor", "Restricts usage of uncooked editor role", ParamType = typeof(bool))]
-	[ParamHelp("Device", "List of devices to use for tests", Action = ParamHelpAttribute.ParamAction.Append)]
+	[ParamHelp("Devices", "List of devices to use for tests using the syntax: <platform>:<address>,<platform>:<address>,...", Action = ParamHelpAttribute.ParamAction.Append)]
+	[ParamHelp("DeviceName", "Device name to constrain reservation", Action = ParamHelpAttribute.ParamAction.Append)]
+	[ParamHelp("PrefSpec", "Device perfomance specification to constrain reservation (Minimum,Recommended,High)", Action = ParamHelpAttribute.ParamAction.Append)]
+	[ParamHelp("PrefModel", "Device model to constrain reservation", Action = ParamHelpAttribute.ParamAction.Append)]
 	[ParamHelp("NumClients", "Number of clients to run test with", ParamType = typeof(int), DefaultValue = 1)]
 	[ParamHelp("Server", "Run test with server", ParamType = typeof(bool))]
 	[ParamHelp("NullRHI", "Null Rendering Hardware Interface (run headless)", ParamType = typeof(bool))]
@@ -437,7 +440,7 @@ namespace Gauntlet
 				List<string> ModelArgs = CombinedParams.ParseValues("PerfModel", false);
 				string Model = ModelArgs.Count > 0 ? ModelArgs.Last() : string.Empty;
 
-				List<string> DeviceNameArgs = Globals.Params.ParseValues("Device", false);
+				List<string> DeviceNameArgs = Globals.Params.ParseValues("DeviceName", false);
 				string DeviceName = DeviceNameArgs.Count > 0 ? DeviceNameArgs.Last() : string.Empty;
 
 				TestContext.Constraint = new UnrealDeviceTargetConstraint(UnrealPlatform, PerfSpec, Model, DeviceName);
