[Gauntlet] Revert 41314226 in favor of using a UnrealTestRole to force dev artifact copies

diff --git a/Engine/Source/Programs/AutomationTool/Gauntlet/Framework/Gauntlet.Utils.cs b/Engine/Source/Programs/AutomationTool/Gauntlet/Framework/Gauntlet.Utils.cs
--- a/Engine/Source/Programs/AutomationTool/Gauntlet/Framework/Gauntlet.Utils.cs
+++ b/Engine/Source/Programs/AutomationTool/Gauntlet/Framework/Gauntlet.Utils.cs
@@ -44,7 +44,6 @@ namespace Gauntlet
 		public static bool CancelSignalled { get; private set; }
 
 		public static bool IsRunningDev => Params.ParseParam("dev");
-		public static bool IsForcingCopyArtifacts => Params.ParseParam("ForceCopyArtifacts");
 
 		/// <summary>
 		/// Returns the device pool id 
diff --git a/Engine/Source/Programs/AutomationTool/Gauntlet/Unreal/Base/Gauntlet.UnrealSession.cs b/Engine/Source/Programs/AutomationTool/Gauntlet/Unreal/Base/Gauntlet.UnrealSession.cs
--- a/Engine/Source/Programs/AutomationTool/Gauntlet/Unreal/Base/Gauntlet.UnrealSession.cs
+++ b/Engine/Source/Programs/AutomationTool/Gauntlet/Unreal/Base/Gauntlet.UnrealSession.cs
@@ -157,6 +157,12 @@ namespace Gauntlet
 		/// </summary>
 		public bool SkipCleanDeviceArtifacts { get; set; }
 
+		/// <summary>
+		/// Whether or not this role will save artifacts when running with -dev. 
+		/// These can be huge, so be cautious with using this.
+		/// </summary>
+		public bool ArchiveDevArtifacts { get; set; }
+
 		/// <summary>
 		/// Is this role Null?
 		/// </summary>
@@ -1060,14 +1066,11 @@ namespace Gauntlet
 			// Also, don't move artifacts in dev mode, because peoples saved data could be huuuuuuuge!
 			bool IsEditorBuild = InRunningRole.Role.RoleType.UsesEditor();
 			bool IsBuildMachine = CommandUtils.IsBuildMachine;
-			bool SkipArchivingAssets = Globals.IsRunningDev || (IsEditorBuild && (IsBuildMachine == false || InRunningRole.AppInstance.ExitCode == 0));
+			bool SkipArchivingAssets = (Globals.IsRunningDev && !InRunningRole.Role.ArchiveDevArtifacts) || (IsEditorBuild && (IsBuildMachine == false || InRunningRole.AppInstance.ExitCode == 0));
 			bool bRetainArtifacts = InContext.TestParams.ParseParam("RetainDeviceArtifacts");
-			bool bForceArchivingArtifacts = Globals.IsForcingCopyArtifacts;
 
-			bool bShouldArchiving = (bForceArchivingArtifacts || !SkipArchivingAssets);
-			
 			// Check if we should copy artifacts
-			if (bShouldArchiving)
+			if (!SkipArchivingAssets)
 			{
 				if (SourceDirectory.Exists)
 				{
diff --git a/Engine/Source/Programs/AutomationTool/Gauntlet/Unreal/Base/Gauntlet.UnrealTestConfiguration.cs b/Engine/Source/Programs/AutomationTool/Gauntlet/Unreal/Base/Gauntlet.UnrealTestConfiguration.cs
--- a/Engine/Source/Programs/AutomationTool/Gauntlet/Unreal/Base/Gauntlet.UnrealTestConfiguration.cs
+++ b/Engine/Source/Programs/AutomationTool/Gauntlet/Unreal/Base/Gauntlet.UnrealTestConfiguration.cs
@@ -510,6 +510,12 @@ namespace Gauntlet
 		/// Whether or not this role should skip the cleaning of its device's artifact between each test run
 		/// </summary>
 		public bool SkipCleanDeviceArtifacts { get; set; }
+
+		/// <summary>
+		/// Whether or not this role will save artifacts when running with -dev. 
+		/// These can be huge, so be cautious with using this.
+		/// </summary>
+		public bool ArchiveDevArtifacts { get; set; }
 	}
 
 	/// <summary>
diff --git a/Engine/Source/Programs/AutomationTool/Gauntlet/Unreal/Base/Gauntlet.UnrealTestNode.cs b/Engine/Source/Programs/AutomationTool/Gauntlet/Unreal/Base/Gauntlet.UnrealTestNode.cs
--- a/Engine/Source/Programs/AutomationTool/Gauntlet/Unreal/Base/Gauntlet.UnrealTestNode.cs
+++ b/Engine/Source/Programs/AutomationTool/Gauntlet/Unreal/Base/Gauntlet.UnrealTestNode.cs
@@ -614,6 +614,7 @@ namespace Gauntlet
 					SessionRole.RoleModifier = TestRole.RoleType;
 					SessionRole.Constraint = UseContextConstraint ? Context.Constraint : new UnrealDeviceTargetConstraint(SessionPlatform);
 					SessionRole.SkipCleanDeviceArtifacts = TestRole.SkipCleanDeviceArtifacts;
+					SessionRole.ArchiveDevArtifacts = TestRole.ArchiveDevArtifacts;
 					Log.Verbose("Created SessionRole {Role} from RoleContext {Context} (RoleType={Type})", SessionRole, RoleContext, TypesToRoles.Key);
 
 					// TODO - this can all / mostly go into UnrealTestConfiguration.ApplyToConfig
