Adapt deployment to zen pak streaming:
- add -ZenOmitPakFile as a project parameter to skip deployment of pak/utoc/ucas files

diff --git a/Engine/Source/Programs/AutomationTool/AutomationUtils/ProjectParams.cs b/Engine/Source/Programs/AutomationTool/AutomationUtils/ProjectParams.cs
--- a/Engine/Source/Programs/AutomationTool/AutomationUtils/ProjectParams.cs
+++ b/Engine/Source/Programs/AutomationTool/AutomationUtils/ProjectParams.cs
@@ -567,6 +567,7 @@ namespace AutomationTool
 			bool? ZenStore = null,
 			string NoZenAutoLaunch = null,
 			string ZenWorkspaceSharePath = null,
+			bool? ZenOmitPakFiles = null,
 			bool? SkipIoStore = null,
 			bool? GenerateOptimizationData = null,
 			bool? Prereqs = null,
@@ -802,6 +803,7 @@ namespace AutomationTool
 				this.AdditionalCookerOptions += string.Format(" -NoZenAutoLaunch={0}", this.NoZenAutoLaunch);
 			}
 			this.ZenWorkspaceSharePath = ParseParamValueIfNotSpecified(Command, ZenWorkspaceSharePath, "ZenWorkspaceSharePath", String.Empty);
+			this.ZenOmitPakFiles = GetParamValueIfNotSpecified(Command, ZenOmitPakFiles, false, "ZenOmitPakFiles");
 			this.GenerateOptimizationData = GetParamValueIfNotSpecified(Command, GenerateOptimizationData, this.GenerateOptimizationData, "makebinaryconfig");
 			
 			this.SkipPak = GetParamValueIfNotSpecified(Command, SkipPak, this.SkipPak, "skippak");
@@ -1489,6 +1491,11 @@ namespace AutomationTool
 		/// </summary>
 		public string ZenWorkspaceSharePath { private set; get; }
 
+		/// <summary>
+		/// Should omit pak files when deploying build to target for Zen pak streaming
+		/// </summary>
+		public bool ZenOmitPakFiles { private set; get; }
+
 		/// <summary>
 		/// Shared: True if optimization data is generated during staging that can improve loadtimes
 		/// </summary>
