UAT: Log as an Error if there is an error before the log writer is attached

#rnx

diff --git a/Engine/Source/Programs/AutomationTool/Program.cs b/Engine/Source/Programs/AutomationTool/Program.cs
--- a/Engine/Source/Programs/AutomationTool/Program.cs
+++ b/Engine/Source/Programs/AutomationTool/Program.cs
@@ -463,7 +463,7 @@ namespace AutomationToolDriver
 				{
 					FileReference errorLogFile = FileReference.Combine(Unreal.WritableEngineDirectory, "Programs/AutomationTool/Saved/Logs/ErrorLog.txt");
 					Log.AddFileWriterWithoutBackup("UnhandledFileLogListener", errorLogFile);
-					Logger.LogInformation("AutomationTool encountered an error before a log file was set, see '{File}' for more details", errorLogFile);
+					Logger.LogError("AutomationTool encountered an error before a log file was set, see '{File}' for more details", errorLogFile);
 				}
 
 				// Write the exit code
