Retry health tests that fail with time out

diff --git a/Engine/Source/Programs/AutomationTool/Gauntlet/Framework/Gauntlet.TestExecutor.cs b/Engine/Source/Programs/AutomationTool/Gauntlet/Framework/Gauntlet.TestExecutor.cs
--- a/Engine/Source/Programs/AutomationTool/Gauntlet/Framework/Gauntlet.TestExecutor.cs
+++ b/Engine/Source/Programs/AutomationTool/Gauntlet/Framework/Gauntlet.TestExecutor.cs
@@ -777,7 +777,6 @@ namespace Gauntlet
 				{
 					TestInfo.CancellationReason = string.Format("Terminating Test {0} due to maximum duration of {1} seconds. ", TestInfo.TestNode, TestInfo.TestNode.MaxDuration);
 					TestInfo.FinalResult = TestResult.TimedOut;
-					Log.Error(KnownLogEvents.Gauntlet_TestEvent, TestInfo.CancellationReason);
 				}
 				else if (TestInfo.TestNode.MaxDurationReachedResult == EMaxDurationReachedResult.Success)
 				{
@@ -816,7 +815,14 @@ namespace Gauntlet
 					TestInfo.TimeTestEnded = DateTime.Now;
 
 					TestResult NodeResult = TestInfo.TestNode.GetTestResult();
-					TestInfo.FinalResult = (TestInfo.FinalResult != TestResult.Invalid) ? TestInfo.FinalResult : NodeResult;
+					if (NodeResult == TestResult.WantRetry)
+					{
+						TestInfo.FinalResult = TestResult.WantRetry;
+					}
+					else
+					{
+						TestInfo.FinalResult = (TestInfo.FinalResult != TestResult.Invalid) ? TestInfo.FinalResult : NodeResult;
+					}
 
 					bool bCanFinalizeTest = true;
 					if (TestInfo.FinalResult == TestResult.WantRetry)
@@ -830,13 +836,14 @@ namespace Gauntlet
 						{
 							// Mark us as still running
 							TestInfo.CancellationReason = "";
+							TestInfo.TestNode.SetCancellationReason("");
 							TestInfo.FinalResult = TestResult.Invalid;
 							bCanFinalizeTest = false;
 						}
 						else
 						{
-							TestInfo.CancellationReason = "Failed to restart during retry.";
-							Log.Error(KnownLogEvents.Gauntlet_TestEvent, TestInfo.CancellationReason);
+							TestInfo.CancellationReason = !string.IsNullOrEmpty(TestInfo.CancellationReason) ?
+								$"Failed to restart during retry. Cancellation reason before retrying: {TestInfo.CancellationReason}" : "Failed to restart during retry.";
 							TestInfo.FinalResult = TestResult.Failed;
 						}
 					}
@@ -844,6 +851,10 @@ namespace Gauntlet
 					if (bCanFinalizeTest)
 					{
 						Log.Info("{Name} result={Result}", TestInfo, TestInfo.FinalResult);
+						if (!string.IsNullOrEmpty(TestInfo.CancellationReason) && TestInfo.FinalResult != TestResult.Cancelled)
+						{
+							Log.Error(KnownLogEvents.Gauntlet_TestEvent, "{CancellationReason}", TestInfo.CancellationReason);
+						}
 
 						if (!Options.DeferReports)
 						{
@@ -878,7 +889,7 @@ namespace Gauntlet
 					}
 
 					TestInfo.FinalResult = TestResult.Failed;
-				}				
+				}
 			}
 
 			return TestInfo.FinalResult;
