PRS - support summary stats in inline imports

Also adds a statname sanitizer and throws an exception if invalid inline stats are being provided.

#rnx
#tests local

diff --git a/Engine/Source/Programs/AutomationTool/Gauntlet/Framework/Utils/Gauntlet.CsvImporter.cs b/Engine/Source/Programs/AutomationTool/Gauntlet/Framework/Utils/Gauntlet.CsvImporter.cs
--- a/Engine/Source/Programs/AutomationTool/Gauntlet/Framework/Utils/Gauntlet.CsvImporter.cs
+++ b/Engine/Source/Programs/AutomationTool/Gauntlet/Framework/Utils/Gauntlet.CsvImporter.cs
@@ -71,6 +71,13 @@ namespace Gauntlet
 		[AllowNull]
 		public ICsvMetadata Metadata { get; set; }
 
+		/// <summary>
+		/// Additional csv stat averages to include.
+		/// Currently only implemented for inline imports.
+		/// </summary>
+		[AllowNull]
+		public Dictionary<string, double> AdditionalStatAverages { get; }
+
 		/// <summary>
 		/// Any additional files that should be bundled with this csv.
 		/// Key is the file type, value is the path to the file.
@@ -98,6 +105,7 @@ namespace Gauntlet
 		public Dictionary<string, string> CustomFields { get; private set; }
 		public ICsvMetadata Metadata { get; set; } = null;
 		public IReadOnlyDictionary<string, string> AdditionalFiles { get; private set; }
+		public Dictionary<string, double> AdditionalStatAverages { get; private set; } = null;
 
 		/// <summary>
 		/// CsvFile import constructor.
@@ -132,20 +140,37 @@ namespace Gauntlet
 		/// <param name="ImportName">Identifier for this import, primarily for logging.</param>
 		/// <param name="InlineMetadata">Metadata that would normally be provided by a csv. Must contain necessary fields.</param>
 		/// <param name="CustomFields">Optional datasource specific metadata. Corresponds to "CustomFields" in datasource config.</param>
-		public static CsvImportEntry CreateInlineImportEntry(string ImportName, ICsvMetadata CsvMetadata, Dictionary<string, string> CustomFields = null)
+		/// <param name="AdditionalStatAverages">
+		/// Optional dictionary of stat data to include in the import. These stats are treated like regular csv stat averages.
+		/// It's expected that these stats have been sanitized by the user. Any invalid stats will cause an ArgumentException to be thrown.
+		/// You can use the CsvStatSanitizer helper class to sanitize your stats.
+		/// </param>
+		public static CsvImportEntry CreateInlineImportEntry(string ImportName, ICsvMetadata CsvMetadata, Dictionary<string, double> AdditionalStatAverages = null, Dictionary<string, string> CustomFields = null)
 		{
 			if (CsvMetadata == null)
 			{
 				throw new ArgumentNullException("CsvMetadata", "Inline imports must have valid metadata.");
 			}
 
+			if (AdditionalStatAverages != null)
+			{
+				foreach ((string StatName, double Value) in AdditionalStatAverages)
+				{
+					if (!CsvStatSanitizer.IsValidStatName(StatName))
+					{
+						throw new ArgumentException($"AdditionalStat {StatName} is invalid. Please sanitize stats with CsvStatSanitizer first.", "AdditionalStatAverages");
+					}
+				}
+			}
+
 			return new CsvImportEntry()
 			{
 				ImportType = CsvImportType.Inline,
 				ImportName = ImportName,
 				Metadata = CsvMetadata,
 				CustomFields = CustomFields ?? new Dictionary<string, string>(),
-				CsvFilename = null
+				CsvFilename = null,
+				AdditionalStatAverages = AdditionalStatAverages
 			};
 		}
 
@@ -176,4 +201,54 @@ namespace Gauntlet
 	{
 		public void Import(IEnumerable<ICsvImport> Imports);
 	}
+
+	/// <summary>
+	/// Helper class to sanitize stat names.
+	/// </summary>
+	public class CsvStatSanitizer
+	{
+		// Based on FCsvStatNameValidator::ValidCharacters in CsvProfiler.cpp
+		private static string ValidCharacters = "0123456789abcdefghijklmnopqrstuvwxyzABCDEFGHIJKLMNOPQRSTUVWXYZ /_-[]()#.:";
+
+		/// <summary>
+		/// Returns true if the stat name only contains valid characters.
+		/// </summary>
+		public static bool IsValidStatName(string StatName)
+		{
+			return StatName.All(c => ValidCharacters.Contains(c));
+		}
+
+		/// <summary>
+		/// Returns a new stat string with all invalid characters removed.
+		/// </summary>
+		public static string GetSanitizedStatName(string StatName)
+		{
+			return SanitizeNameWithReplacement(StatName, (c) => null);
+		}
+
+		/// <summary>
+		/// Returns a new stat string with all invalid character replaced with the result of ReplacementFunc(char).
+		/// If ReplacementFunc(c) returns null that character is skipped.
+		/// </summary>
+		public static string SanitizeNameWithReplacement(string StatName, Func<char, char?> ReplacementFunc)
+		{
+			string SanitizedName = string.Empty;
+			foreach (char c in StatName)
+			{
+				if (!ValidCharacters.Contains(c))
+				{
+					char? ReplacementChar = ReplacementFunc(c);
+					if (ReplacementChar != null)
+					{
+						SanitizedName += ReplacementChar;
+					}
+				}
+				else
+				{
+					SanitizedName += c;
+				}
+			}
+			return SanitizedName;
+		}
+	}
 }
