[WinArm64] Using assembly manifest to redirect AppLocalDepdency DLLs to arm64 subfolder

#rnx

diff --git a/Engine/Build/Windows/Resources/arm64/DefaultSxS-Win64-arm64.manifest b/Engine/Build/Windows/Resources/arm64/DefaultSxS-Win64-arm64.manifest
new file mode 100644
--- /dev/null
+++ b/Engine/Build/Windows/Resources/arm64/DefaultSxS-Win64-arm64.manifest
@@ -0,0 +1,4 @@
+<?xml version="1.0" encoding="UTF-8"?>
+<assembly xmlns="urn:schemas-microsoft-com:asm.v1" manifestVersion="1.0">
+    <assemblyIdentity type="Win32" name="arm64" version="1.0.0.0" processorArchitecture="arm64" />
+</assembly>
diff --git a/Engine/Source/Programs/AutomationTool/Win/WinPlatform.Automation.cs b/Engine/Source/Programs/AutomationTool/Win/WinPlatform.Automation.cs
--- a/Engine/Source/Programs/AutomationTool/Win/WinPlatform.Automation.cs
+++ b/Engine/Source/Programs/AutomationTool/Win/WinPlatform.Automation.cs
@@ -17,6 +17,7 @@ using System.Reflection.PortableExecutable;
 using Microsoft.Extensions.Logging;
 
 using static AutomationTool.CommandUtils;
+using System.Xml.Linq;
 
 public class Win64Platform : Platform
 {
@@ -202,16 +203,17 @@ public class Win64Platform : Platform
 
 		if (!string.IsNullOrWhiteSpace(Params.AppLocalDirectory))
 		{
-			IEnumerable<string> ArchitectureNames = SC.StageTargets.SelectMany( Target => Target.Receipt.Architectures.Architectures.Select( Arch => Arch.WindowsToolChain) ).Distinct();
-			if (ArchitectureNames.Count() > 1)
+			bool bIsStagingX64   = (SC.StageTargets.Any(Target => Target.Receipt.Architectures.SingleArchitecture.bIsX64));
+			bool bIsStagingARM64 = (SC.StageTargets.Any(Target => !Target.Receipt.Architectures.SingleArchitecture.bIsX64));
+
+			if (bIsStagingX64)
 			{
-				Logger.LogWarning("Currently cannot use AppLocalDirectory with multiple staged architectures - only x64 AppLocalDependencies will be staged");
-				ArchitectureNames = ["x64"];
+				StageAppLocalDependencies(Params, SC, "Win64", UnrealArch.X64);
 			}
 
-			foreach (string ArchitectureName in ArchitectureNames)
+			if (bIsStagingARM64)
 			{
-				StageAppLocalDependencies(Params, SC, "Win64", ArchitectureName);
+				StageAppLocalDependencies(Params, SC, "Win64", UnrealArch.Arm64); // also used for Arm64ec
 			}
 		}
 	}
@@ -427,7 +429,7 @@ public class Win64Platform : Platform
 		return true;
 	}
 
-	public void StageAppLocalDependencies(ProjectParams Params, DeploymentContext SC, string SourcePlatformDir, string SourceArchitecture)
+	public void StageAppLocalDependencies(ProjectParams Params, DeploymentContext SC, string SourcePlatformDir, UnrealArch Architecture)
 	{
 		Dictionary<string, string> PathVariables = new Dictionary<string, string>();
 		PathVariables["EngineDir"] = SC.EngineRoot.FullName;
@@ -438,20 +440,21 @@ public class Win64Platform : Platform
 		foreach (string AppLocalDirectory in AppLocalDirectories)
 		{
 			string ExpandedAppLocalDir = Utils.ExpandVariables(AppLocalDirectory, PathVariables);
+			string ArchitectureDir = Architecture.WindowsToolChain;
 
-			DirectoryReference BaseAppLocalDependenciesPath = Path.IsPathRooted(ExpandedAppLocalDir) ? new DirectoryReference(CombinePaths(ExpandedAppLocalDir, SourcePlatformDir, SourceArchitecture)) : DirectoryReference.Combine(SC.ProjectRoot, ExpandedAppLocalDir, SourcePlatformDir, SourceArchitecture);
+			DirectoryReference BaseAppLocalDependenciesPath = Path.IsPathRooted(ExpandedAppLocalDir) ? new DirectoryReference(CombinePaths(ExpandedAppLocalDir, SourcePlatformDir, ArchitectureDir)) : DirectoryReference.Combine(SC.ProjectRoot, ExpandedAppLocalDir, SourcePlatformDir, ArchitectureDir);
 			if (!DirectoryReference.Exists(BaseAppLocalDependenciesPath))
 			{
 				BaseAppLocalDependenciesPath = BaseAppLocalDependenciesPath.ParentDirectory!; // no architecture subfolder - must be user-supplied dependency folder in need of upgrade
 				if (DirectoryReference.Exists(BaseAppLocalDependenciesPath))
 				{
-					Logger.LogWarning("{path} should have {arch} subfolder with the redist DLLs inside. Falling back to use {path} directly", BaseAppLocalDependenciesPath, SourceArchitecture, BaseAppLocalDependenciesPath);
+					Logger.LogWarning("{path} should have {arch} subfolder with the redist DLLs inside. Falling back to use {path} directly", BaseAppLocalDependenciesPath, Architecture, BaseAppLocalDependenciesPath);
 				}
 			}
 			if (DirectoryReference.Exists(BaseAppLocalDependenciesPath))
 			{
-				StageAppLocalDependenciesToDir(SC, BaseAppLocalDependenciesPath, StagedDirectoryReference.Combine("Engine", "Binaries", PlatformType.ToString()));
-				StageAppLocalDependenciesToDir(SC, BaseAppLocalDependenciesPath, StagedDirectoryReference.Combine(SC.RelativeProjectRootForStage, "Binaries", PlatformType.ToString()));
+				StageAppLocalDependenciesToDir(SC, BaseAppLocalDependenciesPath, Architecture, StagedDirectoryReference.Combine("Engine", "Binaries", PlatformType.ToString()));
+				StageAppLocalDependenciesToDir(SC, BaseAppLocalDependenciesPath, Architecture, StagedDirectoryReference.Combine(SC.RelativeProjectRootForStage, "Binaries", PlatformType.ToString()));
 			}
 			else
 			{
@@ -460,23 +463,60 @@ public class Win64Platform : Platform
 		}
 	}
 
-	static void StageAppLocalDependenciesToDir(DeploymentContext SC, DirectoryReference BaseAppLocalDependenciesPath, StagedDirectoryReference StagedBinariesDir)
+	static void StageAppLocalDependenciesToDir(DeploymentContext SC, DirectoryReference BaseAppLocalDependenciesPath, UnrealArch Architecture, StagedDirectoryReference StagedBinariesDir)
 	{
+		// x64 DLL dependencies are placed directly in StagedBinariesDir
+		bool bUseArchitectureSubdirectory = !Architecture.bIsX64;
+
+		StagedDirectoryReference StagedArchitectureBinariesDir = bUseArchitectureSubdirectory ? StagedDirectoryReference.Combine(StagedBinariesDir, Architecture.WindowsToolChain) : StagedBinariesDir;
+
 		// Check if there are any executables being staged in this directory. Usually we only need to stage runtime dependencies next to the executable, but we may be staging
 		// other engine executables too (eg. CEF)
 		List<StagedFileReference> FilesInTargetDir = SC.FilesToStage.NonUFSFiles.Keys.Where(x => x.IsUnderDirectory(StagedBinariesDir) && (x.HasExtension(".exe") || x.HasExtension(".dll"))).ToList();
 		if(FilesInTargetDir.Count > 0)
 		{
-			Logger.LogInformation("Copying AppLocal dependencies from {BaseAppLocalDependenciesPath} to {StagedBinariesDir}", BaseAppLocalDependenciesPath, StagedBinariesDir);
+			Logger.LogInformation("Copying AppLocal dependencies from {BaseAppLocalDependenciesPath} to {StagedArchitectureBinariesDir}", BaseAppLocalDependenciesPath, StagedArchitectureBinariesDir);
 
 			// Stage files in subdirs
 			foreach (DirectoryReference DependencyDirectory in DirectoryReference.EnumerateDirectories(BaseAppLocalDependenciesPath))
 			{	
-				SC.StageFiles(StagedFileType.NonUFS, DependencyDirectory, StageFilesSearch.AllDirectories, StagedBinariesDir);
+				SC.StageFiles(StagedFileType.NonUFS, DependencyDirectory, StageFilesSearch.AllDirectories, StagedArchitectureBinariesDir);
 			}
 		}
 	}
 
+
+	public void UpdateStagedSxSManifest( DeploymentContext SC, UnrealArch Architecture, StagedDirectoryReference StagedBinariesDir)
+	{
+		// x64 doesn't need SxS manifest - DLL dependencies are placed directly in StagedBinariesDir
+		if (Architecture.bIsX64)
+		{
+			return;
+		}
+
+		// try to update the staged SxS manifest
+		DirectoryReference ArchitectureDir = DirectoryReference.Combine( SC.StageDirectory, StagedBinariesDir.Name, Architecture.WindowsToolChain );
+		FileReference ManifestFile = FileReference.Combine(ArchitectureDir, Architecture + ".manifest");
+		if (FileReference.Exists(ManifestFile))
+		{
+			WindowsExports.UpdateSxSManifest(ManifestFile);
+		}
+	}
+
+
+	public override void PostStagingFileCopy(ProjectParams Params, DeploymentContext SC)
+	{
+		base.PostStagingFileCopy(Params, SC);
+
+		// refresh the architecture SxS manifests if they are in use
+		foreach (UnrealArch Architecture in SC.StageTargets.Select( Target => Target.Receipt.Architectures.SingleArchitecture ).Distinct())
+		{
+			UpdateStagedSxSManifest(SC, Architecture, StagedDirectoryReference.Combine("Engine", "Binaries", PlatformType.ToString()));
+			UpdateStagedSxSManifest(SC, Architecture, StagedDirectoryReference.Combine(SC.RelativeProjectRootForStage, "Binaries", PlatformType.ToString()));
+		}
+	}
+
+
     /// <summary>
     /// Try to get the symstore.exe path from the given Windows SDK version
     /// </summary>
diff --git a/Engine/Source/Programs/UnrealBuildTool/Platform/Windows/UEBuildWindows.cs b/Engine/Source/Programs/UnrealBuildTool/Platform/Windows/UEBuildWindows.cs
--- a/Engine/Source/Programs/UnrealBuildTool/Platform/Windows/UEBuildWindows.cs
+++ b/Engine/Source/Programs/UnrealBuildTool/Platform/Windows/UEBuildWindows.cs
@@ -547,6 +547,12 @@ namespace UnrealBuildTool
 		/// </summary>
 		public string? ManifestFile;
 
+		/// <summary>
+		/// Specifies the path to a manifest file that defines our SxS assembly on non-x64 architecture, allowing both x64 and arm64 to be staged together with separate runtime dependencies. 
+		/// Defaults to the manifest in Engine/Build/Windows/Resources/{architecture}/. Can be assigned to null if the SxS manifest is not required
+		/// </summary>
+		public string? SxSAssemblyManifestFile;
+
 		/// <summary>
 		/// Specifies the path to an .ico file to use as the appliction icon. Can be assigned to null and will default to Engine/Build/Windows/Resources/Default.ico for engine targets or Build/Windows/Application.ico for projects.
 		/// </summary>
@@ -893,6 +899,15 @@ namespace UnrealBuildTool
 				Platform += "-arm64";
 			}
 			ManifestFile = FileReference.Combine(Unreal.EngineDirectory, "Build", "Windows", "Resources", String.Format("Default-{0}.manifest", Platform)).FullName;
+
+			if (Target.Platform == UnrealTargetPlatform.Win64 && !Target.Architecture.bIsX64)
+			{
+				SxSAssemblyManifestFile = FileReference.Combine(Unreal.EngineDirectory, "Build", "Windows", "Resources", "arm64", String.Format("DefaultSxS-{0}.manifest", Platform)).FullName;
+			}
+			else
+			{
+				SxSAssemblyManifestFile = null;
+			}
 		}
 	}
 
@@ -995,7 +1010,9 @@ namespace UnrealBuildTool
 		public string? ModuleDefinitionFile => Inner.ModuleDefinitionFile;
 
 		public string? ManifestFile => Inner.ManifestFile;
-
+		
+		public string? SxSAssemblyManifestFile => Inner.SxSAssemblyManifestFile;
+		
 		public string? ApplicationIcon => Inner.ApplicationIcon;
 
 		public bool bNeedsLegacyStdioDefinitionsLib => Inner.bNeedsLegacyStdioDefinitionsLib;
diff --git a/Engine/Source/Programs/UnrealBuildTool/Platform/Windows/VCToolChain.cs b/Engine/Source/Programs/UnrealBuildTool/Platform/Windows/VCToolChain.cs
--- a/Engine/Source/Programs/UnrealBuildTool/Platform/Windows/VCToolChain.cs
+++ b/Engine/Source/Programs/UnrealBuildTool/Platform/Windows/VCToolChain.cs
@@ -9,6 +9,7 @@ using System.IO;
 using System.Linq;
 using System.Runtime.Versioning;
 using System.Text;
+using System.Xml.Linq;
 using UnrealBuildBase;
 
 namespace UnrealBuildTool
@@ -1480,6 +1481,12 @@ namespace UnrealBuildTool
 				Arguments.Add("/MANIFEST:EMBED");
 				FileItem ManifestFile = FileItem.GetItemByPath(Target.WindowsPlatform.ManifestFile);
 				Arguments.Add($"/MANIFESTINPUT:\"{NormalizeCommandLinePath(ManifestFile, RootPaths)}\"");
+
+				// Embed a dependency on the SxS assembly manifest. this is for dependent runtime DLLs
+				if (Target.WindowsPlatform.SxSAssemblyManifestFile != null && !Target.Architecture.bIsX64)
+				{
+					Arguments.Add("/MANIFESTDEPENDENCY:\"name='arm64' processorArchitecture='arm64' version='1.0.0.0' type='win32'\"");
+				}
 			}
 
 			// Prevents the linker from displaying its logo for each invocation.
@@ -2432,6 +2439,77 @@ namespace UnrealBuildTool
 					RuntimeDependencies.Add(new RuntimeDependency(DstFile, StagedFileType.NonUFS));
 				}
 			}
+
+			// mark the generated SxS manifest as a runtime dependency
+			if (Target.WindowsPlatform.SxSAssemblyManifestFile != null && !Target.Architecture.bIsX64)
+			{
+				string Architecture = Target.Architecture.WindowsToolChain;
+				FileReference OutputManifestFile = FileReference.Combine(ExeDir, Architecture, Architecture + ".manifest");
+				RuntimeDependencies.Add( new RuntimeDependency(OutputManifestFile));
+			}
+		}
+
+
+
+		[ToolMode("UpdateSxSManifest", ToolModeOptions.None)]
+		private class UpdateSxSManifest : ToolMode
+		{
+			[CommandLine("-Template", Required = true)]
+			FileReference? Template = null;
+
+			[CommandLine("-Manifest", Required = true)]
+			FileReference? Manifest = null;
+
+			public override System.Threading.Tasks.Task<int> ExecuteAsync(EpicGames.Core.CommandLineArguments Arguments, ILogger Logger)
+			{
+				Arguments.ApplyTo(this);
+				Arguments.CheckAllArgumentsUsed();
+
+				bool bSuccess = false;
+				try
+				{
+					FileReference.Copy(Template!, Manifest!, true);
+					bSuccess = WindowsExports.UpdateSxSManifest(Manifest!);
+				}
+				catch(Exception e)
+				{
+					Logger.LogError(e.Message);
+				}
+
+				return System.Threading.Tasks.Task.FromResult(bSuccess ? 0 : 1);
+			}
+		};
+
+		public override ICollection<FileItem> PostBuild(ReadOnlyTargetRules Target, FileItem Executable, LinkEnvironment BinaryLinkEnvironment, IActionGraphBuilder Graph)
+		{
+			ICollection<FileItem> OutputFiles = base.PostBuild(Target, Executable, BinaryLinkEnvironment, Graph);
+
+			if (Target.WindowsPlatform.ManifestFile != null && Target.WindowsPlatform.SxSAssemblyManifestFile != null && !Target.Architecture.bIsX64 && !BinaryLinkEnvironment.bIsBuildingDLL)
+			{
+				string Architecture = Target.Architecture.WindowsToolChain;
+				DirectoryReference ManifestDir = DirectoryReference.Combine(Executable.Directory.Location, Architecture);
+				FileItem ManifestFileItem = FileItem.GetItemByFileReference(FileReference.Combine( ManifestDir, Architecture + ".manifest" ));
+
+				string ActionArgs = $"-Template={Target.WindowsPlatform.SxSAssemblyManifestFile} -Manifest={ManifestFileItem.Location}";
+				Action PostBuildAction = Graph.CreateRecursiveAction<UpdateSxSManifest>(ActionType.CreateAppBundle, ActionArgs);
+				PostBuildAction.ProducedItems.Add(ManifestFileItem);
+				PostBuildAction.DeleteItems.Add(ManifestFileItem);
+				PostBuildAction.StatusDescription = $"{Architecture}.manifest";
+				PostBuildAction.bCanExecuteRemotely = false;
+
+				foreach ( ModuleRules.RuntimeDependency Dependency in BinaryLinkEnvironment.RuntimeDependencies.Where( D => D?.Type == StagedFileType.NonUFS))
+				{
+					FileReference DependencyFile = FileReference.FromString( Dependency.Path );
+					if (DependencyFile.HasExtension(".dll") && DependencyFile.IsUnderDirectory(ManifestDir))
+					{
+						PostBuildAction.PrerequisiteItems.Add( FileItem.GetItemByFileReference(DependencyFile) );
+					}
+				}
+
+				OutputFiles.Add(ManifestFileItem);
+			}
+
+			return OutputFiles;
 		}
 
 		public virtual FileReference GetApplicationIcon(FileReference? ProjectFile)
@@ -2968,6 +3046,11 @@ namespace UnrealBuildTool
 				{
 					PrerequisiteItems.Add(FileItem.GetItemByPath(Target.WindowsPlatform.ManifestFile));
 				}
+
+				if (Target.WindowsPlatform.SxSAssemblyManifestFile != null && !LinkEnvironment.bIsBuildingDLL)
+				{
+					PrerequisiteItems.Add(FileItem.GetItemByPath(Target.WindowsPlatform.SxSAssemblyManifestFile));
+				}
 			}
 
 			Arguments.AddRange(InputFileNames);
diff --git a/Engine/Source/Programs/UnrealBuildTool/Platform/Windows/WindowsExports.cs b/Engine/Source/Programs/UnrealBuildTool/Platform/Windows/WindowsExports.cs
--- a/Engine/Source/Programs/UnrealBuildTool/Platform/Windows/WindowsExports.cs
+++ b/Engine/Source/Programs/UnrealBuildTool/Platform/Windows/WindowsExports.cs
@@ -5,6 +5,7 @@ using System.Collections.Generic;
 using System.Diagnostics.CodeAnalysis;
 using System.Linq;
 using System.Runtime.Versioning;
+using System.Xml.Linq;
 using EpicGames.Core;
 using Microsoft.Extensions.Logging;
 
@@ -126,5 +127,45 @@ namespace UnrealBuildTool
 			return MicrosoftPlatformSDK.TryGetToolChainDir(Compiler, CompilerVersion, Architecture, Logger, out OutToolChainVersion, out OutToolChainDir, out OutRedistDir);
 		}
 
+
+
+		/// <summary>
+		/// Updates the given SxS manifest to include all of the DLL files in its directory
+		/// </summary>
+		/// <param name="SxSManifestFile"></param>
+		/// <returns>True if it succeeded</returns>
+		public static bool UpdateSxSManifest(FileReference SxSManifestFile)
+		{
+			bool bSuccess = false;
+			try
+			{
+				XDocument XmlManifest = XDocument.Load(SxSManifestFile.FullName);
+				if (XmlManifest.Root != null && XmlManifest.Root.Name.LocalName == "assembly")
+				{
+					XNamespace XmlNs = XmlManifest.Root.Name.Namespace;
+
+					// update the manifest with any DLLs that are in this folder
+					HashSet<string> ExistingDLLNames = [.. XmlManifest.Root.Elements(XmlNs+"file").Select( X => X.Attribute("name")!.Value )];
+					IEnumerable<FileReference> DLLs = DirectoryReference.EnumerateFiles(SxSManifestFile.Directory, "*.dll");
+					foreach (FileReference DLL in DLLs)
+					{
+						string DLLName = DLL.MakeRelativeTo(SxSManifestFile.Directory);
+						if (!ExistingDLLNames.Contains(DLLName))
+						{
+							XmlManifest.Root.Add( new XElement(XmlNs+"file", new XAttribute("name", DLLName ) ) );
+						}
+					}
+
+					// save the manifest again
+					XmlManifest.Save(SxSManifestFile.FullName);
+					bSuccess = true;
+				}
+			}
+			catch(Exception)
+			{
+			}
+
+			return bSuccess;
+		}
 	}
 }
