[TargetDeviceIOS] Improve Agent connection stability across sessions.
- Added parsing for Tunnel State when validating iDevice connections.
- Kill and restart usbmuxd at the start of tests.

diff --git a/Engine/Source/Programs/AutomationTool/Gauntlet/Platform/IOS/Gauntlet.TargetDeviceIOS.cs b/Engine/Source/Programs/AutomationTool/Gauntlet/Platform/IOS/Gauntlet.TargetDeviceIOS.cs
--- a/Engine/Source/Programs/AutomationTool/Gauntlet/Platform/IOS/Gauntlet.TargetDeviceIOS.cs
+++ b/Engine/Source/Programs/AutomationTool/Gauntlet/Platform/IOS/Gauntlet.TargetDeviceIOS.cs
@@ -269,6 +269,7 @@ namespace Gauntlet
 				}
 
 				IProcessResult ConnectResult = ExecuteDevicectlCommand("manage pair");
+				ConnectResult.WaitForExit();
 				if(ConnectResult.ExitCode != 0)
 				{
 					Log.Info("Connecting to {DeviceName} exited with code {ExitCode}:\n{Output}", UUID, ConnectResult.ExitCode, ConnectResult.Output);
@@ -278,20 +279,17 @@ namespace Gauntlet
 				string DetailsFile = Path.GetTempFileName();
 				string DetailsCommand = string.Format("device info details -j {0}", DetailsFile);
 				IProcessResult DetailResult = ExecuteDevicectlCommand(DetailsCommand);
+				DetailResult.WaitForExit();
 				if(DetailResult.ExitCode != 0)
 				{
 					Log.Info("Detailing {DeviceName} exited with code {ExitCode}. Skipping connection:\n{Output}", UUID, DetailResult.ExitCode, ConnectResult.Output);
 				}
 
 				JsonObject JsonDevice = JsonObject.Parse(File.ReadAllText(DetailsFile)).GetObjectField("result");
+				bool bIsConnected = ValidateDeviceConnection(JsonDevice, true);
 
-				JsonObject ConnectionProperties = JsonDevice
-					.GetObjectField("connectionProperties");
-				string PairStatus = ConnectionProperties.GetStringField("pairingState");
-
-				if(!PairStatus.Equals("paired", StringComparison.OrdinalIgnoreCase))
+				if (!bIsConnected)
 				{
-					Log.Info("{DeviceName}'s pair status is still not paired. Skipping connection.", UUID);
 					return false;
 				}
 
@@ -1008,13 +1006,15 @@ namespace Gauntlet
 
 				UseDeviceCtl = XCodeMajorVersion >= 15;
 
-				// Gather the iOS devices are connected to this UAT host
-				ConnectedDevices = GetConnectedDevices();
-
 				// Clear any zombie processes previous UAT instances may have left running
 				AppleBuild.ExecuteCommand("killall", "xcrun").WaitForExit();
 				AppleBuild.ExecuteCommand("killall", "ios-deploy").WaitForExit();
 				AppleBuild.ExecuteCommand("killall", "lldb").WaitForExit();
+				AppleBuild.ExecuteCommand("killall", "usbmuxd").WaitForExit();
+				AppleBuild.ExecuteCommand("launchctl", "kickstart -k system/com.apple.usbmuxd").WaitForExit();
+
+				// Gather the iOS devices are connected to this UAT host
+				ConnectedDevices = GetConnectedDevices();
 
 				if (!ConnectedDevices.Any())
 				{
@@ -1055,9 +1055,10 @@ namespace Gauntlet
 					string UUID = JsonDevice.GetObjectField("hardwareProperties").GetStringField("udid");
 					string OSVersionString = JsonDevice.GetObjectField("deviceProperties").GetStringField("osVersionNumber");
 					int IOSMajorVersion = int.Parse(OSVersionString.Substring(0, OSVersionString.IndexOf('.')));
-					bool bIsPaired = JsonDevice.GetObjectField("connectionProperties").GetStringField("pairingState").Equals("paired", StringComparison.OrdinalIgnoreCase);
 
-					ConnectedDevice Device = new ConnectedDevice(UUID, IOSMajorVersion, bIsPaired);
+					bool bIsConnected = ValidateDeviceConnection(JsonDevice, true);
+
+					ConnectedDevice Device = new ConnectedDevice(UUID, IOSMajorVersion, bIsConnected);
 					Devices.Add(UUID, Device);
 				}
 			}
@@ -1100,6 +1101,29 @@ namespace Gauntlet
 			return string.Empty;
 		}
 
+		private static bool ValidateDeviceConnection(JsonObject JsonDevice, bool LogFailure = false)
+		{
+			JsonObject HardwareProperites = JsonDevice.GetObjectField("hardwareProperties");
+			bool bParseUUID = HardwareProperites.TryGetStringField("udid", out string UUID);
+
+			JsonObject ConnectionProperties = JsonDevice.GetObjectField("connectionProperties");
+			// Two hardware values that indicate a healthy pairing to the host agent.
+			bool bParsePairStatus = ConnectionProperties.TryGetStringField("pairingState", out string PairStatus);
+			bool bParseTunnelState = ConnectionProperties.TryGetStringField("tunnelState", out string TunnelState);
+
+			bool bParseSuccessful = bParseUUID && bParsePairStatus && bParseTunnelState;
+			bool bIsConnected = bParseSuccessful &&
+								PairStatus.Equals("paired", StringComparison.OrdinalIgnoreCase) &&
+								TunnelState.Equals("connected", StringComparison.OrdinalIgnoreCase);
+
+			if (!bIsConnected && LogFailure)
+			{
+				Log.Info("{DeviceName} is still not connected to host. Skipping connection. \n Pair Status: {PairStatus} \n Tunnel State: {TunnelState}", UUID, PairStatus, TunnelState);
+			}
+
+			return bIsConnected;
+		}
+
 		// DEPRECATED
 		public IAppInstall InstallApplication(UnrealAppConfig AppConfig)
 		{
