Localization:
- Fixed a bug in SmartlingLocalizationProvider. Newly downloaded files from Smartling would be marked "read-only" before getting to the commandlets causing jobs that require writing the .po files to fail.
- Updated the logs of SmartlingLocalizationProvider to show download/upload tiems up to 2 decimal places
- Updated the logs of Smartling downloads to print the size of the downloaded fiels from Smartling to verify successful download and content
- Added logs when a culture directory is missing and being created as part of the UAT script

#test Locally built UE5 and performed the download, gather and export steps for the MetaHuman plugin. Verified that this fixes the download and export issue locally.

diff --git a/Engine/Source/Programs/AutomationTool/SmartlingLocalization/SmartlingLocalizationProvider.cs b/Engine/Source/Programs/AutomationTool/SmartlingLocalization/SmartlingLocalizationProvider.cs
--- a/Engine/Source/Programs/AutomationTool/SmartlingLocalization/SmartlingLocalizationProvider.cs
+++ b/Engine/Source/Programs/AutomationTool/SmartlingLocalization/SmartlingLocalizationProvider.cs
@@ -111,7 +111,7 @@ namespace EpicGames.SmartlingLocalization
 			}
 			await Task.WhenAll(DownloadTasks);
 			Watch.Stop();
-			Logger.LogInformation($"Completed Smartling download for {ProjectName} project files in {Watch.ElapsedMilliseconds / 1000} seconds.");
+			Logger.LogInformation($"Completed Smartling download for {ProjectName} project files in {Watch.ElapsedMilliseconds / 1000.0:F2} seconds.");
 		}
 
 		private async Task DownloadLatestPOFileAndLog(string EpicLocale, string Platform, ProjectImportExportInfo ProjectImportInfo)
@@ -192,7 +192,7 @@ namespace EpicGames.SmartlingLocalization
 					if (CurrentTries > MaxTries)
 					{
 						DownloadStopWatch.Stop();
-						Logs.Add($"[FAILED] Downloading: '{ExportFile.FullName}' ({EpicLocale}) exhausted all retries in {DownloadStopWatch.ElapsedMilliseconds / 1000} seconds. - {Ex}");
+						Logs.Add($"[FAILED] Downloading: '{ExportFile.FullName}' ({EpicLocale}) exhausted all retries in {DownloadStopWatch.ElapsedMilliseconds / 1000.0:F2} seconds. - {Ex}");
 						return Logs;
 					}
 					Logs.Add($"Failed to get download response. Retrying {CurrentTries}/{MaxTries} times.");
@@ -205,9 +205,15 @@ namespace EpicGames.SmartlingLocalization
 			{
 				if (!CultureDirectory.Exists)
 				{
+					Logger.LogInformation($"Culture directory for '{EpicLocale}' does not exist. Creating directory '{CultureDirectory.FullName}'");
 					CultureDirectory.Create();
+					CultureDirectory.Refresh();
 				}
-				bool ExportFileWasReadOnly = true;
+				// Do not initialize this to true. Further down we will do a check to see the previous state of the file and restore the read-only settings of the file.
+				// Even if a file did not exist on disk yet and was downloaded by Smartling, fi this defaults to true, the file will be set to read-only
+				// Newly downloaded files would be added by Perforce during commandlet execution. Files added by Perforce will NOT have their read-only status removed 
+				// That only happens for perforce check outs. As such, adding a newly downloaded file that is read-only will cause commandlets that need to update the file to fail 
+				bool ExportFileWasReadOnly = false;
 				if (ExportFile.Exists)
 				{
 					// We're going to clobber the existing PO file, so make sure it's writable (it may be read-only if in Perforce)
@@ -222,8 +228,10 @@ namespace EpicGames.SmartlingLocalization
 						await DownloadStream.CopyToAsync(DownloadFileStream);
 					}
 				}
+				// We refresh to update the state of the ifle after it's downloaded 
+				ExportFile.Refresh();
 				DownloadStopWatch.Stop();
-				Logs.Add($"[SUCCESS] Downloading: '{SmartlingFileUri}' as '{ExportFile.FullName}' ({EpicLocale}) in {DownloadStopWatch.ElapsedMilliseconds / 1000} seconds.");
+				Logs.Add($"[SUCCESS] Downloading: '{SmartlingFileUri}' as '{ExportFile.FullName}' ({EpicLocale}) size {ExportFile.Length / 1024.0:F2}kb in {DownloadStopWatch.ElapsedMilliseconds / 1000.0:F2} seconds.");
 				// Reset the write status of the file
 				if (ExportFileWasReadOnly)
 				{
@@ -261,7 +269,7 @@ namespace EpicGames.SmartlingLocalization
 			{
 				// The file may not currently exist in Smartling and will need to be uploaded first via the Upload step later on. 
 				DownloadStopWatch.Stop();
-				Logs.Add($"[FAILED] Downloading: '{ExportFile.FullName}' ({EpicLocale} in {DownloadStopWatch.ElapsedMilliseconds / 1000} seconds. The file may need to be uploaded first.)");
+				Logs.Add($"[FAILED] Downloading: '{ExportFile.FullName}' ({EpicLocale} in {DownloadStopWatch.ElapsedMilliseconds / 1000.0:F2} seconds. The file may need to be uploaded first.)");
 				await AppendRequestErrorsToLog(DownloadResponse, Logs);
 			}
 			return Logs;
@@ -295,7 +303,7 @@ namespace EpicGames.SmartlingLocalization
 
 			// Uploading all cultures is a legacy behavior from Onesky that isn't currently needed within our workflow. We can support this if there is a need. 
 			Watch.Stop();
-			Logger.LogInformation($"Completed Smartling upload for {ProjectName} project files in {Watch.ElapsedMilliseconds / 1000} seconds.");
+			Logger.LogInformation($"Completed Smartling upload for {ProjectName} project files in {Watch.ElapsedMilliseconds / 1000.0:F2} seconds.");
 		}
 
 		private async Task UploadLatestPOFileAndLog(string EpicLocale, string Platform, ProjectImportExportInfo ProjectExportInfo)
