Add argument to remove throw of a timeout error for RPC calls for debugging reason

#rnx

diff --git a/Engine/Source/Programs/AutomationTool/Gauntlet/Framework/RpcFramework/Gauntlet.RpcExecutor.cs b/Engine/Source/Programs/AutomationTool/Gauntlet/Framework/RpcFramework/Gauntlet.RpcExecutor.cs
--- a/Engine/Source/Programs/AutomationTool/Gauntlet/Framework/RpcFramework/Gauntlet.RpcExecutor.cs
+++ b/Engine/Source/Programs/AutomationTool/Gauntlet/Framework/RpcFramework/Gauntlet.RpcExecutor.cs
@@ -122,8 +122,8 @@ namespace Gauntlet
 					if (ExecutorInstance == null)
 					{
 						ExecutorInstance = new RpcExecutor();
-						int RpcHttpTimeoutInMilliseconds = Globals.Params.ParseValue("RPCHttpTimeoutInMilliseconds", 100 * 1000);
-						ExecutorClient.Timeout = new TimeSpan(0, 0, 0, 0, RpcHttpTimeoutInMilliseconds);
+						ExecutorClient.Timeout = Globals.Params.ParseParam("NoRPCHttpTimeout") ? Timeout.InfiniteTimeSpan 
+							: TimeSpan.FromMilliseconds(Globals.Params.ParseValue("RPCHttpTimeoutInMilliseconds", 100 * 1000));
 					}
 					return ExecutorInstance;
 				}
