Fixed QueueForMatch C# wrapper and C++ endpoint

#rnx

diff --git a/Engine/Source/Programs/AutomationTool/Gauntlet/Framework/RpcFramework/Gauntlet.RpcLibrary.cs b/Engine/Source/Programs/AutomationTool/Gauntlet/Framework/RpcFramework/Gauntlet.RpcLibrary.cs
--- a/Engine/Source/Programs/AutomationTool/Gauntlet/Framework/RpcFramework/Gauntlet.RpcLibrary.cs
+++ b/Engine/Source/Programs/AutomationTool/Gauntlet/Framework/RpcFramework/Gauntlet.RpcLibrary.cs
@@ -4,6 +4,7 @@ using Newtonsoft.Json;
 using System;
 using System.Collections.Generic;
 using System.Diagnostics;
+using System.Diagnostics.CodeAnalysis;
 using System.IO;
 using System.Linq;
 using System.Net.Http;
@@ -631,6 +632,21 @@ namespace Gauntlet
 
 	public static class RpcLibraryExtensions
 	{
+		// Adds an RPC argument to the dictionary if it's value is not null
+#nullable enable
+		public static bool AddOptional<T>(this Dictionary<string, object> InArgs, string InKey, [NotNullWhen(true)] T? InValue)
+		{
+			if (InValue is not null)
+			{
+				InArgs.Add(InKey, InValue);
+
+				return true;
+			}
+
+			return false;
+		}
+#nullable disable
+
 		public static bool AutomationControllerClearReports(this RpcTarget InTarget)
 		{
 			const string RPCName = "AutomationControllerClearReports";
