Fix for CreatePlatformExtension when -project= does not specify a full project path.

#rnx

diff --git a/Engine/Source/Programs/AutomationTool/Scripts/CreatePlatformExtension.cs b/Engine/Source/Programs/AutomationTool/Scripts/CreatePlatformExtension.cs
--- a/Engine/Source/Programs/AutomationTool/Scripts/CreatePlatformExtension.cs
+++ b/Engine/Source/Programs/AutomationTool/Scripts/CreatePlatformExtension.cs
@@ -184,8 +184,10 @@ public class CreatePlatformExtension : BuildCommand
 	/// </summary>
 	public void GeneratePlatformExtension( string InProject, string InSource, string[] InPlatforms )
 	{
+		FileReference ProjectFile = string.IsNullOrEmpty(InProject) ? null : ProjectUtils.FindProjectFileFromName(InProject);
+
 		// Prepare values
-		ProjectDir = string.IsNullOrEmpty(InProject) ? null : new FileReference(InProject).Directory;
+		ProjectDir = ProjectFile?.Directory;
 		GameIni = ConfigCache.ReadHierarchy(ConfigHierarchyType.Game, ProjectDir, BuildHostPlatform.Current.Platform);
 		Platforms = InPlatforms;
 
