Print turnkey google drive download progress.

#rnx

diff --git a/Engine/Source/Programs/AutomationTool/AutomationUtils/GoogleDriveHelper.cs b/Engine/Source/Programs/AutomationTool/AutomationUtils/GoogleDriveHelper.cs
--- a/Engine/Source/Programs/AutomationTool/AutomationUtils/GoogleDriveHelper.cs
+++ b/Engine/Source/Programs/AutomationTool/AutomationUtils/GoogleDriveHelper.cs
@@ -287,11 +287,17 @@ namespace DriveHelper
 				Logger.LogWarning("Attempted to download file that does not exist in Google Drive to {Arg0} with FileID {Arg1}.", DestinationFilePathAndName, FileID);
 				return;
 			}
-			FilesResource.GetRequest Request = Service.Files.Get(FileID);
+
+			var Request = Service.Files.Get(FileID);
 			Request.SupportsTeamDrives = true;
+			Request.Fields = "size";
+			var Response = Request.Execute();
+			long FileSize = (Response.Size != null) ? (long)Response.Size : 0;
 			using (FileStream Stream = System.IO.File.Create(DestinationFilePathAndName + ".tmp"))
 			{
-				Request.MediaDownloader.ProgressChanged += (Google.Apis.Download.IDownloadProgress progress) => HandleDownloadProgressAndCreateFileOnComplete(Stream, progress, DestinationFilePathAndName);
+				DateTime TimeStamp = DateTime.Now;
+				long Percent = 0;
+				Request.MediaDownloader.ProgressChanged += (Google.Apis.Download.IDownloadProgress progress) => HandleDownloadProgressAndCreateFileOnComplete(ref TimeStamp, ref Percent, Stream, FileSize, progress, DestinationFilePathAndName);
 				Console.WriteLine("Beginning download...");
 				Request.Download(Stream);
 			}
@@ -308,22 +314,47 @@ namespace DriveHelper
 				Logger.LogWarning("Attempted to download doc that does not exist in Google Drive to {Arg0} with FileID {Arg1}.", DestinationFilePathAndName, FileID);
 				return;
 			}
-			FilesResource.ExportRequest Request = Service.Files.Export(FileID, ExportFormat.ToMimeString());
+
+			var Request = Service.Files.Get(FileID);
+			Request.SupportsTeamDrives = true;
+			Request.Fields = "size";
+			var Response = Request.Execute();
+			long FileSize = (Response.Size != null) ? (long)Response.Size : 0;
 			using (FileStream Stream = System.IO.File.Create(DestinationFilePathAndName + ".tmp"))
 			{
-				Request.MediaDownloader.ProgressChanged += (Google.Apis.Download.IDownloadProgress progress) => HandleDownloadProgressAndCreateFileOnComplete(Stream, progress, DestinationFilePathAndName);
+				DateTime TimeStamp = DateTime.Now;
+				long Percent = 0;
+				Request.MediaDownloader.ProgressChanged += (Google.Apis.Download.IDownloadProgress Progress) => HandleDownloadProgressAndCreateFileOnComplete(ref TimeStamp, ref Percent, Stream, FileSize, Progress, DestinationFilePathAndName);
 				Console.WriteLine("Beginning download...");
 				Request.Download(Stream);
 			}
 		}
 
 		// Handles a non-resumable download by copying to memory and then creating a new file at the destination
-		protected void HandleDownloadProgressAndCreateFileOnComplete(FileStream FileStream, Google.Apis.Download.IDownloadProgress Progress, string DestinationFilePathAndName)
+		protected void HandleDownloadProgressAndCreateFileOnComplete(ref DateTime TimeStamp, ref long Percent, FileStream FileStream, long FileSize, Google.Apis.Download.IDownloadProgress Progress, string DestinationFilePathAndName)
 		{
 			switch (Progress.Status)
 			{
 				case Google.Apis.Download.DownloadStatus.Downloading:
 					{
+						DateTime Now = DateTime.Now;
+						if ((Now - TimeStamp).TotalSeconds >= 5)
+						{
+							if (FileSize != 0)
+							{
+								long NewPercent = 100 * Progress.BytesDownloaded / FileSize;
+								if (NewPercent != Percent)
+								{
+									Percent = NewPercent;
+									Console.WriteLine($"{Percent}% Downloaded");
+								}
+							}
+							else
+							{
+								Console.WriteLine($"{Progress.BytesDownloaded / 1024}KB Downloaded");
+							}
+							TimeStamp = Now;
+						}
 						break;
 					}
 				case Google.Apis.Download.DownloadStatus.Failed:
