[Gauntlet] Convert any updated IPs into IPv4 addresses as some platform's socket subsystems return IPv6 addresses

diff --git a/Engine/Source/Programs/AutomationTool/Gauntlet/Framework/RpcFramework/Gauntlet.RpcTargetRegistry.cs b/Engine/Source/Programs/AutomationTool/Gauntlet/Framework/RpcFramework/Gauntlet.RpcTargetRegistry.cs
--- a/Engine/Source/Programs/AutomationTool/Gauntlet/Framework/RpcFramework/Gauntlet.RpcTargetRegistry.cs
+++ b/Engine/Source/Programs/AutomationTool/Gauntlet/Framework/RpcFramework/Gauntlet.RpcTargetRegistry.cs
@@ -306,7 +306,8 @@ namespace Gauntlet
 				CallResponse.ResponseBody = "UpdateTargetIp: Request requires both newip and target values in body";
 				return;
 			}
-			else if (IPAddress.Parse(IncomingMessage["newip"]) == null)
+			
+			if (!IPAddress.TryParse(IncomingMessage["newip"], out IPAddress NewIp))
 			{
 				CallResponse.ResponseCode = (int)HttpStatusCode.BadRequest;
 				CallResponse.ResponseBody = "UpdateTargetIp: newip is invalid";
@@ -321,7 +322,9 @@ namespace Gauntlet
 				return;
 			}
 
-			DesiredTarget.IpAddress = IncomingMessage["newip"];
+			// Some platform subsocket systems return an IPv6 address that is mapped to an IPv4.
+			// The http client requires special formatting of the uri's to handle these, so instead just convert them to a plain IPv4 address
+			DesiredTarget.IpAddress = NewIp.MapToIPv4().ToString();
 			CallResponse.ResponseBody = String.Format("UpdateTargetIp: Target {0} IP updated to {1}", DesiredTarget.TargetName, DesiredTarget.IpAddress);
 			Console.WriteLine(CallResponse.ResponseBody);
 			Console.ForegroundColor = PrevForegroundColor;
