Only sign windows binaries with sha256

#rnx

diff --git a/Engine/Source/Programs/AutomationTool/AutomationUtils/CommandUtils.cs b/Engine/Source/Programs/AutomationTool/AutomationUtils/CommandUtils.cs
--- a/Engine/Source/Programs/AutomationTool/AutomationUtils/CommandUtils.cs
+++ b/Engine/Source/Programs/AutomationTool/AutomationUtils/CommandUtils.cs
@@ -3515,8 +3515,7 @@ namespace AutomationTool
 					FinalFiles.Add(new FileReference(TargetFileInfo));
 				}
 			}
-			CodeSignWindows.Sign(FinalFiles, CodeSignWindows.SignatureType.SHA1, !FinalFiles.Any(x => x.HasExtension(".msi")), Description: Description, bRunInParallel); // By default we append signatures, but disable if there are MSI files in the list
-			CodeSignWindows.Sign(FinalFiles.Where(x => !x.HasExtension(".msi")).ToList(), CodeSignWindows.SignatureType.SHA256, Description: Description, RunInParallel: bRunInParallel); // MSI files can only have one signature; prefer SHA1 for compatibility
+			CodeSignWindows.Sign(FinalFiles, CodeSignWindows.SignatureType.SHA256, Description: Description, RunInParallel: bRunInParallel); // By default we append signatures, but disable if there are MSI files in the list
 		}
 	}
 
