Fixing secondary signing spot to be 256 only

#rnx

diff --git a/Engine/Source/Programs/AutomationTool/AutomationUtils/CommandUtils.cs b/Engine/Source/Programs/AutomationTool/AutomationUtils/CommandUtils.cs
--- a/Engine/Source/Programs/AutomationTool/AutomationUtils/CommandUtils.cs
+++ b/Engine/Source/Programs/AutomationTool/AutomationUtils/CommandUtils.cs
@@ -3318,12 +3318,8 @@ namespace AutomationTool
 
 			TargetFileInfo.IsReadOnly = false;
 
-			CodeSignWindows.Sign(new FileReference(TargetFileInfo), CodeSignWindows.SignatureType.SHA1, Description: Description);
-			// MSI files can only have one signature; prefer SHA1 for compatibility, so don't run SHA256 on msi files.
-			if (!TargetFileInfo.FullName.EndsWith(".msi", StringComparison.InvariantCultureIgnoreCase))
-			{
-				CodeSignWindows.Sign(new FileReference(TargetFileInfo), CodeSignWindows.SignatureType.SHA256, Description: Description);
-			}
+			CodeSignWindows.Sign(new FileReference(TargetFileInfo), CodeSignWindows.SignatureType.SHA256, Description: Description);
+			
 		}
 
 		/// <summary>
@@ -3515,7 +3511,7 @@ namespace AutomationTool
 					FinalFiles.Add(new FileReference(TargetFileInfo));
 				}
 			}
-			CodeSignWindows.Sign(FinalFiles, CodeSignWindows.SignatureType.SHA256, Description: Description, RunInParallel: bRunInParallel); // By default we append signatures, but disable if there are MSI files in the list
+			CodeSignWindows.Sign(FinalFiles, CodeSignWindows.SignatureType.SHA256, Description: Description, RunInParallel: bRunInParallel)
 		}
 	}
 
