Fixing compile error

#rnx

diff --git a/Engine/Source/Programs/AutomationTool/AutomationUtils/CommandUtils.cs b/Engine/Source/Programs/AutomationTool/AutomationUtils/CommandUtils.cs
--- a/Engine/Source/Programs/AutomationTool/AutomationUtils/CommandUtils.cs
+++ b/Engine/Source/Programs/AutomationTool/AutomationUtils/CommandUtils.cs
@@ -3319,7 +3319,6 @@ namespace AutomationTool
 			TargetFileInfo.IsReadOnly = false;
 
 			CodeSignWindows.Sign(new FileReference(TargetFileInfo), CodeSignWindows.SignatureType.SHA256, Description: Description);
-			
 		}
 
 		/// <summary>
@@ -3511,7 +3510,7 @@ namespace AutomationTool
 					FinalFiles.Add(new FileReference(TargetFileInfo));
 				}
 			}
-			CodeSignWindows.Sign(FinalFiles, CodeSignWindows.SignatureType.SHA256, Description: Description, RunInParallel: bRunInParallel)
+			CodeSignWindows.Sign(FinalFiles, CodeSignWindows.SignatureType.SHA256, Description: Description, RunInParallel: bRunInParallel);
 		}
 	}
 
