Enable iad on android.

diff --git a/Engine/Source/Programs/AutomationTool/Scripts/CopyBuildToStagingDirectory.Automation.cs b/Engine/Source/Programs/AutomationTool/Scripts/CopyBuildToStagingDirectory.Automation.cs
--- a/Engine/Source/Programs/AutomationTool/Scripts/CopyBuildToStagingDirectory.Automation.cs
+++ b/Engine/Source/Programs/AutomationTool/Scripts/CopyBuildToStagingDirectory.Automation.cs
@@ -2759,7 +2759,7 @@ namespace AutomationScripts
 					continue;
 				}
 
-				//LogInformation("Building PakFileRules for Section {0}", SectionName);
+				LogInformation("Building PakFileRules for Section {0}", SectionName);
 
 				bool bOnlyChunkedBuilds = false;
 				if (!PakRulesConfig.TryGetValue(SectionName, "bOnlyChunkedBuilds", out bOnlyChunkedBuilds))
@@ -2900,6 +2900,8 @@ namespace AutomationScripts
 					PakRules.Filter = new FileFilter();
 					PakRules.ExactFileMatches = new HashSet<string>();
 
+					Logger.LogDebug("Using pakFileRule {0} number of overriden files {1}", PakRules.Name, FilesEnumberable.Count);
+
 					if (PakRules.bUseExactFilePathMatch)
 					{
 						PakRules.ExactFileMatches.UnionWith(FilesEnumberable);
@@ -2911,6 +2913,10 @@ namespace AutomationScripts
 
 					RulesList.Add(PakRules);
 				}
+				else 
+				{
+					Logger.LogDebug("No files found for pakfilerule {0}", PakRules.Name);
+				}
 			}
 
 			if (Params.ApplyIoStoreOnDemand)
