Fix for issues with the FE memory capture tests
- Improved some error handling and logging
- Moved ApplyParamsAndDefaults to the constructor for CsvProfilerConfiguration
- Added IOS

diff --git a/Engine/Source/Programs/AutomationTool/Gauntlet/Framework/RpcFramework/Gauntlet.RpcExecutor.cs b/Engine/Source/Programs/AutomationTool/Gauntlet/Framework/RpcFramework/Gauntlet.RpcExecutor.cs
--- a/Engine/Source/Programs/AutomationTool/Gauntlet/Framework/RpcFramework/Gauntlet.RpcExecutor.cs
+++ b/Engine/Source/Programs/AutomationTool/Gauntlet/Framework/RpcFramework/Gauntlet.RpcExecutor.cs
@@ -786,7 +786,7 @@ namespace Gauntlet
 								RpcName, ArgDesc.Name, ArgDesc.Type));
 						continue;
 					}
-					bool CastSucceeded = true;
+					bool CastSucceeded = false;
 					switch (ArgDesc.Type.ToLower())
 					{
 						case "bool":
@@ -838,7 +838,7 @@ namespace Gauntlet
 			}
 			catch (Exception Error)
 			{
-				string Message = $"CallRpc failed with following exception: {Error}";
+				string Message = $"CallRpc for: {RpcName} failed with following exception: {Error}";
 
 				Log.SuspendECErrorParsing();
 				if (InTarget.IsInvalidated())
@@ -902,7 +902,7 @@ namespace Gauntlet
         	}
         	catch (Exception Error)
         	{
-				string Message = $"CallRpc failed with following exception: {Error}";
+				string Message = $"CallRpc for: {RpcName} failed with following exception: {Error}";
 
 				Log.SuspendECErrorParsing();
 				if (InTarget.IsInvalidated())
