Remove logging introduced in previous change.

diff --git a/Engine/Source/Programs/AutomationTool/Scripts/CopyBuildToStagingDirectory.Automation.cs b/Engine/Source/Programs/AutomationTool/Scripts/CopyBuildToStagingDirectory.Automation.cs
--- a/Engine/Source/Programs/AutomationTool/Scripts/CopyBuildToStagingDirectory.Automation.cs
+++ b/Engine/Source/Programs/AutomationTool/Scripts/CopyBuildToStagingDirectory.Automation.cs
@@ -2759,7 +2759,7 @@ namespace AutomationScripts
 					continue;
 				}
 
-				LogInformation("Building PakFileRules for Section {0}", SectionName);
+				// LogInformation("Building PakFileRules for Section {0}", SectionName);
 
 				bool bOnlyChunkedBuilds = false;
 				if (!PakRulesConfig.TryGetValue(SectionName, "bOnlyChunkedBuilds", out bOnlyChunkedBuilds))
