Added error handling if rpclist cannot be obtained during UpdateAvailableRpcs.
Updated the ResponseReader to utilize a using statement and reduced the amount of time the stream was open.

diff --git a/Engine/Source/Programs/AutomationTool/Gauntlet/Framework/RpcFramework/Gauntlet.RpcTargetRegistry.cs b/Engine/Source/Programs/AutomationTool/Gauntlet/Framework/RpcFramework/Gauntlet.RpcTargetRegistry.cs
--- a/Engine/Source/Programs/AutomationTool/Gauntlet/Framework/RpcFramework/Gauntlet.RpcTargetRegistry.cs
+++ b/Engine/Source/Programs/AutomationTool/Gauntlet/Framework/RpcFramework/Gauntlet.RpcTargetRegistry.cs
@@ -71,10 +71,25 @@ namespace Gauntlet
 			{
 				using (HttpResponseMessage RpcResponse = RpcExecutor.ExecutorClient.Send(ListRpcRequest))
 				{
-					StreamReader ResponseReader = new StreamReader(RpcResponse.Content.ReadAsStream());
-					string ResponseString = ResponseReader.ReadToEnd();
+					string ResponseString = "";
+					using (StreamReader ResponseReader = new StreamReader(RpcResponse.Content.ReadAsStream()))
+					{
+						ResponseString = ResponseReader.ReadToEnd();
+						ResponseReader.Close();
+					}
+					if(ResponseString.Contains("errorCode") || ResponseString.Contains("errorMessage"))
+					{
+						if(ResponseString.Contains("route_handler_not_found"))
+						{
+							Log.Error($"Failed to call {StatusUrl} - The RPC was not identified by the target \"{TargetName}\".");
+						}
+						else
+						{
+							Log.Error($"Failed to call {StatusUrl} (\"{TargetName}\") - Reason unknown. Please check the following Response String and update the error handler: {ResponseString}");
+						}
+						return false;
+					}
 					AvailableRPCs = JsonConvert.DeserializeObject<List<RpcDefinition>>(ResponseString);
-					ResponseReader.Close();
 				}
 				
 				Log.Verbose($"{TargetName} received new list of {AvailableRPCs.Count} available RPCs.");
