[Automation] Fix warning RPC CheatCommand unable to cast arg command value "cheat mms pause" to defined type string

#rnx

diff --git a/Engine/Source/Programs/AutomationTool/Gauntlet/Framework/RpcFramework/Gauntlet.RpcExecutor.cs b/Engine/Source/Programs/AutomationTool/Gauntlet/Framework/RpcFramework/Gauntlet.RpcExecutor.cs
--- a/Engine/Source/Programs/AutomationTool/Gauntlet/Framework/RpcFramework/Gauntlet.RpcExecutor.cs
+++ b/Engine/Source/Programs/AutomationTool/Gauntlet/Framework/RpcFramework/Gauntlet.RpcExecutor.cs
@@ -809,10 +809,15 @@ namespace Gauntlet
 								CastSucceeded = ArgValue is double;
 								break;
 							}
+						case "string":
+							{
+								CastSucceeded = ArgValue is string;
+								break;
+							}
 						default:
 							{
 								// If listed object type is something not handled in this switch, either add a new handler or we basically just
-								// accept that it is a var type that is not going to be automatically linted. Or it's a string which serializing will automatically handle.
+								// accept that it is a var type that is not going to be automatically linted.
 								break;
 							}
 					}
