Move LowLevelTests into AutomationScripts output directory

#rnx

diff --git a/Engine/Source/Programs/AutomationTool/LowLevelTests/LowLevelTests.Automation.csproj b/Engine/Source/Programs/AutomationTool/LowLevelTests/LowLevelTests.Automation.csproj
--- a/Engine/Source/Programs/AutomationTool/LowLevelTests/LowLevelTests.Automation.csproj
+++ b/Engine/Source/Programs/AutomationTool/LowLevelTests/LowLevelTests.Automation.csproj
@@ -9,7 +9,7 @@
     <AssemblyName>LowLevelTests.Automation</AssemblyName>
     <GenerateAssemblyInfo>false</GenerateAssemblyInfo>
     <GenerateTargetFrameworkAttribute>false</GenerateTargetFrameworkAttribute>
-    <OutputPath>..\..\..\..\Binaries\DotNET\AutomationTool\LowLevelTests</OutputPath>
+    <OutputPath>..\..\..\..\Binaries\DotNET\AutomationTool\AutomationScripts\LowLevelTests</OutputPath>
     <Configurations>Debug;Release;Development</Configurations>
     <TreatWarningsAsErrors>true</TreatWarningsAsErrors>
     <AppendTargetFrameworkToOutputPath>false</AppendTargetFrameworkToOutputPath>
