Improve AIBotSoak, make it wait for the server to be ready for matchmaking.

diff --git a/Engine/Source/Programs/AutomationTool/Gauntlet/Framework/RpcFramework/Gauntlet.RpcExecutor.cs b/Engine/Source/Programs/AutomationTool/Gauntlet/Framework/RpcFramework/Gauntlet.RpcExecutor.cs
--- a/Engine/Source/Programs/AutomationTool/Gauntlet/Framework/RpcFramework/Gauntlet.RpcExecutor.cs
+++ b/Engine/Source/Programs/AutomationTool/Gauntlet/Framework/RpcFramework/Gauntlet.RpcExecutor.cs
@@ -365,6 +365,11 @@ namespace Gauntlet
 			}
 		}
 
+		public void ClearMessageDelegates()
+		{
+			MessageDelegates.Clear();
+		}
+
 		/// <summary>
 		/// Checks a single message for matches against a list of message filters to see if it matches any of them.
 		/// </summary>
diff --git a/Engine/Source/Programs/AutomationTool/Gauntlet/Framework/RpcFramework/Gauntlet.RpcTargetRegistry.cs b/Engine/Source/Programs/AutomationTool/Gauntlet/Framework/RpcFramework/Gauntlet.RpcTargetRegistry.cs
--- a/Engine/Source/Programs/AutomationTool/Gauntlet/Framework/RpcFramework/Gauntlet.RpcTargetRegistry.cs
+++ b/Engine/Source/Programs/AutomationTool/Gauntlet/Framework/RpcFramework/Gauntlet.RpcTargetRegistry.cs
@@ -579,7 +579,13 @@ namespace Gauntlet
 
 		private void ReleaseUnmanagedResources()
 		{
-			RpcExecutor.Instance?.ShutdownListenThread();
+			if (RpcExecutor.Instance != null)
+			{
+				Log.Verbose("Cleaning RPC executor instance");
+				RpcExecutor.Instance.ShutdownListenThread();
+				RpcExecutor.Instance.ClearIncomingMessageQueue();
+				RpcExecutor.Instance.ClearMessageDelegates();
+			}
 			RpcTargets.Clear();
 			
 			EditorTarget = null;
