Cloud Artifacts - Don't assume the directory has valid files by default

An empty directory isn't a valid set of files, so default based on file count

#rnx

diff --git a/Engine/Source/Programs/AutomationTool/BuildGraph/Tasks/CreateCloudArtifactTask.cs b/Engine/Source/Programs/AutomationTool/BuildGraph/Tasks/CreateCloudArtifactTask.cs
--- a/Engine/Source/Programs/AutomationTool/BuildGraph/Tasks/CreateCloudArtifactTask.cs
+++ b/Engine/Source/Programs/AutomationTool/BuildGraph/Tasks/CreateCloudArtifactTask.cs
@@ -440,15 +440,22 @@ namespace AutomationTool.Tasks
 			DirectoryReference baseDir = ResolveDirectory(_parameters.BaseDir);
 			List<FileReference> files = BgTaskImpl.ResolveFilespec(baseDir, _parameters.Files, tagNameToFileSet).ToList();
 
-			bool validFiles = true;
-			foreach (FileReference file in files)
+			bool validFiles = files.Count != 0;
+			if (files.Count != 0)
 			{
-				if (!file.IsUnderDirectory(baseDir))
+				foreach (FileReference file in files)
 				{
-					Logger.LogError("Artifact file {File} is not under {BaseDir}", file, baseDir);
-					validFiles = false;
+					if (!file.IsUnderDirectory(baseDir))
+					{
+						Logger.LogError("Artifact file {File} is not under {BaseDir}", file, baseDir);
+						validFiles = false;
+					}
 				}
 			}
+			else
+			{
+				Logger.LogWarning("No files found under {BaseDir}", baseDir);
+			}
 
 			if (!validFiles)
 			{
