[TargetDeviceIOS] Remove logging of connection issues during initial check for connected devices.

diff --git a/Engine/Source/Programs/AutomationTool/Gauntlet/Platform/IOS/Gauntlet.TargetDeviceIOS.cs b/Engine/Source/Programs/AutomationTool/Gauntlet/Platform/IOS/Gauntlet.TargetDeviceIOS.cs
--- a/Engine/Source/Programs/AutomationTool/Gauntlet/Platform/IOS/Gauntlet.TargetDeviceIOS.cs
+++ b/Engine/Source/Programs/AutomationTool/Gauntlet/Platform/IOS/Gauntlet.TargetDeviceIOS.cs
@@ -1056,7 +1056,7 @@ namespace Gauntlet
 					string OSVersionString = JsonDevice.GetObjectField("deviceProperties").GetStringField("osVersionNumber");
 					int IOSMajorVersion = int.Parse(OSVersionString.Substring(0, OSVersionString.IndexOf('.')));
 
-					bool bIsConnected = ValidateDeviceConnection(JsonDevice, true);
+					bool bIsConnected = ValidateDeviceConnection(JsonDevice);
 
 					ConnectedDevice Device = new ConnectedDevice(UUID, IOSMajorVersion, bIsConnected);
 					Devices.Add(UUID, Device);
