Delimit xrcun arguments from the client application arguments using Apple provided unducumented -- delimiter expected by xcrun. Failing to delimit the arguments was causing xcrun to process applications meant to be forwarded to target application and failing if any arguments met the -*l* or -*j* pattern.

diff --git a/Engine/Source/Programs/AutomationTool/IOS/IOSPlatform.Automation.cs b/Engine/Source/Programs/AutomationTool/IOS/IOSPlatform.Automation.cs
--- a/Engine/Source/Programs/AutomationTool/IOS/IOSPlatform.Automation.cs
+++ b/Engine/Source/Programs/AutomationTool/IOS/IOSPlatform.Automation.cs
@@ -2084,7 +2084,7 @@ Console.WriteLine("target receipt executable path: {0}", InExecutablePath);
 				Arguments += " --terminate-existing";	// if it's already running on device, kill it first
 				Arguments += " --console";	// attach to the console so we can get log output
 				Arguments += " --device " + Params.DeviceNames[0];
-				Arguments += " \"" + BundleIdentifier + "\"";
+				Arguments += " \"" + BundleIdentifier + "\" --";
 			}
 			else
 			{
