Fixed gauntlet not propagating ResX/ResY parameters when explicitly set

#rnx
PR #12730

diff --git a/Engine/Source/Programs/AutomationTool/Gauntlet/Unreal/Base/Gauntlet.UnrealTestConfiguration.cs b/Engine/Source/Programs/AutomationTool/Gauntlet/Unreal/Base/Gauntlet.UnrealTestConfiguration.cs
--- a/Engine/Source/Programs/AutomationTool/Gauntlet/Unreal/Base/Gauntlet.UnrealTestConfiguration.cs
+++ b/Engine/Source/Programs/AutomationTool/Gauntlet/Unreal/Base/Gauntlet.UnrealTestConfiguration.cs
@@ -617,17 +617,17 @@ namespace Gauntlet
 		public bool IgnoreDefaultResolutionAndWindowMode { get; set; }
 
 		/// <summary>
-		/// The width resolution. Default resolution is 1280x720.
+		/// The width resolution.
 		/// </summary>
 		/// 
-		[AutoParam(1280)]
+		[AutoParam(0)]
 		public int ResX { get; set; }
 
 		/// <summary>
-		/// The height resolution. Default resolution is 1280x720.
+		/// The height resolution.
 		/// </summary>
 		/// 
-		[AutoParam(720)]
+		[AutoParam(0)]
 		public int ResY { get; set; }
 
 		/// <summary>
@@ -968,10 +968,16 @@ namespace Gauntlet
 				{
 					if (!IgnoreDefaultResolutionAndWindowMode)
 					{
-						if (Globals.Params.ParseValues("resx").Count() == 0)
+						if (ResX != 0)
+						{
+							AppConfig.CommandLine += String.Format(" -ResX={0} ", ResX);
+						}
+
+						if (ResY != 0)
 						{
-							AppConfig.CommandLine += String.Format(" -ResX={0} -ResY={1}", ResX, ResY);
+							AppConfig.CommandLine += String.Format(" -ResY={0}", ResY);
 						}
+
 						if (WindowMode == EWindowMode.Windowed || Windowed)
 						{
 							AppConfig.CommandLine += " -windowed";
