Windows Code Signing tweaks - restore SHA1 for MSI Files

We can sign with sha1, but need to request a SHA256 digest from the timestamp servers

#rnx

diff --git a/Engine/Source/Programs/AutomationTool/AutomationUtils/CommandUtils.cs b/Engine/Source/Programs/AutomationTool/AutomationUtils/CommandUtils.cs
--- a/Engine/Source/Programs/AutomationTool/AutomationUtils/CommandUtils.cs
+++ b/Engine/Source/Programs/AutomationTool/AutomationUtils/CommandUtils.cs
@@ -3088,7 +3088,7 @@ namespace AutomationTool
 				StringBuilder FileList = new StringBuilder();
 				if (SignatureType == SignatureType.SHA1)
 				{
-					CommandLine.Append($"sign{SpecificStoreArg} /a /fd sha1 /td sha1 /n \"{SigningIdentity}\" /v {MultipleSignatureArg} {DescriptionArg}");
+					CommandLine.Append($"sign{SpecificStoreArg} /a /fd sha1 /td sha256 /n \"{SigningIdentity}\" /v {MultipleSignatureArg} {DescriptionArg}");
 				}
 				else if (SignatureType == SignatureType.SHA256)
 				{
@@ -3318,7 +3318,15 @@ namespace AutomationTool
 
 			TargetFileInfo.IsReadOnly = false;
 
-			CodeSignWindows.Sign(new FileReference(TargetFileInfo), CodeSignWindows.SignatureType.SHA256, Description: Description);
+			// MSI files can only have one signature; prefer SHA1 for compatibility, so don't run SHA256 on msi files.
+			if (!TargetFileInfo.FullName.EndsWith(".msi", StringComparison.InvariantCultureIgnoreCase))
+			{
+				CodeSignWindows.Sign(new FileReference(TargetFileInfo), CodeSignWindows.SignatureType.SHA256, Description: Description);
+			}
+			else
+			{
+				CodeSignWindows.Sign(new FileReference(TargetFileInfo), CodeSignWindows.SignatureType.SHA1, Description: Description);
+			}
 		}
 
 		/// <summary>
@@ -3510,7 +3518,17 @@ namespace AutomationTool
 					FinalFiles.Add(new FileReference(TargetFileInfo));
 				}
 			}
-			CodeSignWindows.Sign(FinalFiles, CodeSignWindows.SignatureType.SHA256, Description: Description, RunInParallel: bRunInParallel);
+
+			if (FinalFiles.Any(x => x.HasExtension(".msi")))
+			{
+				// If an MSI is in the list fallback to SHA1
+				CodeSignWindows.Sign(FinalFiles, CodeSignWindows.SignatureType.SHA1, Description: Description, RunInParallel: bRunInParallel);
+			}
+			else
+			{
+				CodeSignWindows.Sign(FinalFiles, CodeSignWindows.SignatureType.SHA256, Description: Description, RunInParallel: bRunInParallel);
+			}
+
 		}
 	}
 
