[vk] Content pipeline changes to support encryption of UGC content

- added ProjectParam DLC overrides for crypto settings (will be a copy of the shipped settings plus an additional key) and the crypto key for UnrealPak to use

diff --git a/Engine/Source/Programs/AutomationTool/AutomationUtils/ProjectParams.cs b/Engine/Source/Programs/AutomationTool/AutomationUtils/ProjectParams.cs
--- a/Engine/Source/Programs/AutomationTool/AutomationUtils/ProjectParams.cs
+++ b/Engine/Source/Programs/AutomationTool/AutomationUtils/ProjectParams.cs
@@ -316,6 +316,8 @@ namespace AutomationTool
 			this.DLCPakPluginFile = InParams.DLCPakPluginFile;
 			this.DLCOverrideCookedSubDir = InParams.DLCOverrideCookedSubDir;
 			this.DLCOverrideStagedSubDir = InParams.DLCOverrideStagedSubDir;
+			this.DLCOverrideCryptoSettings = InParams.DLCOverrideCryptoSettings;
+			this.DLCOverrideEncryptionKeyGuid = InParams.DLCOverrideEncryptionKeyGuid;
 			this.DiffCookedContentPath = InParams.DiffCookedContentPath;
             this.AdditionalBuildOptions = InParams.AdditionalBuildOptions;
             this.AdditionalCookerOptions = InParams.AdditionalCookerOptions;
@@ -538,6 +540,8 @@ namespace AutomationTool
             string DLCName = null,
 			string DLCOverrideCookedSubDir = null,
 			string DLCOverrideStagedSubDir = null,
+			FileReference DLCOverrideCryptoSettings = null,
+			string DLCOverrideEncryptionKeyGuid = null,
 			string DiffCookedContentPath = null,
             bool? DLCIncludeEngineContent = null,
 			bool? DLCPakPluginFile = null,
@@ -1899,7 +1903,7 @@ namespace AutomationTool
 		public string ReferenceContainerChangesCSVFileName;
 
 		/// <summary>
-		/// Stage: Path to the crypto.json file to use for decrypting ReferenceContainerFlobalFileName, if needed.
+		/// Stage: Path to the crypto.json file to use for decrypting ReferenceContainerGlobalFileName, if needed.
 		/// </summary>
 		public string ReferenceContainerCryptoKeys;
 
@@ -1954,6 +1958,16 @@ namespace AutomationTool
 		/// </summary>
 		public string DLCOverrideStagedSubDir;
 
+		/// <summary>
+		/// Specifies keys to use rather than those generated from ini files when using a staging manifest (ignored when using chunk manifests)
+		/// </summary>
+		public EncryptionAndSigning.CryptoSettings DLCOverrideCryptoSettings;
+
+		/// <summary>
+		/// Specifies an encryption key for UnrealPak to use when using a staging manifest (not a chunk manifest)
+		/// </summary>
+		public string DLCOverrideEncryptionKeyGuid;
+
 		/// <summary>
 		/// After cook completes diff the cooked content against another cooked content directory.
 		///  report all errors to the log
diff --git a/Engine/Source/Programs/AutomationTool/Scripts/CopyBuildToStagingDirectory.Automation.cs b/Engine/Source/Programs/AutomationTool/Scripts/CopyBuildToStagingDirectory.Automation.cs
--- a/Engine/Source/Programs/AutomationTool/Scripts/CopyBuildToStagingDirectory.Automation.cs
+++ b/Engine/Source/Programs/AutomationTool/Scripts/CopyBuildToStagingDirectory.Automation.cs
@@ -3080,13 +3080,15 @@ namespace AutomationScripts
 				UnrealPakResponseFile.Remove(FileToRemove);
 			}
 
-			EncryptionAndSigning.CryptoSettings PakCryptoSettings = EncryptionAndSigning.ParseCryptoSettings(DirectoryReference.FromFile(Params.RawProjectPath), SC.StageTargetPlatform.IniPlatformType, Log.Logger);
-			FileReference CryptoKeysCacheFilename = FileReference.Combine(SC.MetadataDir, "Crypto.json");
-			PakCryptoSettings.Save(CryptoKeysCacheFilename);
+			EncryptionAndSigning.CryptoSettings PakCryptoSettings = Params.DLCOverrideCryptoSettings;
+			if (PakCryptoSettings == null)
+			{
+			 	PakCryptoSettings = EncryptionAndSigning.ParseCryptoSettings(DirectoryReference.FromFile(Params.RawProjectPath), SC.StageTargetPlatform.IniPlatformType, Log.Logger);
+			}
 
 			List<CreatePakParams> PakInputs = new List<CreatePakParams>();
-			PakInputs.Add(new CreatePakParams(SC.ShortProjectName, UnrealPakResponseFile, Params.Compressed, Params.RehydrateAssets, null));
-			CreatePaks(Params, SC, PakInputs, PakCryptoSettings, CryptoKeysCacheFilename);
+			PakInputs.Add(new CreatePakParams(SC.ShortProjectName, UnrealPakResponseFile, Params.Compressed, Params.RehydrateAssets, Params.DLCOverrideEncryptionKeyGuid));
+			CreatePaks(Params, SC, PakInputs, PakCryptoSettings);
 		}
 
 		private static void PrepareZenWorkspaceEnv(ProjectParams Params, DeploymentContext SC)
@@ -3435,8 +3437,7 @@ namespace AutomationScripts
 		/// <param name="SC"></param>
 		/// <param name="PakParamsList"></param>
 		/// <param name="CryptoSettings"></param>
-		/// <param name="CryptoKeysCacheFilename"></param>
-		private static void CreatePaks(ProjectParams Params, DeploymentContext SC, List<CreatePakParams> PakParamsList, EncryptionAndSigning.CryptoSettings CryptoSettings, FileReference CryptoKeysCacheFilename)
+		private static void CreatePaks(ProjectParams Params, DeploymentContext SC, List<CreatePakParams> PakParamsList, EncryptionAndSigning.CryptoSettings CryptoSettings)
 		{
 			bool? bCustomDeploymentNeedsDiffPak = SC.CustomDeployment?.GetPlatformPatchesWithDiffPak(Params, SC);
 			bool bShouldCreateIoStoreContainerFiles = ShouldCreateIoStoreContainerFiles(Params, SC);
@@ -4164,6 +4165,11 @@ namespace AutomationScripts
 				Outputs.Add(Tuple.Create(OutputLocation, OutputRelativeLocation, PatchSourceContentPath));
 			}
 
+			// save out the crypto settings to the Metadata folder; UnrealPak/IasTool uses it immediately, but
+			// there seem to be other things in the build that potentially rely on a copy living there
+			FileReference CryptoKeysCacheFilename = FileReference.Combine(SC.MetadataDir, "Crypto.json");
+			CryptoSettings.Save(CryptoKeysCacheFilename);
+
 			// Actually execute UnrealPak
 			if (PakCommands.Count > 0)
 			{
@@ -4403,10 +4409,9 @@ namespace AutomationScripts
 						UploadArgs.Append(" " + Params.AdditionalPakOptions);
 					}
 
-					FileReference CryptoKeysFilename = FileReference.Combine(SC.MetadataDir, "Crypto.json");
-					if (FileExists_NoExceptions(CryptoKeysFilename.FullName))
+					if (CryptoKeysCacheFilename != null)
 					{
-						UploadArgs.Append(string.Format(" -CryptoKeys=\"{0}\"", CryptoKeysFilename.FullName));
+						UploadArgs.Append(string.Format(" -CryptoKeys=\"{0}\"", CryptoKeysCacheFilename.FullName));
 					}
 				}
 				else
@@ -4959,8 +4964,6 @@ namespace AutomationScripts
 
 			// Parse and cache crypto settings from INI file
 			EncryptionAndSigning.CryptoSettings PakCryptoSettings = EncryptionAndSigning.ParseCryptoSettings(DirectoryReference.FromFile(Params.RawProjectPath), SC.StageTargetPlatform.IniPlatformType, Log.Logger);
-			FileReference CryptoKeysCacheFilename = FileReference.Combine(SC.MetadataDir, "Crypto.json");
-			PakCryptoSettings.Save(CryptoKeysCacheFilename);
 
 			// We still want to have a list of all files to stage. We will use the chunk manifests
 			// to put the files from staging manifest into the right chunk
@@ -5354,7 +5357,7 @@ namespace AutomationScripts
 				}
 			}
 
-			CreatePaks(Params, SC, PakInputs, PakCryptoSettings, CryptoKeysCacheFilename);
+			CreatePaks(Params, SC, PakInputs, PakCryptoSettings);
 
 			String ChunkLayerFilename = CombinePaths(GetChunkManifestPath(Params, SC), GetChunkPakLayerListName());
 			String OutputChunkLayerFilename = Path.Combine(SC.ProjectRoot.FullName, "Build", SC.FinalCookPlatform, "ChunkLayerInfo", GetChunkPakLayerListName());
