Windows Signing: Fixing multi signing errors

#rnx

diff --git a/Engine/Source/Programs/AutomationTool/AutomationUtils/CommandUtils.cs b/Engine/Source/Programs/AutomationTool/AutomationUtils/CommandUtils.cs
--- a/Engine/Source/Programs/AutomationTool/AutomationUtils/CommandUtils.cs
+++ b/Engine/Source/Programs/AutomationTool/AutomationUtils/CommandUtils.cs
@@ -3325,7 +3325,7 @@ namespace AutomationTool
 			}
 			else
 			{
-				CodeSignWindows.Sign(new FileReference(TargetFileInfo), CodeSignWindows.SignatureType.SHA1, Description: Description);
+				CodeSignWindows.Sign(new FileReference(TargetFileInfo), CodeSignWindows.SignatureType.SHA1, AllowMultipleSignatures: false, Description: Description);
 			}
 		}
 
@@ -3522,7 +3522,7 @@ namespace AutomationTool
 			if (FinalFiles.Any(x => x.HasExtension(".msi")))
 			{
 				// If an MSI is in the list fallback to SHA1
-				CodeSignWindows.Sign(FinalFiles, CodeSignWindows.SignatureType.SHA1, Description: Description, RunInParallel: bRunInParallel);
+				CodeSignWindows.Sign(FinalFiles, CodeSignWindows.SignatureType.SHA1, AllowMultipleSignatures: false, Description: Description, RunInParallel: bRunInParallel);
 			}
 			else
 			{
