LaunchMultiServer support for custom config files

You can now specify "-CustomConfig=<string>" to enable an override config layer

diff --git a/Engine/Source/Programs/AutomationTool/Scripts/LaunchMultiServer.Automation.cs b/Engine/Source/Programs/AutomationTool/Scripts/LaunchMultiServer.Automation.cs
--- a/Engine/Source/Programs/AutomationTool/Scripts/LaunchMultiServer.Automation.cs
+++ b/Engine/Source/Programs/AutomationTool/Scripts/LaunchMultiServer.Automation.cs
@@ -270,7 +270,13 @@ namespace AutomationTool
 			var StartTime = DateTime.UtcNow;
 
 			// Parse server configuration from ini files
-			ConfigHierarchy ProjectGameConfig = ConfigCache.ReadHierarchy(ConfigHierarchyType.Game, DirectoryReference.FromFile(ProjectPath), UnrealTargetPlatform.Win64);
+			string CustomConfig = "";
+			if (ParseParam("CustomConfig"))
+			{
+				CustomConfig = ParseParamValue("CustomConfig");
+			}
+
+			ConfigHierarchy ProjectGameConfig = ConfigCache.ReadHierarchy(ConfigHierarchyType.Game, DirectoryReference.FromFile(ProjectPath), UnrealTargetPlatform.Win64, CustomConfig);
 
 			const String ServerDefConfigSection = "/Script/MultiServerConfiguration.MultiServerSettings";
 			const String ProxyConfigSection = "/Script/MultiServerConfiguration.Proxy";
@@ -299,7 +305,14 @@ namespace AutomationTool
 					}
 				}
 
-				StartProcessesForServerInfos(UniqueServerInfos, ProcessLaunchType.Server, " -MultiServerPeers=" + PeerList);
+				string AdditionalParameters = "";
+				if (!string.IsNullOrEmpty(CustomConfig))
+				{
+					AdditionalParameters += "-CustomConfig=" + CustomConfig + " ";
+				}
+
+				string AdditionalServerParameters = AdditionalParameters + "-MultiServerPeers=" + PeerList;
+				StartProcessesForServerInfos(UniqueServerInfos, ProcessLaunchType.Server, AdditionalServerParameters);
 
 				if (ParseParam("client"))
 				{
@@ -310,7 +323,7 @@ namespace AutomationTool
 						if (ClientValue == "direct")
 						{
 							Logger.LogInformation("Starting direct client instances connecting to MultiServer instances");
-							StartProcessesForServerInfos(UniqueServerInfos, ProcessLaunchType.DirectClient, "");
+							StartProcessesForServerInfos(UniqueServerInfos, ProcessLaunchType.DirectClient, AdditionalParameters);
 						}
 						else if (ClientValue == "proxy")
 						{
@@ -348,7 +361,7 @@ namespace AutomationTool
 							{
 								Logger.LogInformation("Starting proxy server on port {0} connected to MultiServer instances", ProxyListeningPort);
 								{
-									var ProxyServerAdditionalParams = String.Format(" -port={0} -ConsoleTitle=\"{1} Proxy Server {2}\" -SessionName=\"{1} Proxy Server {2}\" -log={1}ProxyServer-{2}.log ", ProxyListeningPort, ProjectName, ProxyServerIndex + 1);
+									var ProxyServerAdditionalParams = AdditionalParameters + String.Format("-port={0} -ConsoleTitle=\"{1} Proxy Server {2}\" -SessionName=\"{1} Proxy Server {2}\" -log={1}ProxyServer-{2}.log ", ProxyListeningPort, ProjectName, ProxyServerIndex + 1);
 									StartProcessesForServerInfos(UniqueServerInfos, ProcessLaunchType.ProxyServer, ProxyServerAdditionalParams);
 								}
 
@@ -358,7 +371,7 @@ namespace AutomationTool
 								for (int ClientIndex = 0; ClientIndex < NumClientsToSpawn; ClientIndex++)
 								{
 									++TotalClientCount;
-									var ProxyClientAdditionalParams = String.Format("127.0.0.1:{0} -ConsoleTitle=\"{1} Proxy Client {2}\" -SessionName=\"{1} Proxy Client {2}\" -log={1}ProxyClient-{2}.log -SaveWinPos={2} ", ProxyListeningPort, ProjectName, TotalClientCount);
+									var ProxyClientAdditionalParams = AdditionalParameters + String.Format("127.0.0.1:{0} -ConsoleTitle=\"{1} Proxy Client {2}\" -SessionName=\"{1} Proxy Client {2}\" -log={1}ProxyClient-{2}.log -SaveWinPos={2} ", ProxyListeningPort, ProjectName, TotalClientCount);
 					
 									StartProcessesForServerInfos(UniqueServerInfos, ProcessLaunchType.ProxyClient, ProxyClientAdditionalParams);
 								}
@@ -369,11 +382,11 @@ namespace AutomationTool
 								for (int BotIndex = 0; BotIndex < NumBotsToSpawn; BotIndex++)
 								{
 									++TotalBotCount;
-									var BotClientAdditionalParams = String.Format("127.0.0.1:{0} -log={2}ProxyBot-{1}.log -log -nullrhi -bot -nosound -unattended -nosplash ", ProxyListeningPort, TotalBotCount, ProjectName);
+									var BotClientAdditionalParams = AdditionalParameters + String.Format("127.0.0.1:{0} -log={2}ProxyBot-{1}.log -log -nullrhi -bot -nosound -unattended -nosplash ", ProxyListeningPort, TotalBotCount, ProjectName);
 									BotClientAdditionalParams += String.Format("-ConsoleTitle=\"{0} Proxy Bot Client {1}\" -SessionName=\"{0} Proxy Bot Client {1}\" ", ProjectName, TotalBotCount);
 									if (!ParseParam("nonewconsole"))
 									{
-										BotClientAdditionalParams += " -newconsole ";
+										BotClientAdditionalParams += "-newconsole ";
 									}
 									StartProcessesForServerInfos(UniqueServerInfos, ProcessLaunchType.ProxyClient, BotClientAdditionalParams);
 								}
