Retry for timeout is not considering new start time

diff --git a/Engine/Source/Programs/AutomationTool/Gauntlet/Framework/Gauntlet.TestExecutor.cs b/Engine/Source/Programs/AutomationTool/Gauntlet/Framework/Gauntlet.TestExecutor.cs
--- a/Engine/Source/Programs/AutomationTool/Gauntlet/Framework/Gauntlet.TestExecutor.cs
+++ b/Engine/Source/Programs/AutomationTool/Gauntlet/Framework/Gauntlet.TestExecutor.cs
@@ -8,7 +8,6 @@ using System.Net;
 using System.Text.RegularExpressions;
 using System.Threading;
 using System.Threading.Tasks;
-using AutomationTool;
 using EpicGames.Core;
 using Gauntlet.Utils;
 
@@ -106,7 +105,7 @@ namespace Gauntlet
 
 		private List<TestExecutionInfo> RunningTests;
 
-		private DateTime StartTime;
+		private DateTime CurrentTestStartTime;
 
 		private int CurrentTestPass;
 
@@ -159,7 +158,6 @@ namespace Gauntlet
 			Globals.AbortHandlers.Add(CancelHandler);
 			Globals.PostAbortHandlers.Add(PostCancelHandler);
 
-			StartTime = DateTime.Now;
 			CurrentTestPass = 0;
 
 			IsRunning = true;
@@ -376,7 +374,7 @@ namespace Gauntlet
 						// invalid = no result yet
 						if (Result == TestResult.Invalid)
 						{
-							TimeSpan RunningTime = DateTime.Now - TestInfo.TimeSetupEnded;
+							TimeSpan RunningTime = DateTime.Now - CurrentTestStartTime;
 
 							if ((SecondsRunning % 60) == 0)
 							{
@@ -590,15 +588,15 @@ namespace Gauntlet
 		{
 			string Name = TestInfo.TestNode.Name;
 
-			Log.Info("Starting Test {Name}", TestInfo);
+			Log.Info("Starting Test {TestName}", TestInfo);
 
 			try
 			{
 				TestInfo.TimeSetupBegan = DateTime.Now;
 				if (TestInfo.TestNode.StartTest(Pass, NumPasses))
 				{
-					TestInfo.TimeSetupEnded = DateTime.Now;
-					Log.Info("Launched test {Name} at {Time}", Name, TestInfo.TimeSetupEnded.ToString("h:mm:ss"));
+					CurrentTestStartTime = TestInfo.TimeSetupEnded = DateTime.Now;
+					Log.Info("Launched test {TestName} at {StartTime}", Name, TestInfo.TimeSetupEnded.ToString("h:mm:ss"));
 					return true;
 				}
 			}
@@ -758,7 +756,7 @@ namespace Gauntlet
 			// Give the test a chance to update itself
 			try
 			{
-				TestInfo.TestNode.TickTest();					
+				TestInfo.TestNode.TickTest();
 			}
 			catch (Exception Ex)
 			{
@@ -769,7 +767,7 @@ namespace Gauntlet
 			// Does the test still say it's running?
 			bool TestIsRunning = TestInfo.TestNode.GetTestStatus() == TestStatus.InProgress;
 
-			TimeSpan RunningTime = DateTime.Now - TestInfo.TimeSetupEnded;
+			TimeSpan RunningTime = DateTime.Now - CurrentTestStartTime;
 
 			if (TestIsRunning && RunningTime.TotalSeconds > TestInfo.TestNode.MaxDuration && !Options.NoTimeout)
 			{
@@ -777,6 +775,7 @@ namespace Gauntlet
 				{
 					TestInfo.CancellationReason = string.Format("Terminating Test {0} due to maximum duration of {1} seconds. ", TestInfo.TestNode, TestInfo.TestNode.MaxDuration);
 					TestInfo.FinalResult = TestResult.TimedOut;
+					Log.Verbose("Terminating current test pass {TestName} due to maximum duration of {TestMaxDuration} seconds.", TestInfo.TestNode, TestInfo.TestNode.MaxDuration);
 				}
 				else if (TestInfo.TestNode.MaxDurationReachedResult == EMaxDurationReachedResult.Success)
 				{
@@ -827,18 +826,18 @@ namespace Gauntlet
 					bool bCanFinalizeTest = true;
 					if (TestInfo.FinalResult == TestResult.WantRetry)
 					{
-						Log.Info("{Name} requested retry. Cleaning up old test and relaunching", TestInfo);
-
-						DateTime OriginalStartTime = TestInfo.TimeSetupEnded;
+						Log.Info("Test {TestName} requested retry. Cleaning up old test and relaunching", TestInfo);
 
 						bool bIsRestarted = TestInfo.TestNode.RestartTest();
 						if (bIsRestarted)
 						{
 							// Mark us as still running
+							CurrentTestStartTime = DateTime.Now;
 							TestInfo.CancellationReason = "";
 							TestInfo.TestNode.SetCancellationReason("");
 							TestInfo.FinalResult = TestResult.Invalid;
 							bCanFinalizeTest = false;
+							Log.Info("Restarted test {TestName} for requested retry at {StartTime}", TestInfo, CurrentTestStartTime.ToString("h:mm:ss"));
 						}
 						else
 						{
