Add an ability to filter by timers by thread and make sure UnrealInsights can run on GPU-less machines.

diff --git a/Engine/Source/Programs/AutomationTool/Scripts/ExportTimerStatisticsFromUtrace.cs b/Engine/Source/Programs/AutomationTool/Scripts/ExportTimerStatisticsFromUtrace.cs
--- a/Engine/Source/Programs/AutomationTool/Scripts/ExportTimerStatisticsFromUtrace.cs
+++ b/Engine/Source/Programs/AutomationTool/Scripts/ExportTimerStatisticsFromUtrace.cs
@@ -17,6 +17,7 @@ namespace ExportTimerStatisticsFromUtrace.Automation
 	[Help("Runs UnrealInsights to export timers from a utrace file to CSV file.")]
 	[ParamHelp("TraceFile", "Input UnrealInsights .utrace file to open", ParamType = typeof(string), Required = true)]
 	[ParamHelp("CSVFile", "Output CSV file to write", ParamType = typeof(string), Required = true)]
+	[ParamHelp("Threads", "Threads to export", ParamType = typeof(string), Required = false)]
 	[ParamHelp("TimerRegion", "TimerRegion to export", ParamType = typeof(string), Required = true)]
 	[ParamHelp("MaxTimerCount", "Number of top timers to export", ParamType = typeof(string), Required = true)]
 	class ExportTimerStatisticsFromUtrace : BuildCommand
@@ -28,11 +29,19 @@ namespace ExportTimerStatisticsFromUtrace.Automation
 			string TimerRegion = ParseRequiredStringParam("TimerRegion");
 			string NumTimers = ParseRequiredStringParam("MaxTimerCount");
 
-			string Arguments = string.Format("-OpenTraceFile=\"{0}\" -unattended -autoquit -noui -ExecOnAnalysisCompleteCmd=\"TimingInsights.ExportTimerStatistics {1} -region={2} -maxtimercount={3} -sortby=totalinclusivetime -sortorder=descending\"", 
+			string Threads = ParseOptionalStringParam("Threads");
+			string ThreadArgument = "";
+			if (!string.IsNullOrWhiteSpace(Threads))
+			{
+				ThreadArgument = string.Format(" -threads={0}", Threads);
+			}
+
+			string Arguments = string.Format("-OpenTraceFile=\"{0}\" -unattended -autoquit -noui -nullrhi -ExecOnAnalysisCompleteCmd=\"TimingInsights.ExportTimerStatistics {1} -region={2} -maxtimercount={3} {4} -sortby=totalinclusivetime -sortorder=descending\"", 
 				TraceFile,
 				CSVFile,
 				TimerRegion,
-				NumTimers
+				NumTimers,
+				ThreadArgument
 				);
 			Logger.LogInformation("About to run Insights with arguments '{Arguments}'", Arguments);
 			int ReturnCode = RunInsights(Arguments);
