Changing the MultiServer start-up code to add option -MultiServerLocalHost which derives the ServerID from the -Port argument.
This will help start-up a multi-server cluster using existing infrastructure.  The Proxy now also allows to connect to non-localhost.  The scripts have been updated to document all such features.

#rnx

diff --git a/Engine/Plugins/Runtime/MultiServerReplication/Source/MultiServerConfiguration/Public/MultiServerSettings.h b/Engine/Plugins/Runtime/MultiServerReplication/Source/MultiServerConfiguration/Public/MultiServerSettings.h
--- a/Engine/Plugins/Runtime/MultiServerReplication/Source/MultiServerConfiguration/Public/MultiServerSettings.h
+++ b/Engine/Plugins/Runtime/MultiServerReplication/Source/MultiServerConfiguration/Public/MultiServerSettings.h
@@ -3,42 +3,21 @@
 #pragma once
 
 #include "UObject/Object.h"
-#include "Containers/Array.h"
 #include "MultiServerSettings.generated.h"
 
-class FString;
-
-/** Parameters for a single server process in a multi-server setup. */
-USTRUCT()
-struct FMultiServerDefinition
-{
-	GENERATED_BODY()
-
-	/** User-specified identifier that must be unique for each server in a MultiServer cluster. */
-	UPROPERTY(Config)
-	FString LocalId;
-
-	/** UDP port on which the MultiServer beacon will listen for new connections. */
-	UPROPERTY(Config)
-	uint16 ListenPort = 0;
-};
-
 /**
  * Settings that control how a multi-server setup is configured and launched.
  * For example, to define 4 servers in a Game ini file it might look like this:
  * [/Script/MultiServerConfiguration.MultiServerSettings]
- * !ServerDefinitions=ClearArray
- * +ServerDefinitions=(LocalId="0",ListenPort=17000)
- * +ServerDefinitions=(LocalId="1",ListenPort=17001)
- * +ServerDefinitions=(LocalId="2",ListenPort=17002)
- * +ServerDefinitions=(LocalId="3",ListenPort=17003)
+ * TotalNumServers=4
  */
 UCLASS(Config=Game)
 class MULTISERVERCONFIGURATION_API UMultiServerSettings : public UObject
 {
 	GENERATED_BODY()
 
-public:
+private:
+	// This is read from the LaunchMultiServer script (not code). We take the runtime value from -MultiServerNumServers parameter.
 	UPROPERTY(Config)
-	TArray<FMultiServerDefinition> ServerDefinitions;
+	int32 TotalNumServers = 0;
 };
diff --git a/Engine/Plugins/Runtime/MultiServerReplication/Source/MultiServerReplication/Private/MultiServerNode.cpp b/Engine/Plugins/Runtime/MultiServerReplication/Source/MultiServerReplication/Private/MultiServerNode.cpp
--- a/Engine/Plugins/Runtime/MultiServerReplication/Source/MultiServerReplication/Private/MultiServerNode.cpp
+++ b/Engine/Plugins/Runtime/MultiServerReplication/Source/MultiServerReplication/Private/MultiServerNode.cpp
@@ -23,6 +23,12 @@ void UMultiServerNode::ParseCommandLineIntoCreateParams(FMultiServerNodeCreatePa
 	FParse::Value(FCommandLine::Get(), TEXT("MultiServerPeers="), PeerAddressesString, false);
 
 	PeerAddressesString.ParseIntoArray(InOutParams.PeerAddresses, TEXT(","), true);
+
+	FParse::Value(FCommandLine::Get(), TEXT("MultiServerNumServers="), InOutParams.NumServers);
+	if (InOutParams.NumServers < 1)
+	{
+		InOutParams.NumServers = InOutParams.PeerAddresses.Num();
+	}
 }
 
 UMultiServerNode::UMultiServerNode()
diff --git a/Engine/Plugins/Runtime/MultiServerReplication/Source/MultiServerReplication/Private/MultiServerProxy.cpp b/Engine/Plugins/Runtime/MultiServerReplication/Source/MultiServerReplication/Private/MultiServerProxy.cpp
--- a/Engine/Plugins/Runtime/MultiServerReplication/Source/MultiServerReplication/Private/MultiServerProxy.cpp
+++ b/Engine/Plugins/Runtime/MultiServerReplication/Source/MultiServerReplication/Private/MultiServerProxy.cpp
@@ -142,7 +142,7 @@ void UProxyListenerNotify::ConnectToGameServer(UNetConnection* ProxyConnection,
 		GameServerConnectionState->NetDriver->SetWorld(ProxyNetDriver->GetWorld());
 		Cast<UProxyBackendNetDriver>(GameServerConnectionState->NetDriver)->SetProxyNetDriver(ProxyNetDriver);
 
-		UE_LOG(LogNetProxy, Log, TEXT("Created a game server NetDriver (name=%s, port=%d)"), *GameServerConnectionState->NetDriver->GetName(), GameServerConnectionState->Port);
+		UE_LOG(LogNetProxy, Log, TEXT("Created a Proxy Game Server NetDriver (name=%s, url=%s)"), *GameServerConnectionState->NetDriver->GetName(), *GameServerConnectionState->GameServerURL.ToString());
 	}
 
 	// Add a player to use on the game server.
@@ -161,9 +161,7 @@ void UProxyListenerNotify::ConnectToGameServer(UNetConnection* ProxyConnection,
 	const bool bIsFirstGameServerConnection = (GameServerConnection == nullptr);
 	if (bIsFirstGameServerConnection)
 	{
-		const int32 GameServerPort = GameServerConnectionState->Port;
-		FString URLStr = FString::Format(TEXT("127.0.0.1:{0}"), { GameServerPort });
-		FURL URL(nullptr, *URLStr, TRAVEL_Absolute);
+		FURL URL = GameServerConnectionState->GameServerURL;
 		
 		const FString GameLoginOptions = NewPlayer->GetGameLoginOptions();
 		if (!GameLoginOptions.IsEmpty())
@@ -487,13 +485,10 @@ void UProxyBackendNetDriver::InternalProcessRemoteFunction(
 	}
 }
 
-void UProxyNetDriver::RegisterGameServer(int32 Port)
+void UProxyNetDriver::RegisterGameServer(const FURL& GameServerURL)
 {
-	UE_LOG(LogNetProxy, Log, TEXT("Registering proxy game server on port %d"), Port);
-
-	FGameServerConnectionState GameServerConnectionState;
-	GameServerConnectionState.Port = Port;
-	GameServerConnections.Add(GameServerConnectionState);
+	UE_LOG(LogNetProxy, Log, TEXT("Registering Game Server for URL: %s"), *GameServerURL.ToString());
+	GameServerConnections.Emplace(GameServerURL);
 }
 
 bool UProxyNetDriver::IsConnectedToAllGameServers() const
@@ -545,11 +540,11 @@ bool UProxyNetDriver::InitBase(bool bInitAsClient, FNetworkNotify* InNotify, con
 
 		for (const FString& Address : Addresses)
 		{
-			FString IPAddressStr, PortStr;
-			Address.Split(TEXT(":"), &IPAddressStr, &PortStr);
-
-			int32 Port = FCString::Atoi(ToCStr(PortStr));
-			RegisterGameServer(Port);
+			FURL GameServerURL { nullptr, *Address, ETravelType::TRAVEL_Absolute };
+			if (ensureMsgf(GameServerURL.Valid, TEXT("Attempting to register an invalid ProxyGameServer URL. Input: %s. Parsed URL: %s"), *Address, *GameServerURL.ToString()))
+			{
+				RegisterGameServer(GameServerURL);
+			}
 		}
 	}
 
diff --git a/Engine/Plugins/Runtime/MultiServerReplication/Source/MultiServerReplication/Public/MultiServerNode.h b/Engine/Plugins/Runtime/MultiServerReplication/Source/MultiServerReplication/Public/MultiServerNode.h
--- a/Engine/Plugins/Runtime/MultiServerReplication/Source/MultiServerReplication/Public/MultiServerNode.h
+++ b/Engine/Plugins/Runtime/MultiServerReplication/Source/MultiServerReplication/Public/MultiServerNode.h
@@ -30,6 +30,9 @@ struct FMultiServerNodeCreateParams
 	/** The port on which this node will listen for new connections. */ 
 	uint16 ListenPort = 0;
 
+	/** The number of servers we're expecting in the cluster */
+	uint32 NumServers = 0;
+
 	/** List of addresses of other nodes to attempt to connect to. */
 	TArray<FString> PeerAddresses;
 
diff --git a/Engine/Plugins/Runtime/MultiServerReplication/Source/MultiServerReplication/Public/MultiServerProxy.h b/Engine/Plugins/Runtime/MultiServerReplication/Source/MultiServerReplication/Public/MultiServerProxy.h
--- a/Engine/Plugins/Runtime/MultiServerReplication/Source/MultiServerReplication/Public/MultiServerProxy.h
+++ b/Engine/Plugins/Runtime/MultiServerReplication/Source/MultiServerReplication/Public/MultiServerProxy.h
@@ -39,7 +39,14 @@ struct FGameServerConnectionState
 {
 	GENERATED_BODY()
 
-	int32 Port;
+	FGameServerConnectionState(const FURL& InGameServerURL)
+	: GameServerURL(InGameServerURL)
+	{
+	}
+
+	/** The URL address of the Game Server (does not indicate a live connection, only intent) */
+	UPROPERTY()
+	FURL GameServerURL;
 
 	UPROPERTY()
 	TObjectPtr<class UWorld> World;
@@ -57,6 +64,11 @@ struct FGameServerConnectionState
 
 	UPROPERTY()
 	TArray<FGameServerSplitJoinRequest> PendingSplitJoinRequests;
+
+private:
+	// We are only valid with a URL, but a default constructor is needed for serialization
+	friend class UScriptStruct;
+	FGameServerConnectionState() = default;
 };
 
 USTRUCT()
@@ -261,7 +273,7 @@ class MULTISERVERREPLICATION_API UProxyNetDriver : public UIpNetDriver
 public:
 	
 	/** Register a game server. */
-	void RegisterGameServer(int32 Port);
+	void RegisterGameServer(const FURL& GameServerURL);
 
 	/** Return true if all registered servers are connected. */
 	bool IsConnectedToAllGameServers() const;
diff --git a/Engine/Source/Programs/AutomationTool/Scripts/LaunchMultiServer.Automation.cs b/Engine/Source/Programs/AutomationTool/Scripts/LaunchMultiServer.Automation.cs
--- a/Engine/Source/Programs/AutomationTool/Scripts/LaunchMultiServer.Automation.cs
+++ b/Engine/Source/Programs/AutomationTool/Scripts/LaunchMultiServer.Automation.cs
@@ -225,182 +225,121 @@ namespace AutomationTool
 			}
 			return DebuggerAttached;
 		}
-	}
-
-	// Mirrors FMultiServerDefinition from the MultiServerReplication plugin
-	class MultiServerDefinition
-	{
-		public string LocalId = "0";
-		public int ListenPort = 0;
-
-		public override string ToString()
-		{
-			return "MultiServerDefinition: LocalId \"" + LocalId + "\", ListenPort " + ListenPort;
-		}
-	}
-
-	// MultiServerDefinition plus additional information only used by this command
-	class MultiServerInstanceInfo
-	{
-		public MultiServerDefinition Definition = null;
-		public int GameListenPort = 0;
-
-		public MultiServerInstanceInfo(MultiServerDefinition definition)
-		{
-			Definition = definition;
-		}
-	}
-
-	[Help("Launches multiple server processes for a project using the MultiServerReplication plugin.")]
-	[Help("project=<project>", "Project to open. Will search current path and paths in ueprojectdirs.")]
-	[Help("map=<MapName>", "Map to load on startup.")]
+	};
+
+	[Help(@"Launches multiple server processes for a project using the MultiServerReplication plugin.
+
+Example running 2 hosts locally with a single proxy server & client:
+	-cook -notimeouts -numservers=2 -client=proxy -proxyclientcount=1 -proxycycleprimary
+
+Example running 4 clients locally connecting to single proxy server w/ a non-local cluster of 3:
+	-cook -notimeouts -client=proxy -proxyclientcount=4 -proxycycleprimary -nonlocalservers=172.27.63.109:9001,172.27.63.109:9002,172.27.63.109:9003
+")]
+	[Help("Project=<project>", "Project to open. Will search current path and paths in ueprojectdirs. Defaults to current project (in uShell).")]
+	[Help("Map=<MapName>", "Map to load on startup.")]
+	[Help("CustomConfig=<Section>", "Read a different section of the Project's GameConfig for default values")]
+    [Help("NumServers=##", "Set the fleet size of servers")]
+	[Help("NonLocalServers=Address1:Port1,Address2:Port2,etc.", "Specify servers that are not hosted locally but belong to the same fleet")]
+	[Help("NoTimeouts", "Disable timeouts (recommended for debugging)")]
+	[Help("NoConsole", "Disable launching a console (not recommended, but may improve performance)")]
+	[Help("NoNewConsole", "Disable launching the new style console")]
+	[Help("AttachToServers", "Force Visual Studio to attach to the servers when launched (requires VS be launched or it will hang)")]
+	[Help("AttachToClients", "Force Visual Studio to attach to the clients when launched (requires VS be launched or it will hang)")]
+	[Help("LogDebug", "Enable certain helpful Verbose logs for debugging")]
+	[Help("LogDebugV", "Enable certain helpful VeryVerbose logs for debugging")]
+	[Help("Client=<Type>", "Specify the client type to connect with (either PROXY or DIRECT)")]
+	[Help("ProxyServerCount=##", "Specify how many proxy servers to launch in client=proxy mode")]
+	[Help("ProxyClientCount=##", "Specify how many clients to launch PER SERVER in client=proxy mode")]
+	[Help("ProxyBotCount=##", "Specify how many bots to launch PER SERVER in client=proxy mode")]
+	[Help("ProxyClientPrimary=##", "Specify which game server to prefer as the initial proxy connection")]
+	[Help("ProxyCyclePrimary", "Specify that each subsequent proxy client should choose a different server")]
 	public class LaunchMultiServer : BuildCommand, IProjectParamsHelpers
 	{
+		// We will use a convention where ServerId 1 will start at BaseGameListenPort + 1;
+		int BaseGameListenPort = 9000;
+		int ProxyServerBasePort = 11000;
+
 		private enum ProcessLaunchType
 		{
 			Server,
 			DirectClient,
 			ProxyServer,
-			ProxyClient
+			ProxyClient,
+			ProxyBot
 		}
 
 		public override ExitCode Execute()
 		{
 			Logger.LogInformation("********** RUN MULTISERVER COMMAND STARTED **********");
 			var StartTime = DateTime.UtcNow;
+			var ProjectName = ProjectUtils.GetShortProjectName(ProjectPath);
 
 			// Parse server configuration from ini files
-			string CustomConfig = "";
-			if (ParseParam("CustomConfig"))
-			{
-				CustomConfig = ParseParamValue("CustomConfig");
-			}
+			string CustomConfig = ParseParamValue("CustomConfig", string.Empty);
 
 			ConfigHierarchy ProjectGameConfig = ConfigCache.ReadHierarchy(ConfigHierarchyType.Game, DirectoryReference.FromFile(ProjectPath), UnrealTargetPlatform.Win64, CustomConfig);
 
 			const String ServerDefConfigSection = "/Script/MultiServerConfiguration.MultiServerSettings";
 			const String ProxyConfigSection = "/Script/MultiServerConfiguration.Proxy";
 
-			ProjectGameConfig.TryGetValuesGeneric(ServerDefConfigSection, "ServerDefinitions", out MultiServerDefinition[] ConfigServerDefs);
+			if (!ProjectGameConfig.TryGetValue(ProxyConfigSection, "ListenPort", out ProxyServerBasePort))
+			{
+				ProxyServerBasePort = BaseGameListenPort + 2000;
+				Logger.LogWarning($"Could not find {ProxyConfigSection}.ListenPort in Game Config File {ProjectGameConfig} Defaulting to {ProxyServerBasePort}");
+			}
+
+			int TotalNumServers = ParseParamInt("NumServers", -1);
+			if (TotalNumServers < 0 && !ProjectGameConfig.TryGetValue(ServerDefConfigSection, "TotalNumServers", out TotalNumServers))
+			{
+				TotalNumServers = 2;
+				Logger.LogWarning($"Could not find {ServerDefConfigSection}.TotalNumServers in Game Config File {ProjectGameConfig} Defaulting to {TotalNumServers}");
+			}
 
-			if (ConfigServerDefs != null && ConfigServerDefs.Length > 0)
+			// CommonArgs are used for all instances spawned
+			string CommonArgs = (CustomConfig.Length > 0) ? $" -CustomConfig={CustomConfig}" : "";
+
+			// These are servers not on our local host
+			string[] NonLocalServerAddresses = ParseParamValues("NonLocalServers").SelectMany(x => x.Split( new char[]{';', ','}, StringSplitOptions.TrimEntries)).ToArray();
+			if (NonLocalServerAddresses.Length > TotalNumServers)
 			{
-				// Use a dictionary to store ServerDefs by unqiue ID
-				Dictionary<string, MultiServerInstanceInfo> UniqueServerInfos = new Dictionary<string, MultiServerInstanceInfo>();
+				Logger.LogWarning($"Specified more NonLocalServers=({NonLocalServerAddresses.Length}) than NumServers={TotalNumServers}. Inferring that we want to run with no local servers and only non-local servers. Setting NumServers={NonLocalServerAddresses.Length}.");
+				TotalNumServers = NonLocalServerAddresses.Length;
+			}
+			int NumLocalServers = TotalNumServers - NonLocalServerAddresses.Length;
 
-				var GameListenPort = 7777;
+			// Start-up the servers first
+			Logger.LogInformation($"Launching {NumLocalServers} Local Servers for {TotalNumServers} Total Servers");
+			string[] LocalGameServerAddresses = StartProcessesForDedicatedServers(NumLocalServers, NonLocalServerAddresses, CommonArgs);
 
-				// Check for duplicates and build peer list
-				string PeerList = "";
-				for (int i = 0; i < ConfigServerDefs.Length; i++)
-				{
-					if (UniqueServerInfos.TryAdd(ConfigServerDefs[i].LocalId, new MultiServerInstanceInfo(ConfigServerDefs[i])))
-					{
-						PeerList += String.Format("127.0.0.1:{0}{1}", ConfigServerDefs[i].ListenPort, i < (ConfigServerDefs.Length - 1) ? "," : "");
-						UniqueServerInfos[ConfigServerDefs[i].LocalId].GameListenPort = GameListenPort++;
-					}
-					else
-					{
-						Logger.LogInformation("MultiServer server definition with duplicate ID {0}, ignoring duplicate entry {1}", ConfigServerDefs[i].LocalId, ConfigServerDefs[i]);
-					}
-				}
+			// Generate all of the GameServerAddresses by merging the local & non-local servers
+			string[] AllGameServerAddresses = new string[TotalNumServers];
+			LocalGameServerAddresses.CopyTo(AllGameServerAddresses, 0);
+			NonLocalServerAddresses.CopyTo(AllGameServerAddresses, LocalGameServerAddresses.Length);
 
-				string AdditionalParameters = "";
-				if (!string.IsNullOrEmpty(CustomConfig))
+			// Now start-up the clients
+			string[] ClientValues = ParseParamValue("client", string.Empty).Split(',');
+			foreach (string ClientValue in ClientValues)
+			{
+				if (ClientValue == "direct")
 				{
-					AdditionalParameters += "-CustomConfig=" + CustomConfig + " ";
+					Logger.LogInformation("Starting direct client instances connecting to MultiServer instances");
+					StartProcessesForClients(ProcessLaunchType.DirectClient, LocalGameServerAddresses, CommonArgs);
 				}
+				else if (ClientValue == "proxy")
+				{
+					int ProxyServerCount = ParseParamInt("ProxyServerCount", 1);
+					int NumClientsPerServer = ParseParamInt("ProxyClientCount", 1);
+					int NumBotsPerServer = ParseParamInt("ProxyBotCount", 0);
 
-				string AdditionalServerParameters = AdditionalParameters + "-MultiServerPeers=" + PeerList;
-				StartProcessesForServerInfos(UniqueServerInfos, ProcessLaunchType.Server, AdditionalServerParameters);
+					Logger.LogInformation($"Starting {ProxyServerCount} Proxy Servers each with {NumClientsPerServer} Clients and {NumBotsPerServer} Bots");
 
-				if (ParseParam("client"))
-				{
-					string[] ClientValues = ParseParamValue("client").Split(',');
+					string[] LocalProxyAddresses = GenerateLocalServerAddressRange(ProxyServerCount, ProxyServerBasePort);
+					StartProcessesForProxyServer(LocalProxyAddresses, AllGameServerAddresses, CommonArgs);
 
-					foreach (string ClientValue in ClientValues)
-					{
-						if (ClientValue == "direct")
-						{
-							Logger.LogInformation("Starting direct client instances connecting to MultiServer instances");
-							StartProcessesForServerInfos(UniqueServerInfos, ProcessLaunchType.DirectClient, AdditionalParameters);
-						}
-						else if (ClientValue == "proxy")
-						{
-							int ClientCountTotal = 1;
-							if (ParseParam("proxyclientcount"))
-							{
-								string ClientCountStr = ParseParamValue("proxyclientcount", null);
-								Int32.TryParse(ClientCountStr, out ClientCountTotal);
-							}
-
-							int BotCountTotal = 0;
-							if (ParseParam("proxybotcount"))
-							{
-								string ProxyBotCountStr = ParseParamValue("proxybotcount", null);
-								Int32.TryParse(ProxyBotCountStr, out BotCountTotal);
-							}
-
-							int ProxyServerCount = 1;
-							if (ParseParam("proxyservercount"))
-							{
-								string ProxyServerCountStr = ParseParamValue("proxyservercount", null);
-								Int32.TryParse(ProxyServerCountStr, out ProxyServerCount);
-							}
-
-							var FirstProxyListeningPort = 8000;
-							ProjectGameConfig.TryGetValue(ProxyConfigSection, "ListenPort", out FirstProxyListeningPort);
-
-							int TotalClientCount = 0;
-							int TotalBotCount = 0;
-
-							var ProjectName = ProjectUtils.GetShortProjectName(ProjectPath);
-
-							var ProxyListeningPort = FirstProxyListeningPort;
-							for (int ProxyServerIndex = 0; ProxyServerIndex < ProxyServerCount; ++ProxyServerIndex)
-							{
-								Logger.LogInformation("Starting proxy server on port {0} connected to MultiServer instances", ProxyListeningPort);
-								{
-									var ProxyServerAdditionalParams = AdditionalParameters + String.Format("-port={0} -ConsoleTitle=\"{1} Proxy Server {2}\" -SessionName=\"{1} Proxy Server {2}\" -log={1}ProxyServer-{2}.log ", ProxyListeningPort, ProjectName, ProxyServerIndex + 1);
-									StartProcessesForServerInfos(UniqueServerInfos, ProcessLaunchType.ProxyServer, ProxyServerAdditionalParams);
-								}
-
-								int NumClientsPerServer = ClientCountTotal / ProxyServerCount;
-								int NumExtraClients = ClientCountTotal % ProxyServerCount;
-								int NumClientsToSpawn = NumClientsPerServer + (ProxyServerIndex < NumExtraClients ? 1 : 0);
-								for (int ClientIndex = 0; ClientIndex < NumClientsToSpawn; ClientIndex++)
-								{
-									++TotalClientCount;
-									var ProxyClientAdditionalParams = AdditionalParameters + String.Format("127.0.0.1:{0} -ConsoleTitle=\"{1} Proxy Client {2}\" -SessionName=\"{1} Proxy Client {2}\" -log={1}ProxyClient-{2}.log -SaveWinPos={2} ", ProxyListeningPort, ProjectName, TotalClientCount);
-					
-									StartProcessesForServerInfos(UniqueServerInfos, ProcessLaunchType.ProxyClient, ProxyClientAdditionalParams);
-								}
-								
-								int NumBotsPerServer = BotCountTotal / ProxyServerCount;
-								int NumExtraBots = BotCountTotal % ProxyServerCount;
-								int NumBotsToSpawn = NumBotsPerServer + (ProxyServerIndex < NumExtraBots ? 1 : 0);
-								for (int BotIndex = 0; BotIndex < NumBotsToSpawn; BotIndex++)
-								{
-									++TotalBotCount;
-									var BotClientAdditionalParams = AdditionalParameters + String.Format("127.0.0.1:{0} -log={2}ProxyBot-{1}.log -log -nullrhi -bot -nosound -unattended -nosplash ", ProxyListeningPort, TotalBotCount, ProjectName);
-									BotClientAdditionalParams += String.Format("-ConsoleTitle=\"{0} Proxy Bot Client {1}\" -SessionName=\"{0} Proxy Bot Client {1}\" ", ProjectName, TotalBotCount);
-									if (!ParseParam("nonewconsole"))
-									{
-										BotClientAdditionalParams += "-newconsole ";
-									}
-									StartProcessesForServerInfos(UniqueServerInfos, ProcessLaunchType.ProxyClient, BotClientAdditionalParams);
-								}
-								
-								++ProxyListeningPort;
-							}
-						}
-					}
+					StartProcessesForClients(ProcessLaunchType.ProxyClient, LocalProxyAddresses, CommonArgs, NumClientsPerServer);
+					StartProcessesForClients(ProcessLaunchType.ProxyBot, LocalProxyAddresses, $"{CommonArgs} -nullrhi -bot -nosound -unattended -nosplash", NumBotsPerServer);
 				}
 			}
-			else
-			{
-				Logger.LogInformation("No server configs found in {0}, not starting any servers.", ServerDefConfigSection);
-			}
 
 			Logger.LogInformation("Run command time: {0:0.00} s", (DateTime.UtcNow - StartTime).TotalMilliseconds / 1000);
 			Logger.LogInformation("********** RUN MULTISERVER COMMAND COMPLETED **********");
@@ -408,233 +347,190 @@ namespace AutomationTool
 			return ExitCode.Success;
 		}
 
-		private void StartProcessesForServerInfos(Dictionary<string, MultiServerInstanceInfo> ServerInfos, ProcessLaunchType LaunchType, string AdditionalParameters)
+		/**
+		 * Given the number of local servers and the base port, generate the array of address:port for that range of ports
+		 * The address range will be ["127.0.0.1:BasePort+1" ... "127.0.0.1::BasePort+NumLocalServers"]
+		 */
+		string[] GenerateLocalServerAddressRange(int NumLocalServers, int BasePort)
 		{
-			const int PauseBetweenProcessMS = 500;
-
-			var Params = new ProjectParams
-				(
-					Command: this,
-					RawProjectPath: ProjectPath,
-					DedicatedServer: LaunchType == ProcessLaunchType.Server || LaunchType == ProcessLaunchType.ProxyServer,
-					Client: LaunchType == ProcessLaunchType.DirectClient || LaunchType == ProcessLaunchType.ProxyClient
-				);
-
-			var DeployContexts = Project.CreateDeploymentContext(Params, LaunchType == ProcessLaunchType.Server || LaunchType == ProcessLaunchType.ProxyServer);
-
-			if (DeployContexts.Count == 0)
+			string[] LocalServers = new string[NumLocalServers];
+			for (int i = 0; i < NumLocalServers; ++i)
 			{
-				throw new AutomationException("No DeployContexts for launching a process.");
+				// It's important to note that Unreal does not handle "localhost" properly, it needs to be the loopback address
+				LocalServers[i] = $"127.0.0.1:{BasePort + i + 1}";
 			}
 
-			var DeployContext = DeployContexts[0];
+			return LocalServers;
+		}
 
-			var App = CombinePaths(CmdEnv.LocalRoot, "Engine/Binaries/Win64/UnrealEditor.exe");
-			if (Params.Cook)
+		/**
+		 * Starts all of the server processes that should execute on the local host.
+		 * @param NumLocalServers: The number of local servers to spin-up as part of the total server cluster (the total cluster consistent of these local servers + NonLocalServers).
+		 * @param NonLocalServers: Specify servers that are not running on the local machine but contribute to the Metaverse
+		 * @param AdditionalServerParameters: These are extra parameters that the server processes should put on the command-line (shared between all servers)
+		 * @return The Addresses of all of the Game Servers spun-up.
+		 */
+		private string[] StartProcessesForDedicatedServers(int NumLocalServers, string[] NonLocalServers, string AdditionalServerArguments)
+		{
+			if (NumLocalServers < 1)
 			{
-				List<FileReference> Exes = DeployContext.StageTargetPlatform.GetExecutableNames(DeployContext);
-				App = Exes[0].FullName;
+				return new string[] { };
 			}
 
-			var CommonArgs = (Params.Cook ? "" : DeployContext.ProjectArgForCommandLines + " ");
-			if (LaunchType == ProcessLaunchType.Server || LaunchType == ProcessLaunchType.ProxyServer)
-			{
-				CommonArgs += Params.MapToRun + " -server -log -NODEBUGOUTPUT -SuppressConsoleOutputVerboseLogging ";
+			int NumTotalServers = NumLocalServers + NonLocalServers.Length;
+			string[] LocalGameServerAddresses = new string[NumTotalServers];
 
-				// Here are some examples of adding parameters for debugging purposes (note the trailing space within the quotes)
-				// CommonArgs += "-dpcvars=net.MaxConstructedPartialBunchSizeBytes=131071,dstm.SimpleCoinBots=0 ";
-				// CommonArgs += "-logcmds=\"LogActor Verbose, LogRemoteObject Verbose, LogDstmPushMigration VeryVerbose, LogRepCompares Log, LogNet Verbose\" ";
+			// This is optional to the multi-server process, these are in addition to the MultiServerNumPeers argument
+			string ServerCommonArgs = $" -server {AdditionalServerArguments} -NODEBUGOUTPUT -SuppressConsoleOutputVerboseLogging";
+			ServerCommonArgs += ParseParam("notimeouts") ? " -notimeouts" : "";
+			ServerCommonArgs += IsAttachingDebugger(ProcessLaunchType.Server) ? " -WaitForDebuggerNoBreak" : "";
+			ServerCommonArgs += $" -MultiServerNumServers={NumTotalServers}";
+			ServerCommonArgs += (NonLocalServers.Length > 0) ? $" -MultiServerPeers={string.Join(',', NonLocalServers)}" : "";
 
-				if (!ParseParam("nonewconsole"))
-				{
-					CommonArgs += " -newconsole ";
-				}
+			for (int ServerId = 1; ServerId <= NumLocalServers; ++ServerId)
+			{
+				int ServerGamePort = BaseGameListenPort + ServerId;
+				LocalGameServerAddresses[ServerId - 1] = $"127.0.0.1:{ServerGamePort}";
 
-				// Use this option when logging with VeryVerbose to avoid unusable debugging windows (too much spew)
-				// The logging will still go to the output files (and you can use a text editor like Notepad++ to auto-reload them)
-				if (ParseParam("noconsole"))
-				{
-					CommonArgs += " -nodebugoutput -noconsole ";
-				}
-				
-				var Cats = new[] { "LogDstm", "LogDstmObject", "LogDstmDemo", "LogRemoteExec", "LogRemoteObject", "LogDstmPushMigration" };
+				Logger.LogInformation("Starting Dedicated MultiServer for Game Port {0}", ServerGamePort);
 
-				foreach (var Cat in Cats)
-				{
-					if (ParseParam("logdebug")) { CommonArgs += string.Format("-LogCmds=\"{0} Verbose\" ", Cat); }
-					else if (ParseParam("logdebugv")) { CommonArgs += string.Format("-LogCmds=\"{0} VeryVerbose\" ", Cat); }
-				}
+				string ServerArgs = ServerCommonArgs;
+				ServerArgs += String.Format(" -MultiServerLocalHost -port={0} -log=MultiServer-{1}.log", ServerGamePort, ServerId);
 
+				string AppTitle = $"Server ID {ServerId}";
+				ExecuteApp(ProcessLaunchType.Server, AppTitle, ServerArgs);
 			}
 
-			CommonArgs += " -messaging ";
-
-			PushDir(Path.GetDirectoryName(App));
+			return LocalGameServerAddresses;
+		}
 
-			try
+		private void StartProcessesForClients(ProcessLaunchType ClientType, string[] GameServerAddresses, string AdditionalClientArguments, int ClientsPerServer = 1)
+		{
+			for (int Idx= 0; Idx < GameServerAddresses.Length; ++Idx)
 			{
-				if (LaunchType == ProcessLaunchType.DirectClient || LaunchType == ProcessLaunchType.Server)
-				{
-					int ServerNum = 0;
-					foreach (var ServerInfo in ServerInfos.Values)
-					{
-						Logger.LogInformation("Starting MultiServer instance of type {0} for LocalId {1}, beacon port {2}, and game port {3}", LaunchType, ServerInfo.Definition.LocalId, ServerInfo.Definition.ListenPort, ServerInfo.GameListenPort);
-
-						var Args = CommonArgs;
-
-						if (LaunchType == ProcessLaunchType.Server)
-						{
-							Args += String.Format("-MultiServerLocalId={0} -MultiServerListenPort={1} -port={2} -log={3}MultiServer-{0}.log ", ServerInfo.Definition.LocalId, ServerInfo.Definition.ListenPort, ServerInfo.GameListenPort, Params.ShortProjectName);
-							Args += String.Format("-ConsoleTitle=\"{0} Server ID {1}\" -SessionName=\"{0} Server ID {1}\" ", Params.ShortProjectName, ServerInfo.Definition.LocalId);
-							Args += Params.ServerCommandline + " ";
-						}
-						else if (LaunchType == ProcessLaunchType.DirectClient)
-						{
-							Args += String.Format("127.0.0.1:{0} ", ServerInfo.GameListenPort) + (Params.Cook ? "" : "-game ");
-							Args += String.Format(" -log={1}DirectClient-{0}.log", ServerInfo.Definition.LocalId, Params.ShortProjectName);
-							Args += String.Format(" -SessionName=\"{0} Direct Client | Server ID {1}\" ", Params.ShortProjectName, ServerInfo.Definition.LocalId);
-							Args += Params.ClientCommandline + " ";
-
-							// Save a separate window position for each client
-							Args += "-windowed -SaveWinPos=" + (ServerNum + 1) + " ";
-						}
+				string GameServerAddress = GameServerAddresses[Idx];
 
-						if (ParseParam("notimeouts"))
-						{
-							Args += "-notimeouts ";
-						}
+				string ClientArgs = $" -game {AdditionalClientArguments}";
+				ClientArgs += $" {GameServerAddress}";
+				ClientArgs += $" -log={ClientType}-{Idx+1}.log";
+				ClientArgs += $" -windowed -SaveWinPos={Idx+1}";
 
-						if (IsAttachingDebugger(LaunchType))
-						{
-							Args += "-WaitForDebuggerNoBreak ";
-						}
+				// Potentially launch multiple clients per server address
+				for (int InstanceNum = 0; InstanceNum < ClientsPerServer; ++InstanceNum)
+				{
+					string InstanceStr = (ClientsPerServer > 1) ? $" [{InstanceNum + 1}]" : string.Empty;
+					string AppTitle = $"{ClientType} {Idx + 1}{InstanceStr}";
+					ExecuteApp(ClientType, AppTitle, ClientArgs);
+				}
+			}
+		}
 
-						Args += AdditionalParameters;
+		private void StartProcessesForProxyServer(string[] LocalProxyServerAddresses, string[] AllGameServerAddresses, string AdditionalProxyArguments)
+		{
+			string CommonProxyServerArgs = AdditionalProxyArguments;
+			CommonProxyServerArgs += " -NetDriverOverrides=/Script/MultiServerReplication.ProxyNetDriver";
+			CommonProxyServerArgs += ParseParam("notimeouts") ? " -notimeouts" : "";
+			CommonProxyServerArgs += IsAttachingDebugger(ProcessLaunchType.ProxyServer) ? " -WaitForDebuggerNoBreak" : "";
 
-						var NewProcess = Run(App, Args, null, ERunOptions.Default | ERunOptions.NoWaitForExit | ERunOptions.NoStdOutRedirect);
+			if (AllGameServerAddresses.Length > 0)
+			{
+				CommonProxyServerArgs += $" -ProxyGameServers={string.Join(',', AllGameServerAddresses)}";
+			}
 
-						if (NewProcess != null)
-						{
-							// Remove started process so it won't be killed on UAT exit.
-							// Essentially forces the -NoKill command-line option behavior for these.
-							ProcessManager.RemoveProcess(NewProcess);
-						}
+			if (ParseParam("proxyclientprimary"))
+			{
+				CommonProxyServerArgs += $" -ProxyClientPrimaryGameServer={ParseParamInt("proxyclientprimary", 0)}";
+			}
 
-						// Pause between starting processes to enforce startup determinism.
-						System.Threading.Thread.Sleep(PauseBetweenProcessMS);
+			if (ParseParam("proxycycleprimary"))
+			{
+				CommonProxyServerArgs += " -ProxyCyclePrimaryGameServer";
+			}
 
-						if (IsAttachingDebugger(LaunchType))
-						{
-							ProcessDebugger.DebugProcess(NewProcess.ProcessObject.Id, this, Params);
-						}
+			for (int ProxyIdx = 1; ProxyIdx <= LocalProxyServerAddresses.Length; ++ProxyIdx)
+			{
+				// Split off the port
+				string LocalProxyServerAddress = LocalProxyServerAddresses[ProxyIdx - 1];
+				string[] AddressParts = LocalProxyServerAddress.Split(':');
 
-						ServerNum++;
-					}
-				}
-				else if (LaunchType == ProcessLaunchType.ProxyServer)
+				// Parse the port to use for the -port argument
+				int ProxyServerPort;
+				if (!int.TryParse(AddressParts.LastOrDefault(), out ProxyServerPort))
 				{
-					var Args = CommonArgs;
-
-					Args += "-NetDriverOverrides=/Script/MultiServerReplication.ProxyNetDriver ";
-
-					if (ServerInfos.Count() > 0)
-					{
-						// Example of adding additional debug parameters for the proxy
-						// Args += "-logcmds=\"LogNetProxy VeryVerbose\" ";
-
-						if (ParseParam("notimeouts"))
-						{
-							Args += "-notimeouts ";
-						}
+					Logger.LogError($"Error Starting LocalProxyServer with Address {LocalProxyServerAddress} had invalid Port Specifier: {AddressParts.LastOrDefault()}");
+					continue;
+				}
 
-						if (IsAttachingDebugger(LaunchType))
-						{
-							Args += "-WaitForDebuggerNoBreak ";
-						}
+				string ProxyServerArgs = $"{CommonProxyServerArgs} -port={ProxyServerPort}";
 
-						if (ServerInfos.Values.Count() > 0)
-						{
-							Args += "-ProxyGameServers=";
+				string AppTitle = $"ProxyServer {ProxyIdx}";
+				ExecuteApp(ProcessLaunchType.ProxyServer, AppTitle, ProxyServerArgs);
+			}
+		}
 
-							foreach (var ServerInfo in ServerInfos.Values)
-							{
-								Args += String.Format("127.0.0.1:{0},", ServerInfo.GameListenPort);
-							}
 
-							Args += " ";
-						}
+		private void ExecuteApp(ProcessLaunchType LaunchType, string AppTitle, string AdditionalArguments)
+		{
+			const int PauseBetweenProcessMS = 100;
 
-						if (ParseParam("proxyclientprimary"))
-						{
-							int ProxyClientPrimary = 0;
-							string ProxyClientPrimaryStr = ParseParamValue("proxyclientprimary", null);
-							Int32.TryParse(ProxyClientPrimaryStr, out ProxyClientPrimary);
+			var Params = new ProjectParams
+			(
+				Command: this,
+				RawProjectPath: ProjectPath,
+				DedicatedServer: LaunchType == ProcessLaunchType.Server || LaunchType == ProcessLaunchType.ProxyServer,
+				Client: LaunchType == ProcessLaunchType.DirectClient || LaunchType == ProcessLaunchType.ProxyClient || LaunchType == ProcessLaunchType.ProxyBot
+			);
+
+			var DeployContexts = Project.CreateDeploymentContext(Params, Params.DedicatedServer);
+			if (DeployContexts.Count == 0)
+			{
+				throw new AutomationException("No DeployContexts for launching a process.");
+			}
 
-							Args += String.Format("-ProxyClientPrimaryGameServer={0} ", ProxyClientPrimary);
-						}
+			var DeployContext = DeployContexts[0];
 
-						if (ParseParam("proxycycleprimary"))
-						{
-							Args += "-ProxyCyclePrimaryGameServer ";
-						}
-					}
+			var App = CombinePaths(CmdEnv.LocalRoot, "Engine/Binaries/Win64/UnrealEditor.exe");
+			if (Params.Cook)
+			{
+				List<FileReference> Exes = DeployContext.StageTargetPlatform.GetExecutableNames(DeployContext);
+				App = Exes[0].FullName;
+			}
 
-					Args += Params.ServerCommandline + " ";
-					Args += AdditionalParameters;
+			string Args = Params.Client ? Params.ClientCommandline : $"{Params.MapToRun} {Params.ServerCommandline}";
+			Args += " -messaging";
+			Args += (Params.Cook ? "" : " " + DeployContext.ProjectArgForCommandLines);
 
-					var NewProcess = Run(App, Args, null, ERunOptions.Default | ERunOptions.NoWaitForExit | ERunOptions.NoStdOutRedirect);
+			// Use this option when logging with VeryVerbose to avoid unusable debugging windows (too much spew)
+			// The logging will still go to the output files (and you can use a text editor like Notepad++ to auto-reload them)
+			Args += ParseParam("noconsole") ? " -nodebugoutput -noconsole" : "";
+			Args += ParseParam("nonewconsole") ? "" : " -newconsole -log ";
 
-					if (NewProcess != null)
-					{
-						// Remove started process so it won't be killed on UAT exit.
-						// Essentially forces the -NoKill command-line option behavior for these.
-						ProcessManager.RemoveProcess(NewProcess);
-					}
+			// Add any logging in the commands
+			var Cats = new[] { "LogDstm", "LogDstmObject", "LogDstmDemo", "LogRemoteExec", "LogRemoteObject", "LogDstmPushMigration" };
+			Args += ParseParam("logdebug")	? $" -LogCmds=\"{string.Join("=Verbose,",	Cats)}\"" : "";
+			Args += ParseParam("logdebugv") ? $" -LogCmds=\"{string.Join("=VeryVerbose,",Cats)}\"" : "";
 
-					// Pause between starting processes to enforce startup determinism.
-					System.Threading.Thread.Sleep(PauseBetweenProcessMS);
+			Args += String.Format(" -ConsoleTitle=\"{0} {1}\" -SessionName=\"{0} {1}\"", Params.ShortProjectName, AppTitle);
+			Args += AdditionalArguments;
+			Args = Args.Trim();
 
-					if (IsAttachingDebugger(LaunchType))
-					{
-						ProcessDebugger.DebugProcess(NewProcess.ProcessObject.Id, this, Params);
-					}
-				}
-				else if (LaunchType == ProcessLaunchType.ProxyClient)
+			PushDir(Path.GetDirectoryName(App));
+			try
+			{
+				var NewProcess = Run(App, Args, null, ERunOptions.Default | ERunOptions.NoWaitForExit | ERunOptions.NoStdOutRedirect);
+				if (NewProcess != null)
 				{
-					var Args = CommonArgs;
-
-					// The server to connect to should be provided in the AdditionalParameters (e.g. 127.0.0.1:8000).
-					Args += AdditionalParameters + " ";
-					Args += (Params.Cook ? "" : "-game ");
-					Args += Params.ClientCommandline + " ";
-					Args += "-windowed ";
-					
-					if (ParseParam("notimeouts"))
-					{
-						Args += "-notimeouts ";
-					}
-
-					if (IsAttachingDebugger(LaunchType))
-					{
-						Args += "-WaitForDebuggerNoBreak ";
-					}
-
-					var NewProcess = Run(App, Args, null, ERunOptions.Default | ERunOptions.NoWaitForExit | ERunOptions.NoStdOutRedirect);
-
-					if (NewProcess != null)
-					{
-						// Remove started process so it won't be killed on UAT exit.
-						// Essentially forces the -NoKill command-line option behavior for these.
-						ProcessManager.RemoveProcess(NewProcess);
-					}
+					// Remove started process so it won't be killed on UAT exit.
+					// Essentially forces the -NoKill command-line option behavior for these.
+					ProcessManager.RemoveProcess(NewProcess);
+				}
 
-					// Pause between starting processes to enforce startup determinism.
-					System.Threading.Thread.Sleep(PauseBetweenProcessMS);
+				// Pause between starting processes to enforce startup determinism.
+				System.Threading.Thread.Sleep(PauseBetweenProcessMS);
 
-					if (IsAttachingDebugger(LaunchType))
-					{
-						ProcessDebugger.DebugProcess(NewProcess.ProcessObject.Id, this, Params);
-					}
+				if (IsAttachingDebugger(LaunchType))
+				{
+					ProcessDebugger.DebugProcess(NewProcess.ProcessObject.Id, this, Params);
 				}
 			}
 			catch
diff --git a/Engine/Source/Runtime/CoreUObject/Private/UObject/RemoteObject.cpp b/Engine/Source/Runtime/CoreUObject/Private/UObject/RemoteObject.cpp
--- a/Engine/Source/Runtime/CoreUObject/Private/UObject/RemoteObject.cpp
+++ b/Engine/Source/Runtime/CoreUObject/Private/UObject/RemoteObject.cpp
@@ -142,11 +142,23 @@ FRemoteObjectMaps* ObjectMaps = nullptr;
 void InitServerId()
 {
 	FString ServerId;
-	if (!FParse::Value(FCommandLine::Get(), TEXT("MultiServerLocalId="), ServerId))
+	const TCHAR* CommandLine = FCommandLine::Get();
+	if (!FParse::Value(CommandLine, TEXT("MultiServerLocalId="), ServerId))
 	{
-		if (!FParse::Value(FCommandLine::Get(), TEXT("LocalPeerId="), ServerId))
+		if (!FParse::Value(CommandLine, TEXT("LocalPeerId="), ServerId))
 		{
-			RemoteServerId = FRemoteServerId(ERemoteServerIdConstants::Invalid);
+			int ListenPort = 0;
+			if (FParse::Param(CommandLine, TEXT("MultiServerLocalHost")) && FParse::Value(CommandLine, TEXT("Port="), ListenPort))
+			{
+				if (ListenPort > 0)
+				{
+					ServerId = FString::FromInt(ListenPort % 1000);
+				}
+			}
+			else
+			{
+				RemoteServerId = FRemoteServerId(ERemoteServerIdConstants::Invalid);
+			}
 		}
 	}
 	if (!ServerId.IsEmpty())
