Make the Gauntlet controller try to find HOME when USERPROFILE does not exist
- USERPROFILE is only valid on Windows
- HOME should be valid on *NIX and Mac platforms

diff --git a/Engine/Source/Programs/AutomationTool/Gauntlet/Platform/Android/Gauntlet.TargetDeviceAndroid.cs b/Engine/Source/Programs/AutomationTool/Gauntlet/Platform/Android/Gauntlet.TargetDeviceAndroid.cs
--- a/Engine/Source/Programs/AutomationTool/Gauntlet/Platform/Android/Gauntlet.TargetDeviceAndroid.cs
+++ b/Engine/Source/Programs/AutomationTool/Gauntlet/Platform/Android/Gauntlet.TargetDeviceAndroid.cs
@@ -1876,7 +1876,13 @@ namespace Gauntlet
 
 		public static void RestoreBackupKeys()
 		{
-			string LocalKeyPath = Path.Combine(Environment.GetEnvironmentVariable("USERPROFILE"), ".android");
+			string UserDirectory = Environment.GetEnvironmentVariable("USERPROFILE");
+			// Depending on OS, USERPROFILE may be not valid so attempt to look at the HOME variable instead
+			if (UserDirectory == null)
+			{
+				UserDirectory = Path.Combine(Environment.GetEnvironmentVariable("HOME"));
+			} 
+			string LocalKeyPath = Path.Combine(UserDirectory, ".android");
 			string LocalKeyFile = Path.Combine(LocalKeyPath, "adbkey");
 			string LocalPubKeyFile = Path.Combine(LocalKeyPath, "adbkey.pub");
 			string BackupSentry = Path.Combine(LocalKeyPath, "gauntlet.inuse");
