Build ZenLaunch before use if it's not available during BuildCookRun snapshot import.

[RN] none

diff --git a/Engine/Source/Programs/AutomationTool/Scripts/CookCommand.Automation.cs b/Engine/Source/Programs/AutomationTool/Scripts/CookCommand.Automation.cs
--- a/Engine/Source/Programs/AutomationTool/Scripts/CookCommand.Automation.cs
+++ b/Engine/Source/Programs/AutomationTool/Scripts/CookCommand.Automation.cs
@@ -380,6 +380,17 @@ namespace AutomationScripts
 			}
 
 			FileReference ZenLaunchExe = new FileReference(Path.Combine(Unreal.EngineDirectory.FullName, "Binaries", HostPlatform.Current.HostEditorPlatform.ToString(), String.Format("ZenLaunch{0}", RuntimePlatform.ExeExtension)));
+
+			if (!FileReference.Exists(ZenLaunchExe))
+			{
+				CommandUtils.RunUBT(CommandUtils.CmdEnv, Unreal.UnrealBuildToolDllPath, CommandUtils.UBTCommandline(null, "ZenLaunch", HostPlatform.Current.HostEditorPlatform, UnrealTargetConfiguration.Development));
+			}
+
+			if (!FileReference.Exists(ZenLaunchExe))
+			{
+				throw new AutomationException("Failed to find ZenLaunch executable");
+			}
+
 			StringBuilder ZenLaunchCommandline = new StringBuilder();
 			ZenLaunchCommandline.AppendFormat("{0} -SponsorProcessID={1}", CommandUtils.MakePathSafeToUseWithCommandLine(Params.RawProjectPath.FullName), Environment.ProcessId);
 			CommandUtils.Run(ZenLaunchExe.FullName, ZenLaunchCommandline.ToString());
