DLC Packaging: Fix missing .upluginmanifest when packaging without -DLCPakPluginFile

With the plugin manifest present the output folder from a DLC package can be copied into a packaged game and will load out of the box through the default plugin discovery and loading process.

diff --git a/Engine/Source/Programs/AutomationTool/Scripts/CopyBuildToStagingDirectory.Automation.cs b/Engine/Source/Programs/AutomationTool/Scripts/CopyBuildToStagingDirectory.Automation.cs
--- a/Engine/Source/Programs/AutomationTool/Scripts/CopyBuildToStagingDirectory.Automation.cs
+++ b/Engine/Source/Programs/AutomationTool/Scripts/CopyBuildToStagingDirectory.Automation.cs
@@ -2049,7 +2049,7 @@ namespace AutomationScripts
 				// Generate a plugin manifest if we're using a pak file and not creating a mod. Mods can be enumerated independently by users copying them into the Mods directory.
 				if (Params.HasDLCName)
 				{
-					CreatePluginManifest(SC, SC.FilesToStage.UFSFiles, StagedFileType.UFS, Params.DLCFile.GetFileNameWithoutAnyExtensions());
+					CreatePluginManifest(SC, Params.DLCPakPluginFile ? SC.FilesToStage.UFSFiles : SC.FilesToStage.NonUFSFiles, Params.DLCPakPluginFile ? StagedFileType.UFS : StagedFileType.NonUFS, Params.DLCFile.GetFileNameWithoutAnyExtensions());
 				}
 				else
 				{
