Fix some AFS commands always assuming Windows is the host platorm
Fix using USERPROFILE on non-Windows platforms

#rnx

diff --git a/Engine/Source/Programs/AutomationTool/Android/AndroidPlatform.Automation.cs b/Engine/Source/Programs/AutomationTool/Android/AndroidPlatform.Automation.cs
--- a/Engine/Source/Programs/AutomationTool/Android/AndroidPlatform.Automation.cs
+++ b/Engine/Source/Programs/AutomationTool/Android/AndroidPlatform.Automation.cs
@@ -2329,7 +2329,8 @@ public class AndroidPlatform : Platform
 
 	private static IProcessResult RunAFSCommand(string SerialNumber, string Args, string Input = null, ERunOptions Options = ERunOptions.Default, bool bShouldLogCommand = false)
 	{
-		string AFSExecutable = AndroidExports.GetAFSExecutable(UnrealTargetPlatform.Win64, EpicGames.Core.Log.Logger);
+		string AFSExecutable = AndroidExports.GetAFSExecutable(BuildHostPlatform.Current.Platform, EpicGames.Core.Log.Logger);
+
 		AFSExecutable = Path.Combine(Unreal.EngineDirectory.FullName, "Binaries", "DotNET", "Android", "UnrealAndroidFileTool", AFSExecutable);
 		if (Options.HasFlag(ERunOptions.AllowSpew) || Options.HasFlag(ERunOptions.SpewIsVerbose))
 		{
@@ -2341,7 +2342,7 @@ public class AndroidPlatform : Platform
 
 	private string RunAndLogAFSCommand(ProjectParams Params, string SerialNumber, string Args, out int SuccessCode)
 	{
-		string AFSExecutable = AndroidExports.GetAFSExecutable(UnrealTargetPlatform.Win64, EpicGames.Core.Log.Logger);
+		string AFSExecutable = AndroidExports.GetAFSExecutable(BuildHostPlatform.Current.Platform, EpicGames.Core.Log.Logger);
 		AFSExecutable = Path.Combine(Unreal.EngineDirectory.FullName, "Binaries", "DotNET", "Android", "UnrealAndroidFileTool", AFSExecutable);
 		LastSpewFilename = "";
 		return RunAndLog(CmdEnv, AFSExecutable, GetAFSCommandLine(SerialNumber, Args), out SuccessCode, SpewFilterCallback: new ProcessResult.SpewFilterCallbackType(ADBSpewFilter));
diff --git a/Engine/Source/Programs/AutomationTool/Gauntlet/Platform/Android/Gauntlet.TargetDeviceAndroid.cs b/Engine/Source/Programs/AutomationTool/Gauntlet/Platform/Android/Gauntlet.TargetDeviceAndroid.cs
--- a/Engine/Source/Programs/AutomationTool/Gauntlet/Platform/Android/Gauntlet.TargetDeviceAndroid.cs
+++ b/Engine/Source/Programs/AutomationTool/Gauntlet/Platform/Android/Gauntlet.TargetDeviceAndroid.cs
@@ -1876,13 +1876,7 @@ namespace Gauntlet
 
 		public static void RestoreBackupKeys()
 		{
-			string UserDirectory = Environment.GetEnvironmentVariable("USERPROFILE");
-			// Depending on OS, USERPROFILE may be not valid so attempt to look at the HOME variable instead
-			if (UserDirectory == null)
-			{
-				UserDirectory = Path.Combine(Environment.GetEnvironmentVariable("HOME"));
-			} 
-			string LocalKeyPath = Path.Combine(UserDirectory, ".android");
+			string LocalKeyPath = Path.Combine(Environment.GetFolderPath(Environment.SpecialFolder.UserProfile), ".android");
 			string LocalKeyFile = Path.Combine(LocalKeyPath, "adbkey");
 			string LocalPubKeyFile = Path.Combine(LocalKeyPath, "adbkey.pub");
 			string BackupSentry = Path.Combine(LocalKeyPath, "gauntlet.inuse");
@@ -1973,7 +1967,7 @@ namespace Gauntlet
 
 					Log.Info("Using adb keys at {0}", KeyPath);
 
-					string LocalKeyPath = Path.Combine(Environment.GetEnvironmentVariable("USERPROFILE"), ".android");
+					string LocalKeyPath = Path.Combine(Environment.GetFolderPath(Environment.SpecialFolder.UserProfile), ".android");
 					if(!Directory.Exists(LocalKeyPath))
 					{
 						Directory.CreateDirectory(LocalKeyPath);
diff --git a/Engine/Source/Programs/UnrealBuildTool/Platform/Android/UEDeployAndroid.cs b/Engine/Source/Programs/UnrealBuildTool/Platform/Android/UEDeployAndroid.cs
--- a/Engine/Source/Programs/UnrealBuildTool/Platform/Android/UEDeployAndroid.cs
+++ b/Engine/Source/Programs/UnrealBuildTool/Platform/Android/UEDeployAndroid.cs
@@ -5173,7 +5173,7 @@ dependencies.implementation('androidx.tracing:tracing:1.2.0')
 
 					AFSToken = String.IsNullOrEmpty(AFSToken) ? "" : " -k " + AFSToken;
 
-					string AFSExecutable = Path.Combine(Unreal.EngineDirectory.ToString(), @"Binaries/DotNET/Android/UnrealAndroidFileTool", GetAFSExecutable(UnrealTargetPlatform.Win64, Logger));
+					string AFSExecutable = Path.Combine(Unreal.EngineDirectory.ToString(), @"Binaries/DotNET/Android/UnrealAndroidFileTool", GetAFSExecutable(BuildHostPlatform.Current.Platform, Logger));
 					string AFSArguments = $"-p {PackageName}{AFSToken}";
 
 					string? SOPushScriptLocation = Path.GetDirectoryName(SourceSOName)!;
