Cloud Artifacts - Fixed issue were looking up of previously published artifact fails when the stream name contains .
Fixes build issue 892236.

diff --git a/Engine/Source/Programs/AutomationTool/BuildGraph/Tasks/RetrieveCloudArtifactTask.cs b/Engine/Source/Programs/AutomationTool/BuildGraph/Tasks/RetrieveCloudArtifactTask.cs
--- a/Engine/Source/Programs/AutomationTool/BuildGraph/Tasks/RetrieveCloudArtifactTask.cs
+++ b/Engine/Source/Programs/AutomationTool/BuildGraph/Tasks/RetrieveCloudArtifactTask.cs
@@ -308,7 +308,7 @@ namespace AutomationTool.Tasks
 			queryWriter.BeginObject("query"); // query
 			
 			queryWriter.BeginObject("stream");
-			queryWriter.WriteString("$eq", streamId.ToString());
+			queryWriter.WriteString("$eq", CreateCloudArtifactTask.SanitizeBucketValue(streamId.ToString()));
 			queryWriter.EndObject();
 
 			if (maxCommitId != null)
