Fix to Android generated install scripts with overflow OBBs that are broken on non-Windows

#android

diff --git a/Engine/Source/Programs/AutomationTool/Android/AndroidPlatform.Automation.cs b/Engine/Source/Programs/AutomationTool/Android/AndroidPlatform.Automation.cs
--- a/Engine/Source/Programs/AutomationTool/Android/AndroidPlatform.Automation.cs
+++ b/Engine/Source/Programs/AutomationTool/Android/AndroidPlatform.Automation.cs
@@ -1690,7 +1690,6 @@ public class AndroidPlatform : Platform
 					{
 						string AFSOverflowName = string.Format("^overflow{0}obb", Overflow.OverflowIndex);
 						OverflowInstallCommands.Add(Overflow.bNoOverflowInstall ? AFSCommand + " deletefile '" + AFSOverflowName + "'" : AFSCommand + " push " + Path.GetFileName(Overflow.OverflowName) + " '" + AFSOverflowName + "'");
-						OverflowInstallCommands.Add("if \"%ERRORLEVEL%\" NEQ \"0\" goto Error");
 					}
 				}
 			}
@@ -1701,7 +1700,6 @@ public class AndroidPlatform : Platform
 					foreach (OverflowBatchInstallInfo Overflow in OverflowInfo)
 					{
 						OverflowInstallCommands.Add(Overflow.bNoOverflowInstall ? "\t$ADB $DEVICE shell 'rm -r $EXTERNAL_STORAGE/" + Overflow.DeviceOverflowName + "'" : "\t$ADB $DEVICE push " + Path.GetFileName(Overflow.OverflowName) + " " + TargetAndroidTemp + Overflow.DeviceOverflowName);
-						OverflowInstallCommands.Add("if \"%ERRORLEVEL%\" NEQ \"0\" goto Error");
 					}
 				}
 			}
