Update UE.IterativeCookByTheBook because of the renaming Iterative feature to Incremental

1. The renaming Iterative to Incremental affected UE.IterativeCookByTheBook
as the log text has changed.
2. The -iterate works as before, I use -legacyiterative just to highlight the obsolescence.

#rnx

diff --git a/Engine/Source/Programs/AutomationTool/Gauntlet/Unreal/Automation/UE.IterativeCookByTheBook.cs b/Engine/Source/Programs/AutomationTool/Gauntlet/Unreal/Automation/UE.IterativeCookByTheBook.cs
--- a/Engine/Source/Programs/AutomationTool/Gauntlet/Unreal/Automation/UE.IterativeCookByTheBook.cs
+++ b/Engine/Source/Programs/AutomationTool/Gauntlet/Unreal/Automation/UE.IterativeCookByTheBook.cs
@@ -14,7 +14,7 @@ namespace UE
 		private const string PacketsSkippedKey =
 			"The number of cooked packets at the first pass is equal to the number of iteratively skipped packets at the second pass";
 		private const string PacketsStatisticsPattern = "Total Packages";
-		private const string DetailedPacketsStatisticsPattern = @"Packages Cooked: (\d+), Packages Iteratively Skipped: (\d+)";
+		private const string DetailedPacketsStatisticsPattern = @"Packages Cooked: (\d+), Packages Incrementally Skipped: (\d+)";
 		private const string IterativePattern = "IsIterativeCook=true";
 
 		private List<(string Cooked, string Skipped)> PacketsStatistics;
@@ -22,7 +22,7 @@ namespace UE
 		public IterativeCookByTheBook(UnrealTestContext InContext) : base(InContext)
 		{
 			PacketsStatistics = [];
-			BaseEditorCommandLine += " -cookcultures=en -IgnoreIniSettingsOutOfDate -iterate";
+			BaseEditorCommandLine += " -cookcultures=en -IgnoreIniSettingsOutOfDate -legacyiterative";
 		}
 
 		protected override void InitTest()
