Cloud Artifacts - Added option "EnableTracing" to generate a Insights trace of the download/upload process.
Also disabled verbose logging by default as it is quite a lot.

diff --git a/Engine/Source/Programs/AutomationTool/BuildGraph/Tasks/CreateCloudArtifactTask.cs b/Engine/Source/Programs/AutomationTool/BuildGraph/Tasks/CreateCloudArtifactTask.cs
--- a/Engine/Source/Programs/AutomationTool/BuildGraph/Tasks/CreateCloudArtifactTask.cs
+++ b/Engine/Source/Programs/AutomationTool/BuildGraph/Tasks/CreateCloudArtifactTask.cs
@@ -16,6 +16,7 @@ using EpicGames.Horde.Streams;
 using EpicGames.Serialization;
 using Microsoft.Extensions.DependencyInjection;
 using Microsoft.Extensions.Logging;
+using UnrealBuildBase;
 using UnrealBuildTool;
 
 #nullable enable
@@ -156,13 +157,18 @@ namespace AutomationTool.Tasks
 		/// Verbose logging output, can be useful to find out more information when an issue has occured.
 		/// </summary>
 		[TaskParameter(Optional = true)]
-		public bool Verbose { get; set; } = true;
+		public bool Verbose { get; set; } = false;
 
 		/// <summary>
 		/// Increase the number of worker threads used by zen, may cause machine to be less responsive but will generally improve upload times
 		/// </summary>
 		[TaskParameter(Optional = true)]
 		public bool BoostWorkers { get; set; } = false;
+
+		/// <summary>
+		/// Enable to create an Unreal Insights trace of the upload process
+		/// </summary>
+		public bool EnableTracing { get; set; } = true;
 	}
 
 	/// <summary>
@@ -473,6 +479,13 @@ namespace AutomationTool.Tasks
 			CbObjectId buildId = CbObjectId.NewObjectId();
 			CbObjectId partId = CbObjectId.NewObjectId();
 
+			FileReference? traceFile = null;
+			if (_parameters.EnableTracing)
+			{
+				// this is automatically uploaded as it's in the Saved folder
+				traceFile = FileReference.Combine(Unreal.RootDirectory, "Engine/Programs/AutomationTool/Saved/Logs/build-upload.utrace");
+			}
+
 			// set part name to default by default - should be exposed as a option to allow for uploading separate parts
 			string partName = "default";
 
@@ -483,7 +496,7 @@ namespace AutomationTool.Tasks
 				Branch = streamId,
 				BaselineBranch = baselineBranch,
 				Platform = platform
-			}, allowMultipart: _parameters.AllowMultipart, assumeHttp2: String.Equals(httpVersion, "http2-only", StringComparison.OrdinalIgnoreCase), forceClean: _parameters.Force, verify: _parameters.Verify, verbose: _parameters.Verbose, boostWorkers: _parameters.BoostWorkers);
+			}, allowMultipart: _parameters.AllowMultipart, assumeHttp2: String.Equals(httpVersion, "http2-only", StringComparison.OrdinalIgnoreCase), forceClean: _parameters.Force, verify: _parameters.Verify, verbose: _parameters.Verbose, boostWorkers: _parameters.BoostWorkers, traceFile: traceFile);
 
 			Uri buildUri = new Uri($"{descriptorHost}/api/v2/builds/{ns}/{bucketId}/{buildId}");
 
@@ -596,7 +609,7 @@ namespace AutomationTool.Tasks
 		public static void UploadBuild(string host, DirectoryReference baseDir, string namespaceId, string bucketId, CbObjectId buildId,
 			string partName, CbObjectId partId, Dictionary<string, object> metadata, string accessToken, BuildInfo info,
 			bool allowMultipart = false, bool assumeHttp2 = false, bool forceClean = false, bool verify = false, bool verbose = false,
-			bool boostWorkers = false)
+			bool boostWorkers = false, FileReference? traceFile = null)
 		{
 			string metadataJsonPath = Path.Combine(Path.GetTempPath(), Path.GetTempFileName());
 			try
@@ -638,8 +651,10 @@ namespace AutomationTool.Tasks
 				string verifyOption = verify ? "--verify" : "";
 				string verboseOption = verbose ? "--verbose" : "";
 				string boostWorkersOption = boostWorkers ? "--boost-workers" : "";
+				// trace file option needs to be passed before other arguments including the verbs
+				string traceOption = traceFile != null ? $"--tracefile={traceFile.FullName} " : "";
 				// pass the access token via environment variable to avoid showing it on the cli
-				string cmdline = $"builds upload --url {host} --namespace {namespaceId} --bucket {bucketId} --build-id {buildId} --build-part-name \"{partName}\" --build-part-id {partId} --metadata-path \"{metadataJsonPath}\" --access-token-env UE-CloudDataCacheAccessToken --create-build --plain-progress --allow-multipart={allowMultipart} {http2Options} {boostWorkersOption} --clean={forceClean} {verifyOption} {verboseOption} \"{baseDir}\"";
+				string cmdline = $"{traceOption}builds upload --url {host} --namespace {namespaceId} --bucket {bucketId} --build-id {buildId} --build-part-name \"{partName}\" --build-part-id {partId} --metadata-path \"{metadataJsonPath}\" --access-token-env UE-CloudDataCacheAccessToken --create-build --plain-progress --allow-multipart={allowMultipart} {http2Options} {boostWorkersOption} --clean={forceClean} {verifyOption} {verboseOption} \"{baseDir}\"";
 				Log.Logger.LogInformation("Using metadata '{Metadata}'", metadataString);
 				CommandUtils.RunAndLog(CommandUtils.CmdEnv, zenExe.FullName,cmdline , Options: CommandUtils.ERunOptions.Default, EnvVars: new Dictionary<string, string> {{"UE-CloudDataCacheAccessToken", accessToken}});
 			}
diff --git a/Engine/Source/Programs/AutomationTool/BuildGraph/Tasks/RetrieveCloudArtifactTask.cs b/Engine/Source/Programs/AutomationTool/BuildGraph/Tasks/RetrieveCloudArtifactTask.cs
--- a/Engine/Source/Programs/AutomationTool/BuildGraph/Tasks/RetrieveCloudArtifactTask.cs
+++ b/Engine/Source/Programs/AutomationTool/BuildGraph/Tasks/RetrieveCloudArtifactTask.cs
@@ -17,6 +17,7 @@ using EpicGames.Horde.Commits;
 using EpicGames.Horde.Streams;
 using EpicGames.Serialization;
 using Microsoft.Extensions.Logging;
+using UnrealBuildBase;
 using UnrealBuildTool;
 
 #nullable enable
@@ -123,6 +124,11 @@ namespace AutomationTool.Tasks
 		/// </summary>
 		[TaskParameter(Optional = true)]
 		public bool BoostWorkers { get; set; } = false;
+
+		/// <summary>
+		/// Enable to create an Unreal Insights trace of the download process
+		/// </summary>
+		public bool EnableTracing { get; set; } = true;
 	}
 
 	/// <summary>
@@ -383,7 +389,14 @@ namespace AutomationTool.Tasks
 			CbObjectId artifact = artifactIds.First();
 			Logger.LogInformation("Found artifact {ArtifactId}", artifact);
 
-			BuildDownload.DownloadBuild(host, outputDir, ns, bucketId, artifact.ToString(), accessToken, allowMultipart: _parameters.AllowMultipart, assumeHttp2: String.Equals(httpVersion, "http2-only", StringComparison.OrdinalIgnoreCase), boostWorkers: _parameters.BoostWorkers);
+			FileReference? traceFile = null;
+			if (_parameters.EnableTracing)
+			{
+				// this is automatically uploaded as it's in the Saved folder
+				traceFile = FileReference.Combine(Unreal.RootDirectory, "Engine/Programs/AutomationTool/Saved/Logs/build-download.utrace");
+			}
+
+			BuildDownload.DownloadBuild(host, outputDir, ns, bucketId, artifact.ToString(), accessToken, allowMultipart: _parameters.AllowMultipart, assumeHttp2: String.Equals(httpVersion, "http2-only", StringComparison.OrdinalIgnoreCase), boostWorkers: _parameters.BoostWorkers, traceFile: traceFile);
 		}
 
 		static HttpClient BuildHttpClient(string host, string ns, string bucketId, string accessToken)
@@ -457,15 +470,17 @@ namespace AutomationTool.Tasks
 
 	static class BuildDownload
 	{
-		public static void DownloadBuild(string host, DirectoryReference targetDir, string namespaceId, string bucketId, string buildId, string accessToken, bool allowMultipart = false, bool assumeHttp2 = false, bool boostWorkers = false)
+		public static void DownloadBuild(string host, DirectoryReference targetDir, string namespaceId, string bucketId, string buildId, string accessToken, bool allowMultipart = false, bool assumeHttp2 = false, bool boostWorkers = false, FileReference? traceFile = null)
 		{
 			FileInfo zenExe = new FileInfo("Engine/Binaries/Win64/zen.exe");
 
 			string http2Options = assumeHttp2 ? "--assume-http2" : "";
 			string boostWorkersOption = boostWorkers ? "--boost-workers" : "";
+			// trace file option needs to be passed before other arguments including the verbs
+			string traceOption = traceFile != null ? $"--tracefile={traceFile.FullName} " : "";
 			// pass the access token via environment variable to avoid showing it on the cli
 			string cmdline =
-				$"builds download --url {host} --namespace {namespaceId} --bucket {bucketId} --access-token-env UE-CloudDataCacheAccessToken --plain-progress --allow-multipart={allowMultipart} {http2Options} {boostWorkersOption} \"{targetDir}\" {buildId}";
+				$"{traceOption}builds download --url {host} --namespace {namespaceId} --bucket {bucketId} --access-token-env UE-CloudDataCacheAccessToken --plain-progress --allow-multipart={allowMultipart} {http2Options} {boostWorkersOption} \"{targetDir}\" {buildId}";
 			CommandUtils.RunAndLog(CommandUtils.CmdEnv, zenExe.FullName, cmdline, Options: CommandUtils.ERunOptions.Default,
 				EnvVars: new Dictionary<string, string> { { "UE-CloudDataCacheAccessToken", accessToken } });
 		}
