Attempt 2:
Cloud metadata unification part 1: bring exported snapshot metadata in line with common keys used in staged/packaged builds.

diff --git a/Engine/Build/Graph/Tasks/BuildAndTestProject.xml b/Engine/Build/Graph/Tasks/BuildAndTestProject.xml
--- a/Engine/Build/Graph/Tasks/BuildAndTestProject.xml
+++ b/Engine/Build/Graph/Tasks/BuildAndTestProject.xml
@@ -746,7 +746,7 @@
 						<!-- Publish to builds cloud storage endpoint -->
 						<Property Name="SnapshotBuildBaselineBranch" Value="$(Branch)"/>
 						<Property Name="SnapshotBuildBaselineBranch" Value="$(UE-CloudPublishBaselineBranch)" If="'$(UE-CloudPublishBaselineBranch)' != ''"/>
-						<Property Name="SnapshotBuildMetaData" Value="name=$(SnapshotIdentifier);depot=$(Depot);branch=$(Branch);baselineBranch=$(SnapshotBuildBaselineBranch);project=$(ProjectName);runtime=$(TargetType);platform=$(TargetPlatform);cookPlatform=$(CookPlatform);changelist=$(Change);codeChangelist=$(CodeChange);isPreflight=$(IsPreflight);isBuildMachine=$(IsBuildMachine);isHordeEnv=$(IsHordeEnv)"/>
+						<Property Name="SnapshotBuildMetaData" Value="name=$(SnapshotIdentifier);depot=$(Depot);branch=$(Branch);baselineBranch=$(SnapshotBuildBaselineBranch);project=$(ProjectName);runtime=$(TargetType);platform=$(TargetPlatform);cookPlatform=$(CookPlatform);changelist=$(Change);commit=$(Change);codeChangelist=$(CodeChange);codeCommit=$(CodeChange);isPreflight=$(IsPreflight);isBuildMachine=$(IsBuildMachine);isHordeEnv=$(IsHordeEnv)"/>
 						<Property Name="SnapshotBuildMetaData" Value="$(SnapshotBuildMetaData);cookFlavor=$(BCRCookFlavor)" If="'$(BCRCookFlavor)' != ''" />
 						<Property Name="SnapshotBuildMetaData" Value="$(SnapshotBuildMetaData);preflightChangelist=$(PreflightChange)" If="'$(PreflightChange)' != ''" />
 						<Property Name="SnapshotBuildMetaData" Value="$(SnapshotBuildMetaData);hordeTemplateId=$(UE_HORDE_TEMPLATEID)" If="'$(UE_HORDE_TEMPLATEID)' != ''" />
diff --git a/Engine/Build/Graph/Tests/IncrementalCompare.xml b/Engine/Build/Graph/Tests/IncrementalCompare.xml
--- a/Engine/Build/Graph/Tests/IncrementalCompare.xml
+++ b/Engine/Build/Graph/Tests/IncrementalCompare.xml
@@ -94,7 +94,7 @@
 				<StringOp Input="$(TargetPlatform)" Method="SplitFirst" Arguments="Client" Output="ClientPrefix"/>
 				<Property Name="TargetType" Value="Client" If="'$(ClientPrefix)' != ''"/>
 
-				<Property Name="SnapshotBuildMetaData" Value="name=$(SnapshotIdentifier).$(TargetPlatform);depot=$(Depot);branch=$(Branch);baselineBranch=$(SnapshotBuildBaselineBranch);project=$(ProjectName);runtime=$(TargetType);platform=Win64;cookPlatform=$(TargetPlatform);changelist=$(Change);codeChangelist=$(CodeChange);isPreflight=$(IsPreflight);isBuildMachine=$(IsBuildMachine);isHordeEnv=$(IsHordeEnv)"/>
+				<Property Name="SnapshotBuildMetaData" Value="name=$(SnapshotIdentifier).$(TargetPlatform);depot=$(Depot);branch=$(Branch);baselineBranch=$(SnapshotBuildBaselineBranch);project=$(ProjectName);runtime=$(TargetType);platform=Win64;cookPlatform=$(TargetPlatform);changelist=$(Change);commit=$(Change);codeChangelist=$(CodeChange);codeCommit=$(CodeChange);isPreflight=$(IsPreflight);isBuildMachine=$(IsBuildMachine);isHordeEnv=$(IsHordeEnv)"/>
 				<Property Name="SnapshotBuildMetaData" Value="$(SnapshotBuildMetaData);preflightChangelist=$(PreflightChange)" If="'$(PreflightChange)' != ''" />
 				<Property Name="SnapshotBuildMetaData" Value="$(SnapshotBuildMetaData);hordeTemplateId=$(UE_HORDE_TEMPLATEID)" If="'$(UE_HORDE_TEMPLATEID)' != ''" />
 				<Log Message="Exporting cook state with metadata '$(SnapshotBuildMetaData)'"/>
diff --git a/Engine/Source/Programs/AutomationTool/BuildGraph/Tasks/ZenExportSnapshotTask.cs b/Engine/Source/Programs/AutomationTool/BuildGraph/Tasks/ZenExportSnapshotTask.cs
--- a/Engine/Source/Programs/AutomationTool/BuildGraph/Tasks/ZenExportSnapshotTask.cs
+++ b/Engine/Source/Programs/AutomationTool/BuildGraph/Tasks/ZenExportSnapshotTask.cs
@@ -627,10 +627,36 @@ namespace AutomationTool.Tasks
 						string destinationId;
 						if (destinationStorageType == SnapshotStorageType.Builds)
 						{
+							StringBuilder metadata = new StringBuilder();
+							metadata.AppendFormat("type=oplog;createdAt={0}", DateTime.UtcNow.ToString("O"));
+
+							// Add keys for the job that's executing
+							string? stepId = null;
+							string? jobId = Environment.GetEnvironmentVariable("UE_HORDE_JOBID");
+							if (!String.IsNullOrEmpty(jobId))
+							{
+								metadata.AppendFormat(";job={0}", jobId);
+
+								stepId = Environment.GetEnvironmentVariable("UE_HORDE_STEPID");
+								if (!String.IsNullOrEmpty(stepId))
+								{
+									metadata.AppendFormat(";step={0}", stepId);
+								}
+							}
+
+							string? hordeUrl = Environment.GetEnvironmentVariable("UE_HORDE_URL");
+							// if we are running in horde and have the required environment variables we append a link back to the horde job into metadata
+							if (!String.IsNullOrEmpty(hordeUrl) && !String.IsNullOrEmpty(jobId) && !String.IsNullOrEmpty(stepId))
+							{
+								metadata.Append($";buildurl={hordeUrl}job/{jobId}?step={stepId}");
+							}
+
 							if (!String.IsNullOrEmpty(_parameters.Metadata))
 							{
-								exportSingleSourceCommandline.AppendFormat(" --builds-metadata \"{0}\"", _parameters.Metadata);
+								metadata.AppendFormat(";{0}", _parameters.Metadata);
 							}
+
+							exportSingleSourceCommandline.AppendFormat(" --builds-metadata \"{0}\"", metadata.ToString());
 							CbObjectId objectId = CbObjectId.NewObjectId();
 							destinationId = objectId.ToString().ToLowerInvariant();
 						}
