- Some general minor bug fixes
- Turnkey personal projects directory exists check
- Force GShowSplashScreen to false on VisionOS

diff --git a/Engine/Source/Programs/AutomationTool/Turnkey/TurnkeyUtils.cs b/Engine/Source/Programs/AutomationTool/Turnkey/TurnkeyUtils.cs
--- a/Engine/Source/Programs/AutomationTool/Turnkey/TurnkeyUtils.cs
+++ b/Engine/Source/Programs/AutomationTool/Turnkey/TurnkeyUtils.cs
@@ -263,12 +263,15 @@ namespace Turnkey
 
 			// now look for the My Documents location
 			DirectoryReference PersonalProjects = new DirectoryReference(Path.Combine(Environment.GetFolderPath(Environment.SpecialFolder.Personal), "Unreal Projects"));
-			foreach (DirectoryReference PersonalProjectDir in DirectoryReference.EnumerateDirectories(PersonalProjects))
+			if (DirectoryReference.Exists(PersonalProjects))
 			{
-				FileReference PersonalProjectFile = DirectoryReference.EnumerateFiles(PersonalProjectDir, "*.uproject").FirstOrDefault();
-				if (PersonalProjectFile != null)
+				foreach (DirectoryReference PersonalProjectDir in DirectoryReference.EnumerateDirectories(PersonalProjects))
 				{
-					ProjectsByType[ProjectMode.Personal].Add(PersonalProjectFile);
+					FileReference PersonalProjectFile = DirectoryReference.EnumerateFiles(PersonalProjectDir, "*.uproject").FirstOrDefault();
+					if (PersonalProjectFile != null)
+					{
+						ProjectsByType[ProjectMode.Personal].Add(PersonalProjectFile);
+					}
 				}
 			}
 
diff --git a/Engine/Source/Runtime/ApplicationCore/Private/IOS/IOSAppDelegate.cpp b/Engine/Source/Runtime/ApplicationCore/Private/IOS/IOSAppDelegate.cpp
--- a/Engine/Source/Runtime/ApplicationCore/Private/IOS/IOSAppDelegate.cpp
+++ b/Engine/Source/Runtime/ApplicationCore/Private/IOS/IOSAppDelegate.cpp
@@ -1063,6 +1063,8 @@ static FAutoConsoleVariableRef CVarGEnableThermalsReport(
         [self.Window addSubview: self.viewController.view];
         GShowSplashScreen = true;
     }
+#else
+	GShowSplashScreen = false;
 #endif
 
     timer = [NSTimer scheduledTimerWithTimeInterval: 0.05f target:self selector:@selector(timerForSplashScreen) userInfo:nil repeats:YES];
