Making LaunchMultiServer friendlier for public consumption.

diff --git a/Engine/Source/Programs/AutomationTool/Scripts/LaunchMultiServer.Automation.cs b/Engine/Source/Programs/AutomationTool/Scripts/LaunchMultiServer.Automation.cs
--- a/Engine/Source/Programs/AutomationTool/Scripts/LaunchMultiServer.Automation.cs
+++ b/Engine/Source/Programs/AutomationTool/Scripts/LaunchMultiServer.Automation.cs
@@ -234,32 +234,64 @@ Example running 2 hosts locally with a single proxy server & client:
 
 Example running 4 clients locally connecting to single proxy server w/ a non-local cluster of 3:
 	-cook -notimeouts -client=proxy -proxyclientcount=4 -proxycycleprimary -nonlocalservers=172.27.63.109:9001,172.27.63.109:9002,172.27.63.109:9003
-")]
+	")]
 	[Help("Project=<project>", "Project to open. Will search current path and paths in ueprojectdirs. Defaults to current project (in uShell).")]
 	[Help("Map=<MapName>", "Map to load on startup.")]
-	[Help("CustomConfig=<Section>", "Read a different section of the Project's GameConfig for default values")]
-    [Help("NumServers=##", "Set the fleet size of servers")]
-	[Help("NonLocalServers=Address1:Port1,Address2:Port2,etc.", "Specify servers that are not hosted locally but belong to the same fleet")]
-	[Help("NoTimeouts", "Disable timeouts (recommended for debugging)")]
-	[Help("NoConsole", "Disable launching a console (not recommended, but may improve performance)")]
-	[Help("NoNewConsole", "Disable launching the new style console")]
-	[Help("AttachToServers", "Force Visual Studio to attach to the servers when launched (requires VS be launched or it will hang)")]
-	[Help("AttachToClients", "Force Visual Studio to attach to the clients when launched (requires VS be launched or it will hang)")]
-	[Help("LogDebug", "Enable certain helpful Verbose logs for debugging")]
-	[Help("LogDebugV", "Enable certain helpful VeryVerbose logs for debugging")]
-	[Help("Client=<Type>", "Specify the client type to connect with (either PROXY or DIRECT)")]
-	[Help("ProxyServerCount=##", "Specify how many proxy servers to launch in client=proxy mode")]
-	[Help("ProxyClientCount=##", "Specify how many clients to launch PER SERVER in client=proxy mode")]
-	[Help("ProxyBotCount=##", "Specify how many bots to launch PER SERVER in client=proxy mode")]
-	[Help("ProxyClientPrimary=##", "Specify which game server to prefer as the initial proxy connection")]
-	[Help("ProxyCyclePrimary", "Specify that each subsequent proxy client should choose a different server")]
 	public class LaunchMultiServer : BuildCommand, IProjectParamsHelpers
 	{
 		// We will use a convention where ServerId 1 will start at BaseGameListenPort + 1;
-		int BaseGameListenPort = 9000;
-		int ProxyServerBasePort = 11000;
+		public int BaseGameListenPort = 9000;
+		public int ProxyServerBasePort = -1;
 
-		private enum ProcessLaunchType
+		[CommandLine(), Help("CustomConfig=<Section>", "Read a different section of the Project's GameConfig for default values")]
+		public string CustomConfig = "";
+
+		[CommandLine(), Help("NumServers=##", "Set the fleet size of servers (default: 2)")]
+		public int NumServers = -1;
+
+		[CommandLine(ListSeparator = ','), Help("NonLocalServers=Address1:Port1,Address2:Port2,etc.", "Specify servers that are not hosted locally but belong to the same fleet")]
+		public List<string> NonLocalServers = new List<string>();
+
+		[CommandLine("-Client", ListSeparator = ','), Help("Client=<Type>", "Specify the client type to connect with (proxy or direct or proxy,direct)")]
+		public List<string> ClientType = new List<string>();
+
+		[CommandLine(), Help("ProxyServerCount=##", "Specify how many proxy servers to launch in client=proxy mode")]
+		public int ProxyServerCount = 1;
+
+		[CommandLine(), Help("ProxyClientCount=##", "Specify how many clients to launch PER SERVER in client=proxy mode")]
+		public int ProxyClientCount = 1;
+
+		[CommandLine(), Help("ProxyBotCount=##", "Specify how many bots to launch PER SERVER in client=proxy mode")]
+		public int ProxyBotCount = 0;
+
+		[CommandLine(), Help("ProxyClientPrimary=##", "Specify which game server to prefer as the initial proxy connection")]
+		public int ProxyClientPrimary = -1;
+
+		[CommandLine(), Help("ProxyCyclePrimary", "Specify that each subsequent proxy client should choose a different server")]
+		public bool ProxyCyclePrimary = false;
+
+		[CommandLine(), Help("NoTimeouts", "Disable timeouts (recommended for debugging)")]
+		public bool NoTimeouts = false;
+
+		[CommandLine(), Help("NoConsole", "Disable launching a console (not recommended, but may improve performance)")]
+		public bool NoConsole = false;
+
+		[CommandLine(), Help("NoNewConsole", "Disable launching the new style console")]
+		public bool NoNewConsole = false;
+
+		[CommandLine(), Help("AttachToServers", "Force Visual Studio to attach to the servers when launched (requires VS be launched or it will hang)")]
+		public bool AttachToServers = false;
+
+		[CommandLine(), Help("AttachToClients", "Force Visual Studio to attach to the clients when launched (requires VS be launched or it will hang)")]
+		public bool AttachToClients = false;
+
+		[CommandLine(), Help("CommonArgs", "Specify arguments that are run by all processes")]
+		public string CommonArgs = "";
+
+		// We should only parse the commandline arguments once, since doing so multiple times will add entries to the List<> args.
+		bool bParsedCommandLineArguments = false;
+
+		protected enum ProcessLaunchType
 		{
 			Server,
 			DirectClient,
@@ -268,38 +300,33 @@ Example running 4 clients locally connecting to single proxy server w/ a non-loc
 			ProxyBot
 		}
 
+		protected virtual void ParseCommandLineArguments()
+		{
+			if (!bParsedCommandLineArguments)
+			{
+				// We need to prepend the '-' back on to all of the Parameters for ParseArguments to work
+				UnrealBuildTool.CommandLine.ParseArguments(Params.Select(x => $"-{x}"), this);
+				bParsedCommandLineArguments = true;
+			}
+		}
+
+		/// <summary>
+		/// The entry point for the build command
+		/// </summary>
 		public override ExitCode Execute()
 		{
 			Logger.LogInformation("********** RUN MULTISERVER COMMAND STARTED **********");
 			var StartTime = DateTime.UtcNow;
-			var ProjectName = ProjectUtils.GetShortProjectName(ProjectPath);
-
-			// Parse server configuration from ini files
-			string CustomConfig = ParseParamValue("CustomConfig", string.Empty);
-
-			ConfigHierarchy ProjectGameConfig = ConfigCache.ReadHierarchy(ConfigHierarchyType.Game, DirectoryReference.FromFile(ProjectPath), UnrealTargetPlatform.Win64, CustomConfig);
-
-			const String ServerDefConfigSection = "/Script/MultiServerConfiguration.MultiServerSettings";
-			const String ProxyConfigSection = "/Script/MultiServerConfiguration.Proxy";
-
-			if (!ProjectGameConfig.TryGetValue(ProxyConfigSection, "ListenPort", out ProxyServerBasePort))
-			{
-				ProxyServerBasePort = BaseGameListenPort + 2000;
-				Logger.LogWarning($"Could not find {ProxyConfigSection}.ListenPort in Game Config File {ProjectGameConfig} Defaulting to {ProxyServerBasePort}");
-			}
 
-			int TotalNumServers = ParseParamInt("NumServers", -1);
-			if (TotalNumServers < 0 && !ProjectGameConfig.TryGetValue(ServerDefConfigSection, "TotalNumServers", out TotalNumServers))
-			{
-				TotalNumServers = 2;
-				Logger.LogWarning($"Could not find {ServerDefConfigSection}.TotalNumServers in Game Config File {ProjectGameConfig} Defaulting to {TotalNumServers}");
-			}
+			ParseCommandLineArguments();
+			ReadConfiguration();
+			int TotalNumServers = NumServers;
 
 			// CommonArgs are used for all instances spawned
-			string CommonArgs = (CustomConfig.Length > 0) ? $" -CustomConfig={CustomConfig}" : "";
+			CommonArgs += (CustomConfig?.Length > 0) ? $" -CustomConfig={CustomConfig}" : "";
 
 			// These are servers not on our local host
-			string[] NonLocalServerAddresses = ParseParamValues("NonLocalServers").SelectMany(x => x.Split( new char[]{';', ','}, StringSplitOptions.TrimEntries)).ToArray();
+			string[] NonLocalServerAddresses = NonLocalServers.ToArray();
 			if (NonLocalServerAddresses.Length > TotalNumServers)
 			{
 				Logger.LogWarning($"Specified more NonLocalServers=({NonLocalServerAddresses.Length}) than NumServers={TotalNumServers}. Inferring that we want to run with no local servers and only non-local servers. Setting NumServers={NonLocalServerAddresses.Length}.");
@@ -317,27 +344,26 @@ Example running 4 clients locally connecting to single proxy server w/ a non-loc
 			NonLocalServerAddresses.CopyTo(AllGameServerAddresses, LocalGameServerAddresses.Length);
 
 			// Now start-up the clients
-			string[] ClientValues = ParseParamValue("client", string.Empty).Split(',');
-			foreach (string ClientValue in ClientValues)
+			foreach (string ClientValue in ClientType)
 			{
-				if (ClientValue == "direct")
+				if (ClientValue.ToLower() == "direct")
 				{
 					Logger.LogInformation("Starting direct client instances connecting to MultiServer instances");
 					StartProcessesForClients(ProcessLaunchType.DirectClient, LocalGameServerAddresses, CommonArgs);
 				}
-				else if (ClientValue == "proxy")
+				else if (ClientValue.ToLower() == "proxy")
 				{
-					int ProxyServerCount = ParseParamInt("ProxyServerCount", 1);
-					int NumClientsPerServer = ParseParamInt("ProxyClientCount", 1);
-					int NumBotsPerServer = ParseParamInt("ProxyBotCount", 0);
-
-					Logger.LogInformation($"Starting {ProxyServerCount} Proxy Servers each with {NumClientsPerServer} Clients and {NumBotsPerServer} Bots");
+					Logger.LogInformation($"Starting {ProxyServerCount} Proxy Servers each with {ProxyClientCount} Clients and {ProxyBotCount} Bots");
 
 					string[] LocalProxyAddresses = GenerateLocalServerAddressRange(ProxyServerCount, ProxyServerBasePort);
 					StartProcessesForProxyServer(LocalProxyAddresses, AllGameServerAddresses, CommonArgs);
 
-					StartProcessesForClients(ProcessLaunchType.ProxyClient, LocalProxyAddresses, CommonArgs, NumClientsPerServer);
-					StartProcessesForClients(ProcessLaunchType.ProxyBot, LocalProxyAddresses, $"{CommonArgs} -nullrhi -bot -nosound -unattended -nosplash", NumBotsPerServer);
+					StartProcessesForClients(ProcessLaunchType.ProxyClient, LocalProxyAddresses, CommonArgs, ProxyClientCount);
+					StartProcessesForClients(ProcessLaunchType.ProxyBot, LocalProxyAddresses, $"{CommonArgs} -nullrhi -bot -nosound -unattended -nosplash", ProxyBotCount);
+				}
+				else
+				{
+					Logger.LogWarning($"Unknown client type specified: {ClientValue}");
 				}
 			}
 
@@ -347,10 +373,46 @@ Example running 4 clients locally connecting to single proxy server w/ a non-loc
 			return ExitCode.Success;
 		}
 
-		/**
-		 * Given the number of local servers and the base port, generate the array of address:port for that range of ports
-		 * The address range will be ["127.0.0.1:BasePort+1" ... "127.0.0.1::BasePort+NumLocalServers"]
-		 */
+		/// <summary>
+		/// Reads the configuration from the CustomConfig specified.  If no CustomConfig was specified, then rely solely on the command-line parameters.
+		/// The command-line parameters should always take precedence.
+		/// </summary>
+		protected void ReadConfiguration()
+		{
+			// Parse server configuration from ini files
+			ConfigHierarchy ProjectGameConfig = ConfigCache.ReadHierarchy(ConfigHierarchyType.Game, DirectoryReference.FromFile(ProjectPath), UnrealTargetPlatform.Win64, CustomConfig);
+
+			const String ServerDefConfigSection = "/Script/MultiServerConfiguration.MultiServerSettings";
+			const String ProxyConfigSection = "/Script/MultiServerConfiguration.Proxy";
+
+			if (ProxyServerBasePort > 0)
+			{
+				Logger.LogInformation($"Using specified ProxyServerBasePort {ProxyServerBasePort}");
+			}
+			else if (ProjectGameConfig.TryGetValue(ProxyConfigSection, "ListenPort", out ProxyServerBasePort))
+			{
+				Logger.LogInformation($"Found {ProxyConfigSection}.ListenPort setting ProxyServerBasePort to {ProxyServerBasePort}");
+			}
+			else
+			{
+				ProxyServerBasePort = BaseGameListenPort + 2000;
+				Logger.LogWarning($"Could not find {ProxyConfigSection}.ListenPort in Game Config File {ProjectGameConfig} Defaulting ProxyServerBasePort to {ProxyServerBasePort}");
+			}
+
+			int TotalNumServers = NumServers;
+			if (TotalNumServers < 0 && !ProjectGameConfig.TryGetValue(ServerDefConfigSection, "TotalNumServers", out TotalNumServers))
+			{
+				NumServers = 2;
+				Logger.LogWarning($"Could not find {ServerDefConfigSection}.TotalNumServers in Game Config File {ProjectGameConfig} Defaulting to {NumServers}");
+			}
+		}
+
+		/// <summary>
+		/// Given the number of local servers and the base port, generate the array of address:port for that range of ports
+		/// </summary>
+		/// <param name="NumLocalServers">The number of servers in the localhost range</param>
+		/// <param name="BasePort">The port to start enumerating at</param>
+		/// <returns>The address range of servers: ["127.0.0.1:BasePort+1", ..., "127.0.0.1::BasePort+NumLocalServers"]</returns>
 		string[] GenerateLocalServerAddressRange(int NumLocalServers, int BasePort)
 		{
 			string[] LocalServers = new string[NumLocalServers];
@@ -363,14 +425,14 @@ Example running 4 clients locally connecting to single proxy server w/ a non-loc
 			return LocalServers;
 		}
 
-		/**
-		 * Starts all of the server processes that should execute on the local host.
-		 * @param NumLocalServers: The number of local servers to spin-up as part of the total server cluster (the total cluster consistent of these local servers + NonLocalServers).
-		 * @param NonLocalServers: Specify servers that are not running on the local machine but contribute to the Metaverse
-		 * @param AdditionalServerParameters: These are extra parameters that the server processes should put on the command-line (shared between all servers)
-		 * @return The Addresses of all of the Game Servers spun-up.
-		 */
-		private string[] StartProcessesForDedicatedServers(int NumLocalServers, string[] NonLocalServers, string AdditionalServerArguments)
+		/// <summary>
+		///	Starts all of the server processes that should execute on the local host.
+		/// </summary>
+		/// <param name="NumLocalServers">The number of local servers to spin-up as part of the total server cluster(the total cluster consistent of these local servers + NonLocalServers).</param>
+		/// <param name="NonLocalServers">Specify servers that are not running on the local machine but contribute to the Metaverse</param>
+		/// <param name="AdditionalServerArguments">These are extra parameters that the server processes should put on the command-line(shared between all servers)</param>
+		/// <returns>The Addresses of all of the Game Servers spun-up.</returns>
+		protected string[] StartProcessesForDedicatedServers(int NumLocalServers, string[] NonLocalServers, string AdditionalServerArguments)
 		{
 			if (NumLocalServers < 1)
 			{
@@ -382,7 +444,7 @@ Example running 4 clients locally connecting to single proxy server w/ a non-loc
 
 			// This is optional to the multi-server process, these are in addition to the MultiServerNumPeers argument
 			string ServerCommonArgs = $" -server {AdditionalServerArguments} -NODEBUGOUTPUT -SuppressConsoleOutputVerboseLogging";
-			ServerCommonArgs += ParseParam("notimeouts") ? " -notimeouts" : "";
+			ServerCommonArgs += NoTimeouts ? " -notimeouts" : "";
 			ServerCommonArgs += IsAttachingDebugger(ProcessLaunchType.Server) ? " -WaitForDebuggerNoBreak" : "";
 			ServerCommonArgs += $" -MultiServerNumServers={NumTotalServers}";
 			ServerCommonArgs += (NonLocalServers.Length > 0) ? $" -MultiServerPeers={string.Join(',', NonLocalServers)}" : "";
@@ -404,7 +466,14 @@ Example running 4 clients locally connecting to single proxy server w/ a non-loc
 			return LocalGameServerAddresses;
 		}
 
-		private void StartProcessesForClients(ProcessLaunchType ClientType, string[] GameServerAddresses, string AdditionalClientArguments, int ClientsPerServer = 1)
+		/// <summary>
+		/// Start the Client processes
+		/// </summary>
+		/// <param name="ClientType">What type of Client we're launching which will help set the titles properly</param>
+		/// <param name="GameServerAddresses">The game servers (or proxies) that a client will connect to.  There will be ClientsPerServer launched per entry</param>
+		/// <param name="AdditionalClientArguments">What additional parameters we launch the client with</param>
+		/// <param name="ClientsPerServer">The number of clients to launch per-GameServerAddress entry</param>
+		protected void StartProcessesForClients(ProcessLaunchType ClientType, string[] GameServerAddresses, string AdditionalClientArguments, int ClientsPerServer = 1)
 		{
 			for (int Idx= 0; Idx < GameServerAddresses.Length; ++Idx)
 			{
@@ -429,7 +498,7 @@ Example running 4 clients locally connecting to single proxy server w/ a non-loc
 		{
 			string CommonProxyServerArgs = AdditionalProxyArguments;
 			CommonProxyServerArgs += " -NetDriverOverrides=/Script/MultiServerReplication.ProxyNetDriver";
-			CommonProxyServerArgs += ParseParam("notimeouts") ? " -notimeouts" : "";
+			CommonProxyServerArgs += NoTimeouts ? " -notimeouts" : "";
 			CommonProxyServerArgs += IsAttachingDebugger(ProcessLaunchType.ProxyServer) ? " -WaitForDebuggerNoBreak" : "";
 
 			if (AllGameServerAddresses.Length > 0)
@@ -437,12 +506,12 @@ Example running 4 clients locally connecting to single proxy server w/ a non-loc
 				CommonProxyServerArgs += $" -ProxyGameServers={string.Join(',', AllGameServerAddresses)}";
 			}
 
-			if (ParseParam("proxyclientprimary"))
+			if (ProxyClientPrimary > 0)
 			{
-				CommonProxyServerArgs += $" -ProxyClientPrimaryGameServer={ParseParamInt("proxyclientprimary", 0)}";
+				CommonProxyServerArgs += $" -ProxyClientPrimaryGameServer={ProxyClientPrimary}";
 			}
 
-			if (ParseParam("proxycycleprimary"))
+			if (ProxyCyclePrimary)
 			{
 				CommonProxyServerArgs += " -ProxyCyclePrimaryGameServer";
 			}
@@ -469,7 +538,14 @@ Example running 4 clients locally connecting to single proxy server w/ a non-loc
 		}
 
 
-		private void ExecuteApp(ProcessLaunchType LaunchType, string AppTitle, string AdditionalArguments)
+		/// <summary>
+		/// Execute the specified LaunchType app with a particular window AppTitle and arguments
+		/// </summary>
+		/// <param name="LaunchType">The type of app to launch</param>
+		/// <param name="AppTitle">The title of the window (we will try to name the console windows & game windows this)</param>
+		/// <param name="AdditionalArguments">Any additional arguments to pass on the command-line</param>
+		/// <exception cref="AutomationException"></exception>
+		protected void ExecuteApp(ProcessLaunchType LaunchType, string AppTitle, string AdditionalArguments)
 		{
 			const int PauseBetweenProcessMS = 100;
 
@@ -481,7 +557,7 @@ Example running 4 clients locally connecting to single proxy server w/ a non-loc
 				Client: LaunchType == ProcessLaunchType.DirectClient || LaunchType == ProcessLaunchType.ProxyClient || LaunchType == ProcessLaunchType.ProxyBot
 			);
 
-			var DeployContexts = Project.CreateDeploymentContext(Params, Params.DedicatedServer);
+			var DeployContexts = AutomationScripts.Project.CreateDeploymentContext(Params, Params.DedicatedServer);
 			if (DeployContexts.Count == 0)
 			{
 				throw new AutomationException("No DeployContexts for launching a process.");
@@ -502,13 +578,8 @@ Example running 4 clients locally connecting to single proxy server w/ a non-loc
 
 			// Use this option when logging with VeryVerbose to avoid unusable debugging windows (too much spew)
 			// The logging will still go to the output files (and you can use a text editor like Notepad++ to auto-reload them)
-			Args += ParseParam("noconsole") ? " -nodebugoutput -noconsole" : "";
-			Args += ParseParam("nonewconsole") ? "" : " -newconsole -log ";
-
-			// Add any logging in the commands
-			var Cats = new[] { "LogDstm", "LogDstmObject", "LogDstmDemo", "LogRemoteExec", "LogRemoteObject", "LogDstmPushMigration" };
-			Args += ParseParam("logdebug")	? $" -LogCmds=\"{string.Join("=Verbose,",	Cats)}\"" : "";
-			Args += ParseParam("logdebugv") ? $" -LogCmds=\"{string.Join("=VeryVerbose,",Cats)}\"" : "";
+			Args += NoConsole ? " -nodebugoutput -noconsole" : " -log";
+			Args += NoNewConsole ? "" : " -newconsole ";
 
 			Args += String.Format(" -ConsoleTitle=\"{0} {1}\" -SessionName=\"{0} {1}\"", Params.ShortProjectName, AppTitle);
 			Args += AdditionalArguments;
@@ -543,10 +614,15 @@ Example running 4 clients locally connecting to single proxy server w/ a non-loc
 			}
 		}
 
-		private bool IsAttachingDebugger(ProcessLaunchType LaunchType)
+		/// <summary>
+		/// Are we requesting a debugger be attached to this particular process?
+		/// </summary>
+		/// <param name="LaunchType">The type of process we plan to launch</param>
+		/// <returns>true if the commandlet parameters have specified we attach a debugger to this process</returns>
+		protected bool IsAttachingDebugger(ProcessLaunchType LaunchType)
 		{
-			return (ParseParam("attachtoservers") && (LaunchType == ProcessLaunchType.Server || LaunchType == ProcessLaunchType.ProxyServer)) ||
-				(ParseParam("attachtoclients") && (LaunchType == ProcessLaunchType.DirectClient || LaunchType == ProcessLaunchType.ProxyClient));
+			return (AttachToServers && (LaunchType == ProcessLaunchType.Server || LaunchType == ProcessLaunchType.ProxyServer)) ||
+				(AttachToClients && (LaunchType == ProcessLaunchType.DirectClient || LaunchType == ProcessLaunchType.ProxyClient));
 		}
 
 		private FileReference ProjectFullPath;
