Change Gauntlet logs parser to use InvariantCulture settings
#tests Floating point separator doesn't cause an exception regardless of which one is used
#rnx

diff --git a/Engine/Source/Programs/AutomationTool/Gauntlet/Unreal/Utils/Gauntlet.UnrealSnapshotParser.cs b/Engine/Source/Programs/AutomationTool/Gauntlet/Unreal/Utils/Gauntlet.UnrealSnapshotParser.cs
--- a/Engine/Source/Programs/AutomationTool/Gauntlet/Unreal/Utils/Gauntlet.UnrealSnapshotParser.cs
+++ b/Engine/Source/Programs/AutomationTool/Gauntlet/Unreal/Utils/Gauntlet.UnrealSnapshotParser.cs
@@ -5,6 +5,7 @@ using System.Collections.Generic;
 using System.Linq;
 using System.Text;
 using System.Threading.Tasks;
+using System.Globalization;
 
 namespace Gauntlet
 {
@@ -142,80 +143,90 @@ namespace Gauntlet
 
 			RegexUtil.MatchAndApplyGroups(InContent, @"MeasuredPerfTime:\s(\d.+) S", (Groups) =>
 			{
-				ProfileLength = Convert.ToSingle(Groups[1]);
+				ProfileLength = ParseSingleInvariant(Groups[1]);
 			});
 
 			RegexUtil.MatchAndApplyGroups(InContent, @"CPU Memory:[\s\w]+?([\d\.]+)MB,[\s\w]+?([\d\.]+)MB", (Groups) =>
 			{
-				CpuUsedMemory = Convert.ToSingle(Groups[1]);
-				CpuPeakMemory = Convert.ToSingle(Groups[2]);
+				CpuUsedMemory = ParseSingleInvariant(Groups[1]);
+				CpuPeakMemory = ParseSingleInvariant(Groups[2]);
 			});
 
 			RegexUtil.MatchAndApplyGroups(InContent, @"Physical Memory:[\s\w]+?([\d\.]+)MB,[\s\w:]+?([\d\.]+)MB", (Groups) =>
 			{
-				PhysicalUsedMemory = Convert.ToSingle(Groups[1]);
-				PhysicalPeakMemory = Convert.ToSingle(Groups[2]);
+				PhysicalUsedMemory = ParseSingleInvariant(Groups[1]);
+				PhysicalPeakMemory = ParseSingleInvariant(Groups[2]);
 			});
 
 			RegexUtil.MatchAndApplyGroups(InContent, @"Streaming Memory:[\s\w]+?([\d\.]+)MB,[\s\w:]+?([\d\.]+)MB", (Groups) =>
 			{
-				StreamingUsedMemory = Convert.ToSingle(Groups[1]);
-				StreamingPeakMemory = Convert.ToSingle(Groups[2]);
+				StreamingUsedMemory = ParseSingleInvariant(Groups[1]);
+				StreamingPeakMemory = ParseSingleInvariant(Groups[2]);
 			});
 
 			RegexUtil.MatchAndApplyGroups(InContent, @"MVP:\s(\d.+)%,\s.*AvgFPS:(\d.+),\s.*HitchesPerMinute:\s(\d.+),\sAvg\sHitch\s(\d.+)ms", (Groups) =>
 			{
-				MVP = Convert.ToSingle(Groups[1]);
-				AvgFps = Convert.ToSingle(Groups[2]);
-				Hitches = Convert.ToSingle(Groups[3]);
-				AvgHitches = Convert.ToSingle(Groups[4]);
+				MVP = ParseSingleInvariant(Groups[1]);
+				AvgFps = ParseSingleInvariant(Groups[2]);
+				Hitches = ParseSingleInvariant(Groups[3]);
+				AvgHitches = ParseSingleInvariant(Groups[4]);
 			});
 
 			RegexUtil.MatchAndApplyGroups(InContent, @"GT:.* Avg\s(\d.+)ms,", (Groups) =>
 			{
-				GTTime = Convert.ToSingle(Groups[1]);
+				GTTime = ParseSingleInvariant(Groups[1]);
 			});
 
 			RegexUtil.MatchAndApplyGroups(InContent, @"RT:.* Avg\s(\d.+)ms,", (Groups) =>
 			{
-				RTTime = Convert.ToSingle(Groups[1]);
+				RTTime = ParseSingleInvariant(Groups[1]);
 			});
 
 			RegexUtil.MatchAndApplyGroups(InContent, @"GPU:.* Avg\s(\d.+)ms,", (Groups) =>
 			{
-				GPUTime = Convert.ToSingle(Groups[1]);
+				GPUTime = ParseSingleInvariant(Groups[1]);
 			});
 
 			RegexUtil.MatchAndApplyGroups(InContent, @"FT:.* Avg:\s([\d\.]+)ms,", (Groups) =>
 			{
-				FTTime = Convert.ToSingle(Groups[1]);
+				FTTime = ParseSingleInvariant(Groups[1]);
 			});
 
 			RegexUtil.MatchAndApplyGroups(InContent, @"RHIT:.* Avg\s([\d\.]+)ms,", (Groups) =>
 			{
-				RHIT = Convert.ToSingle(Groups[1]);
+				RHIT = ParseSingleInvariant(Groups[1]);
 			});
 
 			RegexUtil.MatchAndApplyGroups(InContent, @"DynRes:.* Avg:\s([\d\.]+)%,", (Groups) =>
 			{
-				DynamicRes = Convert.ToSingle(Groups[1]);
+				DynamicRes = ParseSingleInvariant(Groups[1]);
 			});
 
 			RegexUtil.MatchAndApplyGroups(InContent, @"DrawCalls:[\w\s:]+?([\d\.]+),[\w\s:]+?([\d\.]+),[\w\s:]+?([\d\.]+)", (Groups) =>
 			{
-				DrawCalls = Convert.ToInt32(Groups[1]);
+				DrawCalls = ParseIntInvariant(Groups[1]);
 			});
 
 			RegexUtil.MatchAndApplyGroups(InContent, @"DrawnPrims:[\w\s:]+?([\d\.]+),[\w\s:]+?([\d\.]+),[\w\s:]+?([\d\.]+)", (Groups) =>
 			{
-				DrawnPrims = Convert.ToInt32(Groups[1]);
+				DrawnPrims = ParseIntInvariant(Groups[1]);
 			});
 
 			RegexUtil.MatchAndApplyGroups(InContent, @"UnbuiltHLODs:\s(\d.+)", (Groups) =>
 			{
-				UnbuiltHLODs = Convert.ToInt32(Groups[1]);
+				UnbuiltHLODs = ParseIntInvariant(Groups[1]);
 			});
 		}
+
+		private static float ParseSingleInvariant(string Text)
+		{
+			return Single.Parse(Text, CultureInfo.InvariantCulture);
+		}
+		
+		private static int ParseIntInvariant(string Text)
+		{
+			return int.Parse(Text, CultureInfo.InvariantCulture);
+		}
 	}
 
 	public class UnrealSnapshotSummary<TSnapshotClass>
