Bugfix for LaunchMultiServer script:  Defaulting to 2 Servers correctly.

#rnx

diff --git a/Engine/Source/Programs/AutomationTool/Scripts/LaunchMultiServer.Automation.cs b/Engine/Source/Programs/AutomationTool/Scripts/LaunchMultiServer.Automation.cs
--- a/Engine/Source/Programs/AutomationTool/Scripts/LaunchMultiServer.Automation.cs
+++ b/Engine/Source/Programs/AutomationTool/Scripts/LaunchMultiServer.Automation.cs
@@ -243,10 +243,10 @@ Example running 4 clients locally connecting to single proxy server w/ a non-loc
 		public int BaseGameListenPort = 9000;
 		public int ProxyServerBasePort = -1;
 
-		[CommandLine(), Help("CustomConfig=<Section>", "Read a different section of the Project's GameConfig for default values")]
+		[CommandLine, Help("CustomConfig=<Section>", "Read a different section of the Project's GameConfig for default values")]
 		public string CustomConfig = "";
 
-		[CommandLine(), Help("NumServers=##", "Set the fleet size of servers (default: 2)")]
+		[CommandLine, Help("NumServers=##", "Set the fleet size of servers (default: 2)")]
 		public int NumServers = -1;
 
 		[CommandLine(ListSeparator = ','), Help("NonLocalServers=Address1:Port1,Address2:Port2,etc.", "Specify servers that are not hosted locally but belong to the same fleet")]
@@ -255,37 +255,37 @@ Example running 4 clients locally connecting to single proxy server w/ a non-loc
 		[CommandLine("-Client", ListSeparator = ','), Help("Client=<Type>", "Specify the client type to connect with (proxy or direct or proxy,direct)")]
 		public List<string> ClientType = new List<string>();
 
-		[CommandLine(), Help("ProxyServerCount=##", "Specify how many proxy servers to launch in client=proxy mode")]
+		[CommandLine, Help("ProxyServerCount=##", "Specify how many proxy servers to launch in client=proxy mode")]
 		public int ProxyServerCount = 1;
 
-		[CommandLine(), Help("ProxyClientCount=##", "Specify how many clients to launch PER SERVER in client=proxy mode")]
+		[CommandLine, Help("ProxyClientCount=##", "Specify how many clients to launch PER SERVER in client=proxy mode")]
 		public int ProxyClientCount = 1;
 
-		[CommandLine(), Help("ProxyBotCount=##", "Specify how many bots to launch PER SERVER in client=proxy mode")]
+		[CommandLine, Help("ProxyBotCount=##", "Specify how many bots to launch PER SERVER in client=proxy mode")]
 		public int ProxyBotCount = 0;
 
-		[CommandLine(), Help("ProxyClientPrimary=##", "Specify which game server to prefer as the initial proxy connection")]
+		[CommandLine, Help("ProxyClientPrimary=##", "Specify which game server to prefer as the initial proxy connection")]
 		public int ProxyClientPrimary = -1;
 
-		[CommandLine(), Help("ProxyCyclePrimary", "Specify that each subsequent proxy client should choose a different server")]
+		[CommandLine, Help("ProxyCyclePrimary", "Specify that each subsequent proxy client should choose a different server")]
 		public bool ProxyCyclePrimary = false;
 
-		[CommandLine(), Help("NoTimeouts", "Disable timeouts (recommended for debugging)")]
+		[CommandLine, Help("NoTimeouts", "Disable timeouts (recommended for debugging)")]
 		public bool NoTimeouts = false;
 
-		[CommandLine(), Help("NoConsole", "Disable launching a console (not recommended, but may improve performance)")]
+		[CommandLine, Help("NoConsole", "Disable launching a console (not recommended, but may improve performance)")]
 		public bool NoConsole = false;
 
-		[CommandLine(), Help("NoNewConsole", "Disable launching the new style console")]
+		[CommandLine, Help("NoNewConsole", "Disable launching the new style console")]
 		public bool NoNewConsole = false;
 
-		[CommandLine(), Help("AttachToServers", "Force Visual Studio to attach to the servers when launched (requires VS be launched or it will hang)")]
+		[CommandLine, Help("AttachToServers", "Force Visual Studio to attach to the servers when launched (requires VS be launched or it will hang)")]
 		public bool AttachToServers = false;
 
-		[CommandLine(), Help("AttachToClients", "Force Visual Studio to attach to the clients when launched (requires VS be launched or it will hang)")]
+		[CommandLine, Help("AttachToClients", "Force Visual Studio to attach to the clients when launched (requires VS be launched or it will hang)")]
 		public bool AttachToClients = false;
 
-		[CommandLine(), Help("CommonArgs", "Specify arguments that are run by all processes")]
+		[CommandLine, Help("CommonArgs", "Specify arguments that are run by all processes")]
 		public string CommonArgs = "";
 
 		// We should only parse the commandline arguments once, since doing so multiple times will add entries to the List<> args.
@@ -400,11 +400,15 @@ Example running 4 clients locally connecting to single proxy server w/ a non-loc
 			}
 
 			int TotalNumServers = NumServers;
-			if (TotalNumServers < 0 && !ProjectGameConfig.TryGetValue(ServerDefConfigSection, "TotalNumServers", out TotalNumServers))
+			if (TotalNumServers < 0)
 			{
-				NumServers = 2;
-				Logger.LogWarning($"Could not find {ServerDefConfigSection}.TotalNumServers in Game Config File {ProjectGameConfig} Defaulting to {NumServers}");
+				if (!ProjectGameConfig.TryGetValue(ServerDefConfigSection, "TotalNumServers", out TotalNumServers))
+				{
+					TotalNumServers = 2;
+					Logger.LogWarning($"Could not find {ServerDefConfigSection}.TotalNumServers in Game Config File {ProjectGameConfig} Defaulting to {TotalNumServers}");
+				}
 			}
+			NumServers = TotalNumServers;
 		}
 
 		/// <summary>
