Enable merged monolithic and patch builds.
#rnx

#tests local, preflight

diff --git a/Engine/Source/Programs/AutomationTool/AutomationUtils/Platform.cs b/Engine/Source/Programs/AutomationTool/AutomationUtils/Platform.cs
--- a/Engine/Source/Programs/AutomationTool/AutomationUtils/Platform.cs
+++ b/Engine/Source/Programs/AutomationTool/AutomationUtils/Platform.cs
@@ -848,7 +848,12 @@ namespace AutomationTool
 
 		public virtual bool UpdatePatchPackagingParameters( string ProjectRoot, string BuildToUse, string AdditionalOptions )
 		{
-			Logger.LogWarning("UpdatePackagingParameters() has not been implemented for {Arg0}", PlatformType.ToString());
+			Logger.LogWarning("UpdatePackagingParameters() not needed on platform {Arg0}", PlatformType.ToString());
+			return false;
+		}
+		public virtual bool BuildAdditionalApps(string ProjectRoot, string Configuration, string OutputDir, string AdditionalOptions)
+		{
+			Logger.LogWarning("BuildAdditionalApps() not needed on platform {Arg0}", PlatformType.ToString());
 			return false;
 		}
 
