Faster inclusion of Verse source files in cooked server image

#rnx

diff --git a/Engine/Source/Programs/AutomationTool/Scripts/CopyBuildToStagingDirectory.Automation.cs b/Engine/Source/Programs/AutomationTool/Scripts/CopyBuildToStagingDirectory.Automation.cs
--- a/Engine/Source/Programs/AutomationTool/Scripts/CopyBuildToStagingDirectory.Automation.cs
+++ b/Engine/Source/Programs/AutomationTool/Scripts/CopyBuildToStagingDirectory.Automation.cs
@@ -1687,32 +1687,34 @@ namespace AutomationScripts
 									}
 								}
 							}
-						}
 
-						// Stage Verse source files?
-						if (bStageVerseSourceFiles)
-						{
-							DirectoryReference PluginContentDirectory = DirectoryReference.Combine(StagedPlugin.Value.Directory, "Content");
-							if (DirectoryReference.Exists(PluginContentDirectory))
+							// Stage Verse source files?
+							if (bStageVerseSourceFiles && Info.Descriptor.bCanContainVerse)
 							{
-								// Include Verse source files from the Content folder
-								StagedDirectoryReference StagedContentDirecory = StagedDirectoryReference.Combine(StagedPlugin.Key.Directory, "Content");
-								SC.StageFiles(StagedFileType.NonUFS, PluginContentDirectory, "*.verse", StageFilesSearch.AllDirectories, StagedContentDirecory);
-								SC.StageFiles(StagedFileType.NonUFS, PluginContentDirectory, "*.vmodule", StageFilesSearch.AllDirectories, StagedContentDirecory);
-							}
+								if (Info.Descriptor.VersePath != null)
+								{
+									DirectoryReference PluginContentDirectory = DirectoryReference.Combine(StagedPlugin.Value.Directory, "Content");
+									if (DirectoryReference.Exists(PluginContentDirectory))
+									{
+										// Include Verse source files from the Content folder
+										List<FileReference> ContentFiles = SC.FindFilesToStage(PluginContentDirectory, StageFilesSearch.AllDirectories);
+										SC.StageFiles(StagedFileType.NonUFS, ContentFiles.Where(x => x.HasExtension(".verse") || x.HasExtension(".vmodule")));
+									}
+								}
 
-							// Also include Verse source files of C++ modules
-							DirectoryReference PluginSourceDirectory = DirectoryReference.Combine(StagedPlugin.Value.Directory, "Source");
-							if (DirectoryReference.Exists(PluginSourceDirectory))
-							{
-								foreach (FileReference BuildCsFile in DirectoryReference.EnumerateFiles(PluginSourceDirectory, "*.Build.cs", SearchOption.AllDirectories))
+								// Also include Verse source files of C++ modules
+								DirectoryReference PluginSourceDirectory = DirectoryReference.Combine(StagedPlugin.Value.Directory, "Source");
+								if (DirectoryReference.Exists(PluginSourceDirectory))
 								{
-									DirectoryReference MaybeVerseDirectory = DirectoryReference.Combine(DirectoryReference.FromFile(BuildCsFile), "Verse");
-									if (DirectoryReference.Exists(MaybeVerseDirectory))
+									foreach (FileReference BuildCsFile in DirectoryReference.EnumerateFiles(PluginSourceDirectory, "*.Build.cs", SearchOption.AllDirectories))
 									{
-										string RelativeVerseDirectory = MaybeVerseDirectory.MakeRelativeTo(StagedPlugin.Value.Directory);
-										StagedDirectoryReference StagedVerseDirecory = StagedDirectoryReference.Combine(StagedPlugin.Key.Directory, RelativeVerseDirectory);
-										SC.StageFiles(StagedFileType.NonUFS, MaybeVerseDirectory, "*.*", StageFilesSearch.AllDirectories, StagedVerseDirecory);
+										DirectoryReference MaybeVerseDirectory = DirectoryReference.Combine(DirectoryReference.FromFile(BuildCsFile), "Verse");
+										if (DirectoryReference.Exists(MaybeVerseDirectory))
+										{
+											string RelativeVerseDirectory = MaybeVerseDirectory.MakeRelativeTo(StagedPlugin.Value.Directory);
+											StagedDirectoryReference StagedVerseDirecory = StagedDirectoryReference.Combine(StagedPlugin.Key.Directory, RelativeVerseDirectory);
+											SC.StageFiles(StagedFileType.NonUFS, MaybeVerseDirectory, "*.*", StageFilesSearch.AllDirectories, StagedVerseDirecory);
+										}
 									}
 								}
 							}
