Make sure device cache path use a valid folder name

#rnx

diff --git a/Engine/Source/Programs/AutomationTool/Gauntlet/Framework/Devices/Gauntlet.DevicePool.cs b/Engine/Source/Programs/AutomationTool/Gauntlet/Framework/Devices/Gauntlet.DevicePool.cs
--- a/Engine/Source/Programs/AutomationTool/Gauntlet/Framework/Devices/Gauntlet.DevicePool.cs
+++ b/Engine/Source/Programs/AutomationTool/Gauntlet/Framework/Devices/Gauntlet.DevicePool.cs
@@ -1214,10 +1214,11 @@ namespace Gauntlet
 		/// </summary>
 		private string GetCleanCachePath(DeviceDefinition InDeviceDefiniton)
 		{
+			string SanitizedDeviceName = FileUtils.SanitizeFilename(InDeviceDefiniton.Name);
 			// Give the desktop platform a temp folder with its name under the device cache
 			string DeviceCache = Path.Combine(LocalTempDir, "DeviceCache");
 			string PlatformCache = Path.Combine(DeviceCache, InDeviceDefiniton.Platform.ToString());
-			string ClientCache = Path.Combine(PlatformCache, InDeviceDefiniton.Name);
+			string ClientCache = Path.Combine(PlatformCache, SanitizedDeviceName);
 
 			// On Desktops, builds are installed in the device cache.
 			// When using device reservation blocks, we don't want to fully clean the cache and lose previously installed builds.
@@ -1272,7 +1273,7 @@ namespace Gauntlet
 							string Warning = "Failed to clean client device cache {Folder}. A newly indexed directory will be created instead. {Message}";
 							Log.Warning(KnownLogEvents.Gauntlet_DeviceEvent, Warning, ClientCache, Ex.Message);
 
-							ClientCache = Path.Combine(PlatformCache, $"{InDeviceDefiniton.Name}_{++CleanAttempts}");
+							ClientCache = Path.Combine(PlatformCache, $"{SanitizedDeviceName}_{++CleanAttempts}");
 						}
 					}
 				}
