Move RedirectToShop app build out of game agnostic code where it should be. Don't reference derived platforms in base platform code.

#rnx
#test local, preflights

diff --git a/Engine/Source/Programs/AutomationTool/AutomationUtils/Platform.cs b/Engine/Source/Programs/AutomationTool/AutomationUtils/Platform.cs
--- a/Engine/Source/Programs/AutomationTool/AutomationUtils/Platform.cs
+++ b/Engine/Source/Programs/AutomationTool/AutomationUtils/Platform.cs
@@ -851,11 +851,6 @@ namespace AutomationTool
 			Logger.LogWarning("UpdatePackagingParameters() not needed on platform {Arg0}", PlatformType.ToString());
 			return false;
 		}
-		public virtual bool BuildAdditionalApps(string ProjectRoot, string Configuration, string OutputDir, string AdditionalOptions)
-		{
-			Logger.LogWarning("BuildAdditionalApps() not needed on platform {Arg0}", PlatformType.ToString());
-			return false;
-		}
 
 		/// <summary>
 		/// When overridden, returns the directory structure of the platform's symbol server.
