Adding 'AssetRegistry' channel in order to prevent timeouts from inactivity while the AssetRegistry cache is being generated

#rnx

diff --git a/Engine/Source/Programs/AutomationTool/Gauntlet/Unreal/Utils/Gauntlet.UnrealLogParser.cs b/Engine/Source/Programs/AutomationTool/Gauntlet/Unreal/Utils/Gauntlet.UnrealLogParser.cs
--- a/Engine/Source/Programs/AutomationTool/Gauntlet/Unreal/Utils/Gauntlet.UnrealLogParser.cs
+++ b/Engine/Source/Programs/AutomationTool/Gauntlet/Unreal/Utils/Gauntlet.UnrealLogParser.cs
@@ -36,7 +36,7 @@ namespace Gauntlet
 		/// <summary>
 		/// Set of log channels that are used to monitor Editor processing
 		/// </summary>
-		public static string[] EditorBusyChannels = new string[] { "Automation", "FunctionalTest", "Material", "DerivedDataCache", "ShaderCompilers", "Texture", "SkeletalMesh", "StaticMesh", "Python" };
+		public static string[] EditorBusyChannels = new string[] { "Automation", "FunctionalTest", "Material", "DerivedDataCache", "ShaderCompilers", "Texture", "SkeletalMesh", "StaticMesh", "Python", "AssetRegistry" };
 
 		/// <summary>
 		/// Represents an entry in an Unreal logfile with and contails the associated category, level, and message
