ZenLoader - Quick Launch Automation (Part 3)

1. Consoles support
2. Zen streaming combinations are snapshot-based and use the remote zen server on the rdu-zen-01.epicgames.net in case of consoles
3. Ability to publish staged builds in Horde
4. Unstable test behavior eliminated

#rnx

diff --git a/Engine/Source/Programs/AutomationTool/Gauntlet/Unreal/Automation/UE.ZenLoaderTest.cs b/Engine/Source/Programs/AutomationTool/Gauntlet/Unreal/Automation/UE.ZenLoaderTest.cs
--- a/Engine/Source/Programs/AutomationTool/Gauntlet/Unreal/Automation/UE.ZenLoaderTest.cs
+++ b/Engine/Source/Programs/AutomationTool/Gauntlet/Unreal/Automation/UE.ZenLoaderTest.cs
@@ -57,6 +57,7 @@ namespace UE
 			out string ExitReason, out int ExitCode)
 		{
 			InitStagingOutputConditions();
+			InLog.EngineInitializedPattern = "Starting Game";
 
 			if (!IsStagingOutputValid())
 			{
