Allow job step to link to test results

#rnx

diff --git a/Engine/Source/Programs/AutomationTool/Gauntlet/Framework/Horde/Gauntlet.Horde.cs b/Engine/Source/Programs/AutomationTool/Gauntlet/Framework/Horde/Gauntlet.Horde.cs
--- a/Engine/Source/Programs/AutomationTool/Gauntlet/Framework/Horde/Gauntlet.Horde.cs
+++ b/Engine/Source/Programs/AutomationTool/Gauntlet/Framework/Horde/Gauntlet.Horde.cs
@@ -3,7 +3,9 @@
 using AutomationTool;
 using System;
 using System.IO;
+using System.Linq;
 using System.Text;
+using System.Text.Json;
 
 namespace Gauntlet
 {
@@ -33,6 +35,13 @@ namespace Gauntlet
 				return Environment.GetEnvironmentVariable("UE_HORDE_STEPID");
 			}
 		}
+		public class JobStepReport
+		{
+			public string scope { get; } = "Step";
+			public string name { get; set; }
+			public string placement { get; } = "Summary";
+			public string fileName { get; set; }
+		}
 
 		public static void GenerateSummary(string ArtifactPath)
 		{
@@ -47,38 +56,21 @@ namespace Gauntlet
 
 				string MarkdownFilename = "GauntletStepDetails.md";
 
+				StringBuilder Markdown = new StringBuilder();
+
 				// file:// is attractive here, and file://\\computer\share\folder\file.txt does work ok in Chrome, but it does not Firefox
 				// so, lacking a uniform way to link to a network filesystem, just put the path here as static text and escaped for markdown
-				string Markdown = System.Text.RegularExpressions.Regex.Replace($"Gauntlet Artifacts: {ArtifactPath}", @"([|\\*\+])", @"\$1");
-
-				File.WriteAllText(Path.Combine(LogFolder, MarkdownFilename), Markdown);
-
-				StringBuilder Builder = new StringBuilder();
-
-				Builder.Append("{\n");
-
-				Builder.Append("\"scope\": \"Step\",\n");
-				Builder.Append("\"name\": \"Gauntlet Step Details\",\n");
-				Builder.Append("\"placement\": \"Summary\",\n");
-				Builder.AppendFormat("\"fileName\": \"{0}\"\n", MarkdownFilename);
-
-				Builder.Append("}\n");
-
-				File.WriteAllText(Path.Combine(LogFolder, "GauntletStepDetails.report.json"), Builder.ToString());
-
-				/*
-				using (JsonWriter Writer = new JsonWriter(new FileReference(Path.Combine(LogFolder, "GauntletStepDetails.report.json"))))
+				Markdown.Append($"* **Netshare Artifacts**: {System.Text.RegularExpressions.Regex.Replace(ArtifactPath, @"([|\\*\+])", @"\$1")}");
+				// detect if any testdata was generated
+				if (Globals.Params.ParseParam("UseTestDataV2")
+					&& Directory.GetFiles(Path.Combine(CommandUtils.CmdEnv.EngineSavedFolder, "TestData"), "*.json", SearchOption.TopDirectoryOnly).Any())
 				{
-					Writer.WriteObjectStart();
-
-					Writer.WriteValue("scope", "Step");
-					Writer.WriteValue("name", "Gauntlet Step Details");
-					Writer.WriteValue("placement", "Summary");
-					Writer.WriteValue("fileName", MarkdownFilename);
-
-					Writer.WriteObjectEnd();
+					Markdown.Append($"\n* **Test Results**: [Open report](/test-automation?job={JobId}&step={StepId})");
 				}
-				*/
+				File.WriteAllText(Path.Combine(LogFolder, MarkdownFilename), Markdown.ToString());
+
+				JobStepReport report = new JobStepReport() { name="Gauntlet Step Details", fileName=MarkdownFilename };
+				File.WriteAllText(Path.Combine(LogFolder, "GauntletStepDetails.report.json"), JsonSerializer.Serialize(report));
 			}
 			catch (Exception Ex)
 			{
diff --git a/Engine/Source/Programs/Horde/Plugins/Build/HordeServer.Build.Tests/Jobs/TestData/TestDataTests.cs b/Engine/Source/Programs/Horde/Plugins/Build/HordeServer.Build.Tests/Jobs/TestData/TestDataTests.cs
--- a/Engine/Source/Programs/Horde/Plugins/Build/HordeServer.Build.Tests/Jobs/TestData/TestDataTests.cs
+++ b/Engine/Source/Programs/Horde/Plugins/Build/HordeServer.Build.Tests/Jobs/TestData/TestDataTests.cs
@@ -454,6 +454,13 @@ namespace HordeServer.Tests.Jobs.TestData
 			ActionResult<object> testdata =  await TestDataController.GetTestDataV2Async(sessions[0].TestDataId.ToString());
 			Assert.IsNotNull(testdata);
 			Assert.IsNotNull(testdata.Value);
+
+			ActionResult<List<GetJobStepTestDataIDResponse>> testDataIds = await TestDataController.GetJobStepTestDataV2Async(job3.Id.ToString(), step3.Id.ToString());
+			Assert.IsNotNull(testDataIds);
+			Assert.IsNotNull(testDataIds.Value);
+			Assert.AreEqual(3, testDataIds.Value.Count);
+			Assert.IsTrue(testDataIds.Value.All(item => sessions.Any(s => s.TestDataId == item.TestDataId)));
+			Assert.IsTrue(testDataIds.Value.All(item => sessions.First(s => s.TestDataId == item.TestDataId).NameRef == item.TestNameRef));
 		}
 
 		private readonly string[] _simpleTestDataLines =
diff --git a/Engine/Source/Programs/Horde/Plugins/Build/HordeServer.Build/Jobs/TestData/ITestDataCollection.cs b/Engine/Source/Programs/Horde/Plugins/Build/HordeServer.Build/Jobs/TestData/ITestDataCollection.cs
--- a/Engine/Source/Programs/Horde/Plugins/Build/HordeServer.Build/Jobs/TestData/ITestDataCollection.cs
+++ b/Engine/Source/Programs/Horde/Plugins/Build/HordeServer.Build/Jobs/TestData/ITestDataCollection.cs
@@ -1,5 +1,6 @@
 // Copyright Epic Games, Inc. All Rights Reserved.
 
+using System.Threading;
 using EpicGames.Horde.Commits;
 using EpicGames.Horde.Jobs;
 using EpicGames.Horde.Streams;
@@ -182,6 +183,15 @@ namespace HordeServer.Jobs.TestData
 		/// <returns>The new testdata document</returns>
 		Task<IReadOnlyList<ITestData>> AddAsync(IJob job, IJobStep step, (string key, BsonDocument value)[] data, CancellationToken cancellationToken = default);
 
+		/// <summary>
+		/// Gets testdata ids run within a job step
+		/// </summary>
+		/// <param name="jobId">Unique id of the job</param>
+		/// <param name="stepId">Unique id of the step</param>
+		/// <param name="cancellationToken">Cancellation token for the operation</param>
+		/// <returns>The testdata ids</returns>
+		Task<IReadOnlyList<ITestData>> GetJobStepTestDataAsync(JobId jobId, JobStepId stepId, CancellationToken cancellationToken = default);
+
 		/// <summary>
 		/// Gets a testdata by ID
 		/// </summary>
diff --git a/Engine/Source/Programs/Horde/Plugins/Build/HordeServer.Build/Jobs/TestData/TestDataCollectionV2.cs b/Engine/Source/Programs/Horde/Plugins/Build/HordeServer.Build/Jobs/TestData/TestDataCollectionV2.cs
--- a/Engine/Source/Programs/Horde/Plugins/Build/HordeServer.Build/Jobs/TestData/TestDataCollectionV2.cs
+++ b/Engine/Source/Programs/Horde/Plugins/Build/HordeServer.Build/Jobs/TestData/TestDataCollectionV2.cs
@@ -427,6 +427,7 @@ namespace HordeServer.Jobs.TestData
 			_logger = logger;
 
 			List<MongoIndex<TestDataDocument>> indexes = new List<MongoIndex<TestDataDocument>>();
+			indexes.Add(keys => keys.Descending(x => x.JobId).Descending(x => x.StepId));
 			_testDataDocuments = mongoService.GetCollection<TestDataDocument>("TestDataV2", indexes);
 
 			List<MongoIndex<TestMetadataDocument>> metaIndexes = new List<MongoIndex<TestMetadataDocument>>();
@@ -672,6 +673,12 @@ namespace HordeServer.Jobs.TestData
 			return documents;
 		}
 
+		/// <inheritdoc/>
+		public async Task<IReadOnlyList<ITestData>> GetJobStepTestDataAsync(JobId jobId, JobStepId stepId, CancellationToken cancellationToken)
+		{
+			return await _testDataDocuments.Find(x => x.JobId == jobId && x.StepId == stepId).ToListAsync(cancellationToken);
+		}
+
 		/// <inheritdoc/>
 		public async Task<ITestData?> GetAsync(ObjectId id, CancellationToken cancellationToken)
 		{
diff --git a/Engine/Source/Programs/Horde/Plugins/Build/HordeServer.Build/Jobs/TestData/TestDataController.cs b/Engine/Source/Programs/Horde/Plugins/Build/HordeServer.Build/Jobs/TestData/TestDataController.cs
--- a/Engine/Source/Programs/Horde/Plugins/Build/HordeServer.Build/Jobs/TestData/TestDataController.cs
+++ b/Engine/Source/Programs/Horde/Plugins/Build/HordeServer.Build/Jobs/TestData/TestDataController.cs
@@ -464,13 +464,41 @@ namespace HordeServer.Jobs.TestData
 			}, filter);
 		}
 
+		/// <summary>
+		/// Retrieve list of testdata ids run within a job step
+		/// </summary>
+		/// <param name="jobId">job Id from which the test was run in</param>
+		/// <param name="stepId">step Id from which the test was run in</param>
+		/// <param name="cancellationToken">Cancellation token for the operation</param>
+		/// <returns>List of testdata ids</returns>
+		[HttpGet]
+		[Route("/api/v3/testdata/")]
+		[ProducesResponseType(typeof(List<GetJobStepTestDataIDResponse>), 200)]
+		public async Task<ActionResult<List<GetJobStepTestDataIDResponse>>> GetJobStepTestDataV2Async([FromQuery] string jobId, [FromQuery] string stepId, CancellationToken cancellationToken = default)
+		{
+			if (!await _jobService.AuthorizeAsync(JobId.Parse(jobId), JobAclAction.ViewJob, User, _buildConfig.Value, cancellationToken))
+			{
+				return Forbid();
+			}
+
+			IReadOnlyList<ITestData> testData = await _testDataService.GetJobStepTestDataV2Async(JobId.Parse(jobId), JobStepId.Parse(stepId), cancellationToken);
+			IReadOnlyList<ITestNameRef> tests = await _testDataService.FindTestNameRefsV2Async(keys: testData.Select(x => x.Key).ToArray(), cancellationToken: cancellationToken);
+
+			return testData.Select(item => new GetJobStepTestDataIDResponse() {
+				TestDataId = item.Id.ToString(),
+				TestKey = item.Key,
+				TestName = tests.FirstOrDefault(t => t.Key == item.Key)?.Name ?? item.Key,
+				TestNameRef = tests.FirstOrDefault(t => t.Key == item.Key)?.Id.ToString() ?? item.Id.ToString(),
+			}).ToList();
+		}
+
 		/// <summary>
 		/// Retrieve information about a specific issue
 		/// </summary>
 		/// <param name="testDataId">Id of the document to get information about</param>
 		/// <param name="filter">Filter for the properties to return</param>
 		/// <param name="cancellationToken">Cancellation token for the operation</param>
-		/// <returns>List of matching agents</returns>
+		/// <returns>the testdata information</returns>
 		[HttpGet]
 		[Route("/api/v3/testdata/{testDataId}")]
 		[ProducesResponseType(typeof(GetTestDataResponse), 200)]
diff --git a/Engine/Source/Programs/Horde/Plugins/Build/HordeServer.Build/Jobs/TestData/TestDataService.cs b/Engine/Source/Programs/Horde/Plugins/Build/HordeServer.Build/Jobs/TestData/TestDataService.cs
--- a/Engine/Source/Programs/Horde/Plugins/Build/HordeServer.Build/Jobs/TestData/TestDataService.cs
+++ b/Engine/Source/Programs/Horde/Plugins/Build/HordeServer.Build/Jobs/TestData/TestDataService.cs
@@ -2,6 +2,7 @@
 
 using EpicGames.Core;
 using EpicGames.Horde.Commits;
+using EpicGames.Horde.Jobs;
 using EpicGames.Horde.Streams;
 using Microsoft.Extensions.Hosting;
 using Microsoft.Extensions.Logging;
@@ -174,6 +175,18 @@ namespace HordeServer.Jobs.TestData
 			return await _testData.FindTestDetailsAsync(ids, cancellationToken);
 		}
 
+		/// <summary>
+		/// Gets testdata ids run within a job step
+		/// </summary>
+		/// <param name="jobId">Unique id of the job</param>
+		/// <param name="stepId">Unique id of the step</param>
+		/// <param name="cancellationToken">Cancellation token for the operation</param>
+		/// <returns>The testdata ids</returns>
+		public async Task<IReadOnlyList<ITestData>> GetJobStepTestDataV2Async(JobId jobId, JobStepId stepId, CancellationToken cancellationToken)
+		{
+			return await _testDataV2.GetJobStepTestDataAsync(jobId, stepId, cancellationToken);
+		}
+
 		/// <summary>
 		/// Gets a testdata by ID
 		/// </summary>
diff --git a/Engine/Source/Programs/Shared/EpicGames.Horde/Jobs/TestData/TestDataMessages.cs b/Engine/Source/Programs/Shared/EpicGames.Horde/Jobs/TestData/TestDataMessages.cs
--- a/Engine/Source/Programs/Shared/EpicGames.Horde/Jobs/TestData/TestDataMessages.cs
+++ b/Engine/Source/Programs/Shared/EpicGames.Horde/Jobs/TestData/TestDataMessages.cs
@@ -112,6 +112,32 @@ namespace EpicGames.Horde.Jobs.TestData
 		}
 	}
 
+	/// <summary>
+	/// Response object containing the testdata id and the associated test
+	/// </summary>
+	public class GetJobStepTestDataIDResponse
+	{
+		/// <summary>
+		/// The unique testdata id
+		/// </summary>
+		public string TestDataId { get; set; } = String.Empty;
+
+		/// <summary>
+		/// The test key related to the testdata
+		/// </summary>
+		public string TestKey { get; set; } = String.Empty;
+
+		/// <summary>
+		/// The test name related to the testdata
+		/// </summary>
+		public string TestName { get; set; } = String.Empty;
+
+		/// <summary>
+		/// The test name reference id related to the testdata
+		/// </summary>
+		public string TestNameRef { get; set; } = String.Empty;
+	}
+
 	/// <summary>
 	/// Response object describing test results
 	/// </summary>
