Port the Cooker Test Suite to Mac/Linux

1. Fixing the slash issue in the UE.CookOnTheFly test
2. Creation of platform-dependent providers and a factory for the cooked content as each platform has specific cooked content.
3. Covers the test cases: UE-TC-20206, UE-TC-20067, UE-TC-20066, UE-TC-20065, UE-TC-20064

#rnx

diff --git a/Engine/Source/Programs/AutomationTool/Gauntlet/Unreal/Automation/CookUtils/ContentProviders.cs b/Engine/Source/Programs/AutomationTool/Gauntlet/Unreal/Automation/CookUtils/ContentProviders.cs
new file mode 100644
--- /dev/null
+++ b/Engine/Source/Programs/AutomationTool/Gauntlet/Unreal/Automation/CookUtils/ContentProviders.cs
@@ -0,0 +1,109 @@
+// Copyright Epic Games, Inc. All Rights Reserved.
+
+using System;
+using System.Collections.Generic;
+using UnrealBuildTool;
+
+namespace Gauntlet
+{
+	internal interface ICookedContentProvider
+	{
+		ContentFolder GetCookedContent();
+	}
+
+	internal class ContentProviderConfig
+	{
+		public UnrealTargetPlatform TargetPlatform { get; set; }
+		public string TargetPlatformName { get; set; }
+		public string ProjectName { get; set; }
+	}
+
+	internal class BaseCookedContentProvider(ContentProviderConfig Config)
+	{
+		protected ContentProviderConfig Config = Config;
+
+		public virtual ContentFolder GetCookedContent()
+		{
+			ContentFolder CookedContent = new ContentFolder(Config.TargetPlatformName)
+			{
+				SubFolders =
+				[
+					new ContentFolder("Engine")
+					{
+						SubFolders = new List<ContentFolder>
+						{
+							new("Content"),
+							new("Plugins")
+						}
+					},
+					new ContentFolder(Config.ProjectName)
+					{
+						SubFolders = new List<ContentFolder>
+						{
+							new("Content"),
+							new("Metadata"),
+							new("Plugins")
+						},
+						Files = new List<string> { "AssetRegistry.bin" }
+					}
+				]
+			};
+
+			return CookedContent;
+		}
+	}
+
+	internal class WindowsCookedContentProvider(ContentProviderConfig Config) : BaseCookedContentProvider(Config), ICookedContentProvider
+	{
+		public override ContentFolder GetCookedContent()
+		{
+			ContentFolder CookedContent = base.GetCookedContent();
+			CookedContent.SubFolders[0].Files = ["GlobalShaderCache-PCD3D_SM5.bin", "GlobalShaderCache-VULKAN_SM6.bin", "GlobalShaderCache-PCD3D_SM6.bin"];
+
+			return CookedContent;
+		}
+	}
+
+	internal class LinuxCookedContentProvider(ContentProviderConfig Config) : BaseCookedContentProvider(Config), ICookedContentProvider
+	{
+		public override ContentFolder GetCookedContent()
+		{
+			ContentFolder CookedContent = base.GetCookedContent();
+			CookedContent.SubFolders[0].Files = ["GlobalShaderCache-VULKAN_SM6.bin"];
+
+			return CookedContent;
+		}
+	}
+
+	internal class MacCookedContentProvider(ContentProviderConfig Config) : BaseCookedContentProvider(Config), ICookedContentProvider
+	{
+		public override ContentFolder GetCookedContent()
+		{
+			ContentFolder CookedContent = base.GetCookedContent();
+			CookedContent.SubFolders[0].Files = ["GlobalShaderCache-METAL_SM5.bin"];
+
+			return CookedContent;
+		}
+	}
+
+	internal class CookedContentProviderFactory(ContentProviderConfig Config)
+	{
+		private readonly UnrealTargetPlatform TargetPlatform = Config.TargetPlatform;
+		private readonly Dictionary<UnrealTargetPlatform, ICookedContentProvider> ProvidersByPlatform = new()
+		{
+			{ UnrealTargetPlatform.Win64, new WindowsCookedContentProvider(Config) },
+			{ UnrealTargetPlatform.Linux, new LinuxCookedContentProvider(Config) },
+			{ UnrealTargetPlatform.Mac, new MacCookedContentProvider(Config) }
+		};
+
+		public ICookedContentProvider GetProvider()
+		{
+			if (!ProvidersByPlatform.TryGetValue(TargetPlatform, out ICookedContentProvider Provider))
+			{
+				throw new NotSupportedException($"Unsupported platform: {TargetPlatform}");
+			}
+
+			return Provider;
+		}
+	}
+}
diff --git a/Engine/Source/Programs/AutomationTool/Gauntlet/Unreal/Automation/UE.CookByTheBookEditor.cs b/Engine/Source/Programs/AutomationTool/Gauntlet/Unreal/Automation/UE.CookByTheBookEditor.cs
--- a/Engine/Source/Programs/AutomationTool/Gauntlet/Unreal/Automation/UE.CookByTheBookEditor.cs
+++ b/Engine/Source/Programs/AutomationTool/Gauntlet/Unreal/Automation/UE.CookByTheBookEditor.cs
@@ -1,7 +1,6 @@
 // Copyright Epic Games, Inc. All Rights Reserved.
 
 using Gauntlet;
-using System.Collections.Generic;
 using System.IO;
 using System.Linq;
 using System.Text;
@@ -20,6 +19,7 @@ namespace UE
 		private const string CookingInProgressPattern = @"Cooked\s+packages\s+\d+\s+Packages\s+Remain\s+\d+\s+Total\s+\d+";
 		private const string CookingCompletePattern = "Cook by the book total time in tick";
 		private string[] CookErrors;
+		private ICookedContentProvider ContentProvider;
 
 		protected const string CookingCompleteKey = "The cooking process is complete";
 		protected const string CookingInProgressKey = "The cooking process is in progress";
@@ -85,6 +85,7 @@ namespace UE
 			Checker.AddValidation(CookingContentIsPlacedCorrectlyKey, IsCookedContentPlacedCorrectly);
 
 			InitEditorLogStreamParser();
+			SetupContentProvider();
 		}
 
 		protected override UnrealProcessResult GetExitCodeAndReason(StopReason InReason, UnrealLog InLog, UnrealRoleArtifacts InArtifacts,
@@ -115,6 +116,19 @@ namespace UE
 			EditorRole.CommandLine += $@" {BaseEditorCommandLine} -targetplatform={TargetPlatformName}";
 		}
 
+		private void SetupContentProvider()
+		{
+			ContentProviderConfig Config = new ContentProviderConfig
+			{
+				TargetPlatform = GetTargetPlatform(),
+				TargetPlatformName = GetTargetPlatformName(),
+				ProjectName = Context.Options.Project
+			};
+
+			CookedContentProviderFactory ContentProviderFactory = new CookedContentProviderFactory(Config);
+			ContentProvider = ContentProviderFactory.GetProvider();
+		}
+
 		private void InitEditorLogStreamParser()
 		{
 			if (GetRunningEditor() is not { } EditorApp)
@@ -174,33 +188,7 @@ namespace UE
 
 		protected virtual ContentFolder GetExpectedCookedContent()
 		{
-			string TargetPlatformName = GetTargetPlatformName();
-
-			ContentFolder CookedContent = new ContentFolder(TargetPlatformName)
-			{
-				SubFolders =
-				[
-					new ContentFolder("Engine")
-					{
-						SubFolders = new List<ContentFolder>
-						{
-							new("Content"),
-							new("Plugins")
-						},
-						Files = ["GlobalShaderCache-PCD3D_SM5.bin", "GlobalShaderCache-VULKAN_SM6.bin", "GlobalShaderCache-PCD3D_SM6.bin"]
-					},
-					new ContentFolder(Context.Options.Project)
-					{
-						SubFolders = new List<ContentFolder>
-						{
-							new("Content"),
-							new("Metadata"),
-							new("Plugins")
-						},
-						Files = new List<string>{ "AssetRegistry.bin" }
-					}
-				]
-			};
+			ContentFolder CookedContent = ContentProvider.GetCookedContent();
 
 			return CookedContent;
 		}
diff --git a/Engine/Source/Programs/AutomationTool/Gauntlet/Unreal/Automation/UE.CookOnTheFly.cs b/Engine/Source/Programs/AutomationTool/Gauntlet/Unreal/Automation/UE.CookOnTheFly.cs
--- a/Engine/Source/Programs/AutomationTool/Gauntlet/Unreal/Automation/UE.CookOnTheFly.cs
+++ b/Engine/Source/Programs/AutomationTool/Gauntlet/Unreal/Automation/UE.CookOnTheFly.cs
@@ -16,6 +16,7 @@ namespace UE
 	public class CookOnTheFly : UnrealTestNode<UnrealTestConfiguration>
 	{
 		string CookedProjectDirectoryPath;
+		string SavedCookedPath;
 		string CookedSettingsFilePath;
 		DateTime StartTime = DateTime.Now;
 		bool DeferredClientStarted = false;
@@ -84,14 +85,9 @@ namespace UE
 			UnrealTargetPlatform TargetPlatform = Context.GetRoleContext(Config.GetMainRequiredRole().Type).Platform;
 			CookedProjectDirectoryPath = Path.GetDirectoryName(Context.Options.ProjectPath.ToNormalizedPath());
 			string ProjectName = Context.Options.Project;
-			if (TargetPlatform == UnrealTargetPlatform.Win64)
-			{
-				CookedSettingsFilePath = CookedProjectDirectoryPath + string.Format("\\Saved\\Cooked\\Windows\\{0}\\Metadata\\CookedSettings.txt", ProjectName);
-			}
-			else
-			{
-				CookedSettingsFilePath = CookedProjectDirectoryPath + string.Format("\\Saved\\Cooked\\{0}\\{1}\\Metadata\\CookedSettings.txt", TargetPlatform.ToString(), ProjectName);
-			}
+			string TargetPlatformName = UnrealHelpers.GetPlatformName(TargetPlatform, UnrealTargetRole.Editor, false);
+			SavedCookedPath = Path.Combine(CookedProjectDirectoryPath, "Saved", "Cooked");
+			CookedSettingsFilePath = Path.Combine(SavedCookedPath, TargetPlatformName, ProjectName, "Metadata", "CookedSettings.txt");
 
 			UnrealTestRole ClientRole = Config.RequireRole(UnrealTargetRole.Client);
 			string HostIP = UnrealHelpers.GetHostIpAddress();
@@ -231,7 +227,7 @@ namespace UE
 			{
 				string CookedSettingsFileText = File.ReadAllText(CookedSettingsFilePath);
 				string CookModeString = GetCookModeString();
-				DateTime CookedDirectoryLastModifiedTime = File.GetLastWriteTime(CookedProjectDirectoryPath + "\\Saved\\Cooked\\");
+				DateTime CookedDirectoryLastModifiedTime = File.GetLastWriteTime(SavedCookedPath);
 				if (CookedSettingsFileText.Contains(CookModeString) && (CookedDirectoryLastModifiedTime>StartTime))
 				{
 					Log.Info("Found '{0}'. The cook mode changed for the project", CookModeString);
