Allowing custom parameters for the LaunchMultiServer script to substitute {ServerId} {ClientId} and {ProxyId}. Small fixes for launching Clients.
Allow -EnableInsights for LaunchLocalServers.

#rnx

diff --git a/Engine/Source/Programs/AutomationTool/Scripts/LaunchMultiServer.Automation.cs b/Engine/Source/Programs/AutomationTool/Scripts/LaunchMultiServer.Automation.cs
--- a/Engine/Source/Programs/AutomationTool/Scripts/LaunchMultiServer.Automation.cs
+++ b/Engine/Source/Programs/AutomationTool/Scripts/LaunchMultiServer.Automation.cs
@@ -436,7 +436,7 @@ Example running 4 clients locally connecting to single proxy server w/ a non-loc
 		/// <param name="NonLocalServers">Specify servers that are not running on the local machine but contribute to the Metaverse</param>
 		/// <param name="AdditionalServerArguments">These are extra parameters that the server processes should put on the command-line(shared between all servers)</param>
 		/// <returns>The Addresses of all of the Game Servers spun-up.</returns>
-		protected string[] StartProcessesForDedicatedServers(int NumLocalServers, string[] NonLocalServers, string AdditionalServerArguments)
+		virtual protected string[] StartProcessesForDedicatedServers(int NumLocalServers, string[] NonLocalServers, string AdditionalServerArguments)
 		{
 			if (NumLocalServers < 1)
 			{
@@ -460,7 +460,7 @@ Example running 4 clients locally connecting to single proxy server w/ a non-loc
 
 				Logger.LogInformation("Starting Dedicated MultiServer for Game Port {0}", ServerGamePort);
 
-				string ServerArgs = ServerCommonArgs;
+				string ServerArgs = ServerCommonArgs.Replace("{ServerId}", ServerId.ToString(), StringComparison.OrdinalIgnoreCase);
 				ServerArgs += String.Format(" -MultiServerLocalHost -port={0} -log=MultiServer-{1}.log", ServerGamePort, ServerId);
 
 				string AppTitle = $"Server ID {ServerId}";
@@ -477,7 +477,7 @@ Example running 4 clients locally connecting to single proxy server w/ a non-loc
 		/// <param name="GameServerAddresses">The game servers (or proxies) that a client will connect to.  There will be ClientsPerServer launched per entry</param>
 		/// <param name="AdditionalClientArguments">What additional parameters we launch the client with</param>
 		/// <param name="ClientsPerServer">The number of clients to launch per-GameServerAddress entry</param>
-		protected void StartProcessesForClients(ProcessLaunchType ClientType, string[] GameServerAddresses, string AdditionalClientArguments, int ClientsPerServer = 1)
+		virtual protected void StartProcessesForClients(ProcessLaunchType ClientType, string[] GameServerAddresses, string AdditionalClientArguments, int ClientsPerServer = 1)
 		{
 			for (int Idx= 0; Idx < GameServerAddresses.Length; ++Idx)
 			{
@@ -485,25 +485,28 @@ Example running 4 clients locally connecting to single proxy server w/ a non-loc
 
 				string ClientArgs = $" -game {AdditionalClientArguments}";
 				ClientArgs += $" {GameServerAddress}";
-				ClientArgs += $" -log={ClientType}-{Idx+1}.log";
-				ClientArgs += $" -windowed -SaveWinPos={Idx+1}";
+				ClientArgs += $" -log={ClientType}-{{ClientId}}.log -windowed -SaveWinPos={{ClientId}}";
 
 				// Potentially launch multiple clients per server address
 				for (int InstanceNum = 0; InstanceNum < ClientsPerServer; ++InstanceNum)
 				{
-					string InstanceStr = (ClientsPerServer > 1) ? $" [{InstanceNum + 1}]" : string.Empty;
-					string AppTitle = $"{ClientType} {Idx + 1}{InstanceStr}";
-					ExecuteApp(ClientType, AppTitle, ClientArgs);
+					int ClientId = (Idx + 1 + InstanceNum);
+					string InstanceArgs = ClientArgs.Replace("{ClientId}", ClientId.ToString(), StringComparison.OrdinalIgnoreCase);
+
+					string InstanceDetails = (ClientsPerServer > 1) ? $"(Srv {Idx+1})" : string.Empty;
+					string AppTitle = $"{ClientType} {ClientId} {InstanceDetails}";
+					ExecuteApp(ClientType, AppTitle, InstanceArgs);
 				}
 			}
 		}
 
-		private void StartProcessesForProxyServer(string[] LocalProxyServerAddresses, string[] AllGameServerAddresses, string AdditionalProxyArguments)
+		virtual protected void StartProcessesForProxyServer(string[] LocalProxyServerAddresses, string[] AllGameServerAddresses, string AdditionalProxyArguments)
 		{
 			string CommonProxyServerArgs = AdditionalProxyArguments;
 			CommonProxyServerArgs += " -NetDriverOverrides=/Script/MultiServerReplication.ProxyNetDriver";
 			CommonProxyServerArgs += NoTimeouts ? " -notimeouts" : "";
 			CommonProxyServerArgs += IsAttachingDebugger(ProcessLaunchType.ProxyServer) ? " -WaitForDebuggerNoBreak" : "";
+			CommonProxyServerArgs += " -port={ProxyServerPort}"; // will be Replaced() below.
 
 			if (AllGameServerAddresses.Length > 0)
 			{
@@ -534,7 +537,9 @@ Example running 4 clients locally connecting to single proxy server w/ a non-loc
 					continue;
 				}
 
-				string ProxyServerArgs = $"{CommonProxyServerArgs} -port={ProxyServerPort}";
+				string ProxyServerArgs = CommonProxyServerArgs
+					.Replace("{ProxyId}", ProxyIdx.ToString(), StringComparison.OrdinalIgnoreCase)
+					.Replace("{ProxyServerPort}", ProxyServerPort.ToString(), StringComparison.OrdinalIgnoreCase);
 
 				string AppTitle = $"ProxyServer {ProxyIdx}";
 				ExecuteApp(ProcessLaunchType.ProxyServer, AppTitle, ProxyServerArgs);
