Correct handling of cook platform and selected runtimes in BuildCookRun snapshot import.

[RN] none

diff --git a/Engine/Source/Programs/AutomationTool/Scripts/CookCommand.Automation.cs b/Engine/Source/Programs/AutomationTool/Scripts/CookCommand.Automation.cs
--- a/Engine/Source/Programs/AutomationTool/Scripts/CookCommand.Automation.cs
+++ b/Engine/Source/Programs/AutomationTool/Scripts/CookCommand.Automation.cs
@@ -414,28 +414,30 @@ namespace AutomationScripts
 			}
 
 			List<string> Platforms = PlatformsToCook.Split("+").ToList();
-			List<string> Runtimes = new List<string>();
 
-			if (Params.HasGameTargetDetected)
+			foreach (string CookPlatform in Platforms)
 			{
-				Runtimes.Add("game");
-			}
-			else if (Params.HasClientTargetDetected)
-			{
-				Runtimes.Add("client");
-			}
-			if (Params.DedicatedServer)
-			{
-				Runtimes.Add("server");
-			}
+				string Platform;
+				string Runtime;
 
-			foreach (string Platform in Platforms)
-			{
-				foreach (string Runtime in Runtimes)
+				if (CookPlatform.EndsWith("client", StringComparison.OrdinalIgnoreCase))
+				{
+					Platform = CookPlatform.Substring(0, CookPlatform.Length - "client".Length);
+					Runtime = "client";
+				}
+				else if (CookPlatform.EndsWith("server", StringComparison.OrdinalIgnoreCase))
 				{
-					EpicGames.ProjectStore.Build Snapshot = BuildIndex.GetBuild(Branch, Params.ShortProjectName, Platform, Runtime, Changelist, true);
-					Snapshot.Import(Params.RawProjectPath, PlatformCookedDirectory, false, false);
+					Platform = CookPlatform.Substring(0, CookPlatform.Length - "server".Length);
+					Runtime = "server";
 				}
+                else
+                {
+					Platform = CookPlatform;
+					Runtime = "game";
+                }
+
+                EpicGames.ProjectStore.Build Snapshot = BuildIndex.GetBuild(Branch, Params.ShortProjectName, Platform, Runtime, Changelist, true);
+				Snapshot.Import(Params.RawProjectPath, PlatformCookedDirectory, false, false);
 			}
 		}
 
