TestDataV2 - Fix failure to generate test report when adding image comparison with missing image file path

#rnx

diff --git a/Engine/Source/Programs/AutomationTool/Gauntlet/Framework/Utils/Gauntlet.HordeReport.cs b/Engine/Source/Programs/AutomationTool/Gauntlet/Framework/Utils/Gauntlet.HordeReport.cs
--- a/Engine/Source/Programs/AutomationTool/Gauntlet/Framework/Utils/Gauntlet.HordeReport.cs
+++ b/Engine/Source/Programs/AutomationTool/Gauntlet/Framework/Utils/Gauntlet.HordeReport.cs
@@ -1858,7 +1858,6 @@ namespace Gauntlet
 					FilePath = FileUtils.ConvertPathToUri(Path.GetRelativePath(Globals.UnrealRootDir, FilePath));
 					return true;
 				}
-				FilePath = null;
 				return false;
 			}
 
@@ -1873,6 +1872,7 @@ namespace Gauntlet
 				if (!AttachArtifact(ref FilePath, Name))
 				{
 					AddWarning("Failed to attached Artifact {FilePath}.", Path.GetRelativePath(Globals.UnrealRootDir, FilePath));
+					FilePath = null;
 					return false;
 				}
 				return true;
@@ -2189,7 +2189,10 @@ namespace Gauntlet
 						try
 						{
 							File.WriteAllText(OutputFilePath, JsonSerializer.Serialize(Stream.GetEvents(), GetDefaultJsonOptions()));
-							AttachArtifact(ref OutputFilePath, FileName);
+							if(!AttachArtifact(ref OutputFilePath, FileName))
+							{
+								OutputFilePath = null;
+							}
 						}
 						catch (Exception Ex)
 						{
