Horde: Adding Horde job/step metatagging

diff --git a/Engine/Source/Programs/AutomationTool/AutomationUtils/HordeUtils.cs b/Engine/Source/Programs/AutomationTool/AutomationUtils/HordeUtils.cs
--- a/Engine/Source/Programs/AutomationTool/AutomationUtils/HordeUtils.cs
+++ b/Engine/Source/Programs/AutomationTool/AutomationUtils/HordeUtils.cs
@@ -1,8 +1,11 @@
 // Copyright Epic Games, Inc. All Rights Reserved.
 
 using EpicGames.Horde;
+using EpicGames.Horde.Jobs;
 using EpicGames.Horde.Secrets;
 using Microsoft.Extensions.DependencyInjection;
+using Microsoft.Extensions.Logging;
+using System;
 using System.Collections.Concurrent;
 using System.Collections.Generic;
 using System.Net.Http;
@@ -71,5 +74,22 @@ namespace AutomationTool
 			}
 			return propertyValue;
 		}
+
+		public static async Task AddHordeJobMetadata(JobId id, IEnumerable<string>? jobMetaData = null, Dictionary<string, List<string>>? stepMetaData = null, CancellationToken cancellationToken = default)
+		{
+			IHordeClient hordeClient = ServiceProvider.GetRequiredService<IHordeClient>();
+
+			using HordeHttpClient hordeHttpClient = hordeClient.CreateHttpClient();
+
+			bool bConnected = hordeHttpClient.CheckConnectionAsync(cancellationToken).Result;
+
+			if (!bConnected)
+			{
+				Logger.LogWarning("No connection to the Horde Server when tagging job metadata");
+				return;
+			}
+
+			await hordeHttpClient.PutJobMetadataAsync(id, jobMetaData, stepMetaData, cancellationToken);			
+		}
 	}
 }
diff --git a/Engine/Source/Programs/AutomationTool/Gauntlet/Framework/Horde/Gauntlet.Horde.cs b/Engine/Source/Programs/AutomationTool/Gauntlet/Framework/Horde/Gauntlet.Horde.cs
--- a/Engine/Source/Programs/AutomationTool/Gauntlet/Framework/Horde/Gauntlet.Horde.cs
+++ b/Engine/Source/Programs/AutomationTool/Gauntlet/Framework/Horde/Gauntlet.Horde.cs
@@ -2,6 +2,7 @@
 
 using AutomationTool;
 using System;
+using System.Collections.Generic;
 using System.IO;
 using System.Linq;
 using System.Text;
@@ -43,6 +44,27 @@ namespace Gauntlet
 			public string fileName { get; set; }
 		}
 
+		public static void AddJobBuildMetaData(IBuildSource buildSource)
+		{
+			try
+			{
+				if (!IsHordeJob)
+				{
+					return;
+				}
+
+				if (buildSource != null)
+				{
+					CommandUtils.AddHordeJobMetadata(EpicGames.Horde.Jobs.JobId.Parse(JobId), null, new Dictionary<string, List<string>> { { StepId, new List<string> { $"BuildName=${buildSource.BuildName}?display=true" } } }).Wait();
+				}
+
+			}
+			catch (Exception Ex)
+			{
+				Log.Info("Exception while tagging Horde job with build metadata\n{0}\n", Ex.Message);
+			}
+		}
+
 		public static void GenerateSummary(string ArtifactPath)
 		{
 			try
diff --git a/Engine/Source/Programs/AutomationTool/Gauntlet/Unreal/RunUnreal.cs b/Engine/Source/Programs/AutomationTool/Gauntlet/Unreal/RunUnreal.cs
--- a/Engine/Source/Programs/AutomationTool/Gauntlet/Unreal/RunUnreal.cs
+++ b/Engine/Source/Programs/AutomationTool/Gauntlet/Unreal/RunUnreal.cs
@@ -167,6 +167,9 @@ namespace Gauntlet
 			// todo, pass this in as a BuildSource and remove the ContextOption params specific to finding builds
 			UnrealBuildSource BuildInfo = (UnrealBuildSource)Activator.CreateInstance(ContextOptions.BuildSourceType, new object[] { ContextOptions.Project, ContextOptions.ProjectPath, UnrealPath, ContextOptions.UsesSharedBuildType, ContextOptions.Build, ContextOptions.SearchPaths });
 
+			// Tag the the job with metadata about the build 
+			Horde.AddJobBuildMetaData(BuildInfo);
+
 			// Setup accounts
 			SetupAccounts();
 
diff --git a/Engine/Source/Programs/Shared/EpicGames.Horde/HordeHttpClient.cs b/Engine/Source/Programs/Shared/EpicGames.Horde/HordeHttpClient.cs
--- a/Engine/Source/Programs/Shared/EpicGames.Horde/HordeHttpClient.cs
+++ b/Engine/Source/Programs/Shared/EpicGames.Horde/HordeHttpClient.cs
@@ -511,6 +511,20 @@ namespace EpicGames.Horde
 		{
 			return GetAsync<GetJobResponse>(_httpClient, $"api/v1/jobs/{id}", cancellationToken);
 		}
+
+		/// <summary>
+		/// Apply metadata tags to jobs and steps
+		/// </summary>
+		/// <param name="id"></param>
+		/// <param name="jobMetaData"></param>
+		/// <param name="stepMetaData"></param>
+		/// <param name="cancellationToken"></param>
+		/// <returns></returns>
+		public Task PutJobMetadataAsync(JobId id, IEnumerable<string>? jobMetaData = null, Dictionary<string, List<string>>? stepMetaData = null, CancellationToken cancellationToken = default)
+		{
+			return PutAsync(_httpClient, $"api/v1/jobs/{id}/metadata", new PutJobMetadataRequest() {JobMetaData = jobMetaData?.ToList(), StepMetaData = stepMetaData}, cancellationToken);
+		}		
+
 		#endregion
 
 		#region Log
