Implemented CVar change to true for the project which is using for the CB content test case
Added waiter of assets in the CB content to make the test be more stable

#rnx

diff --git a/Engine/Source/Programs/AutomationTool/Gauntlet/Framework/RpcFramework/Gauntlet.RpcLibrary.cs b/Engine/Source/Programs/AutomationTool/Gauntlet/Framework/RpcFramework/Gauntlet.RpcLibrary.cs
--- a/Engine/Source/Programs/AutomationTool/Gauntlet/Framework/RpcFramework/Gauntlet.RpcLibrary.cs
+++ b/Engine/Source/Programs/AutomationTool/Gauntlet/Framework/RpcFramework/Gauntlet.RpcLibrary.cs
@@ -655,6 +655,26 @@ namespace Gauntlet
 			return RpcTarget.AvailableRPCs.Any(R => Equals(RPCName, R.Name));
 		}
 
+		public static bool ChangeCVarValue(this RpcTarget RpcTarget, string CVar, string NewCVarValue)
+		{
+			var RequestArgs = new Dictionary<string, object>()
+			{
+				{ "cvarname", CVar},
+				{ "newcvarvalue", NewCVarValue},
+			};
+			const string RPCName = "ChangeCVar";
+			using (HttpResponseMessage RpcResponse = RpcExecutor.CallRpc(RpcTarget, RPCName, InArgs: RequestArgs))
+			{
+				RpcLibrary.SimpleResponse Response = RpcLibrary.GetSimpleResponse(RpcResponse);
+				if (!Response.Succeeded)
+				{
+					throw new TestException($"RPC '{RPCName}' failed: {Response.Value}");
+				}
+
+				return true;
+			}
+		}
+
 		public static bool AutomationControllerClearReports(this RpcTarget InTarget)
 		{
 			const string RPCName = "AutomationControllerClearReports";
