[Gauntlet] Check for device reservation blocks during static initialization of the HordeDeviceReservationService.
This invokes reservation block behaviors without needing to make a reservation. This is primarily used so we can use reservation blocks on desktop platforms

diff --git a/Engine/Source/Programs/AutomationTool/Gauntlet/Framework/Devices/Gauntlet.DeviceReservationService.cs b/Engine/Source/Programs/AutomationTool/Gauntlet/Framework/Devices/Gauntlet.DeviceReservationService.cs
--- a/Engine/Source/Programs/AutomationTool/Gauntlet/Framework/Devices/Gauntlet.DeviceReservationService.cs
+++ b/Engine/Source/Programs/AutomationTool/Gauntlet/Framework/Devices/Gauntlet.DeviceReservationService.cs
@@ -3,6 +3,9 @@
 using AutomationTool;
 using AutomationTool.DeviceReservation;
 using EpicGames.Core;
+using EpicGames.Horde;
+using EpicGames.Horde.Jobs;
+using Microsoft.Extensions.DependencyInjection;
 using System;
 using System.Collections.Generic;
 using System.Data;
@@ -32,9 +35,7 @@ namespace Gauntlet
 		/// <summary>
 		/// Whether or not this reservation service is enabled
 		/// </summary>
-		public static bool Enabled 
-			=> CommandUtils.ParseParam(Globals.Params.AllArguments, "DeviceURL") 
-			&& CommandUtils.ParseParam(Globals.Params.AllArguments, "DevicePool");
+		public static bool Enabled => Horde.IsHordeJob || Globals.Params.ParseParam("EnableHordeDeviceReservations");
 
 		/// <summary>
 		/// List of devices reserved from this service
@@ -42,13 +43,13 @@ namespace Gauntlet
 		public List<ITargetDevice> ReservedDevices { get; protected set; } = new List<ITargetDevice>();
 
 		/// <summary>
-		/// The base URL of the horde service to request devices from 
+		/// The base URL of the horde service to request devices from
 		/// </summary>
 		[AutoParam("")]
 		public string DeviceURL { get; protected set; }
 
 		/// <summary>
-		/// Name of the horde device pool to request devices from 
+		/// Name of the horde device pool to request devices from
 		/// </summary>
 		[AutoParamWithNames(Default: "", "DevicePool")]
 		public string DevicePoolID { get; protected set; }
@@ -82,6 +83,75 @@ namespace Gauntlet
 			}
 		}
 
+		static HordeDeviceReservationService()
+		{
+			if (Enabled)
+			{
+				JobId JobId = default;
+				JobStepId StepId = default;
+				try
+				{
+					string JobIdEnvVar = Environment.GetEnvironmentVariable("UE_HORDE_JOBID");
+					JobId = JobId.Parse(JobIdEnvVar);
+
+					string StepIdEnvVar = Environment.GetEnvironmentVariable("UE_HORDE_STEPID");
+					StepId = JobStepId.Parse(StepIdEnvVar);
+				}
+				catch
+				{
+					Log.Verbose("Failed to parse Job or Step Id environment variable - some device reservation blocks may not be considered.");
+					return;
+				}
+
+				try
+				{
+					// Locate the batch and step
+					HordeHttpClient HordeClient = CommandUtils.ServiceProvider.GetRequiredService<IHordeClient>().CreateHttpClient();
+					GetJobResponse Job = HordeClient.GetJobAsync(JobId).Result;
+					GetJobBatchResponse JobBatch = Job.Batches
+						.Where(Batch => Batch.Steps
+						.Where(Step => Step.Id == StepId)
+						.FirstOrDefault()?.Id == StepId)
+						.First();
+					GetJobStepResponse Step = JobBatch.Steps.Where(Step => Step.Id == StepId).First();
+
+					if (Step.Annotations.TryGetValue("DeviceReserve", out string Value) && Value.Equals("Begin", StringComparison.InvariantCultureIgnoreCase))
+					{
+						Log.Verbose("Detected running step as the beginning of a reservation block. Any devices used will be cleaned");
+						DevicePool.IsInstallStep = true;
+						DevicePool.DeviceReservationBlock = true;
+						DevicePool.SkipInstall = false;
+						DevicePool.FullClean = true;
+					}
+					else
+					{
+						foreach (GetJobStepResponse PreviousStep in JobBatch.Steps)
+						{
+							if (PreviousStep.Id == Step.Id)
+							{
+								// Ignore steps that occur after this one
+								break;
+							}
+							else if (PreviousStep.Annotations.TryGetValue("DeviceReserve", out Value) && Value.Equals("Begin", StringComparison.InvariantCultureIgnoreCase))
+							{
+								Log.Verbose("Detected running step is within a reservation block. Any devices used will skip installations");
+								DevicePool.IsInstallStep = false;
+								DevicePool.DeviceReservationBlock = true;
+								DevicePool.SkipInstall = true;
+								DevicePool.FullClean = false;
+							}
+						}
+					}
+				}
+				catch (Exception Ex)
+				{
+					Log.Info("Encountered an {ExceptionType} when trying to determine if a Horde device reservation block is active. - some device reservation blocks may not be considered", Ex.GetType().Name);
+					Log.Verbose("{Exception}", Ex.Message);
+					return;
+				}
+			}
+		}
+
 		public HordeDeviceReservationService()
 		{
 			AutoParam.ApplyParamsAndDefaults(this, Globals.Params.AllArguments);
diff --git a/Engine/Source/Programs/AutomationTool/Gauntlet/Gauntlet.Automation.csproj b/Engine/Source/Programs/AutomationTool/Gauntlet/Gauntlet.Automation.csproj
--- a/Engine/Source/Programs/AutomationTool/Gauntlet/Gauntlet.Automation.csproj
+++ b/Engine/Source/Programs/AutomationTool/Gauntlet/Gauntlet.Automation.csproj
@@ -65,6 +65,7 @@
 
   <ItemGroup>
     <ProjectReference Include="..\..\Shared\EpicGames.Core\EpicGames.Core.csproj" PrivateAssets="All"><Private>false</Private></ProjectReference>
+    <ProjectReference Include="..\..\Shared\EpicGames.Horde\EpicGames.Horde.csproj" />
     <ProjectReference Include="..\..\UnrealBuildTool\UnrealBuildTool.csproj" PrivateAssets="All"><Private>false</Private></ProjectReference>
     <ProjectReference Include="..\AutomationUtils\AutomationUtils.Automation.csproj" PrivateAssets="All"><Private>false</Private></ProjectReference>
     <ProjectReference Include="..\Scripts\AutomationScripts.Automation.csproj" PrivateAssets="All"><Private>false</Private></ProjectReference>
