RPC Framework: Set keep-alive time to 600 for RpcFramework for less socket errors when communcating from C# -> C++ eg: during loadscreens, map transitions, etc.

diff --git a/Engine/Source/Programs/AutomationTool/Gauntlet/Framework/RpcFramework/Gauntlet.RpcExecutor.cs b/Engine/Source/Programs/AutomationTool/Gauntlet/Framework/RpcFramework/Gauntlet.RpcExecutor.cs
--- a/Engine/Source/Programs/AutomationTool/Gauntlet/Framework/RpcFramework/Gauntlet.RpcExecutor.cs
+++ b/Engine/Source/Programs/AutomationTool/Gauntlet/Framework/RpcFramework/Gauntlet.RpcExecutor.cs
@@ -769,6 +769,8 @@ namespace Gauntlet
 			HttpMethod MethodToUse = new HttpMethod(RpcEntry.Verb);
 			HttpRequestMessage RpcRequest = new HttpRequestMessage(MethodToUse, RpcUrl);
 			RpcRequest.Headers.Add("rpcname", RpcName);
+			RpcRequest.Headers.Connection.Add("keep-alive");
+			RpcRequest.Headers.Add("Keep-Alive", "600");
 			if (RpcEntry.Args != null && RpcEntry.Args.Count > 0)
 			{
 				foreach (RpcArgumentDesc ArgDesc in RpcEntry.Args)
@@ -903,9 +905,11 @@ namespace Gauntlet
             RpcUrl = String.Format($"{RpcUrl}?{UrlArgs}");
             HttpMethod MethodToUse = new HttpMethod(RpcEntry.Verb);
         	HttpRequestMessage RpcRequest = new HttpRequestMessage(MethodToUse, RpcUrl);
-            
-        	// We might timeout calling CallRpc on a client that's closing - we do not want this to be a fail;
-        	try
+			RpcRequest.Headers.Connection.Add("keep-alive");
+			RpcRequest.Headers.Add("Keep-Alive", "600");
+
+			// We might timeout calling CallRpc on a client that's closing - we do not want this to be a fail;
+			try
         	{
 				HttpResponseMessage RpcResponse = ExecutorClient.Send(RpcRequest);
         		return RpcResponse;
diff --git a/Engine/Source/Programs/AutomationTool/Gauntlet/Framework/RpcFramework/Gauntlet.RpcExecutorAsync.cs b/Engine/Source/Programs/AutomationTool/Gauntlet/Framework/RpcFramework/Gauntlet.RpcExecutorAsync.cs
--- a/Engine/Source/Programs/AutomationTool/Gauntlet/Framework/RpcFramework/Gauntlet.RpcExecutorAsync.cs
+++ b/Engine/Source/Programs/AutomationTool/Gauntlet/Framework/RpcFramework/Gauntlet.RpcExecutorAsync.cs
@@ -61,7 +61,8 @@ public partial class RpcExecutor
 		var MethodToUse = new HttpMethod(RpcEntry.Verb);
 
 		var RpcRequest = new HttpRequestMessage(MethodToUse, RpcUrl);
-
+		RpcRequest.Headers.Connection.Add("keep-alive");
+		RpcRequest.Headers.Add("Keep-Alive", "600");
 		RpcRequest.Headers.Add("rpcname", InRpcName);
 
 		if (InArgs.Count > 0)
