The UE.IterativeCookByTheBook test update

1. Covers UE-TC-13778 - the test already does it.
2. Covers UE-TC-13780 - an additional log validation was added

#rnx

diff --git a/Engine/Source/Programs/AutomationTool/Gauntlet/Unreal/Automation/UE.IterativeCookByTheBook.cs b/Engine/Source/Programs/AutomationTool/Gauntlet/Unreal/Automation/UE.IterativeCookByTheBook.cs
--- a/Engine/Source/Programs/AutomationTool/Gauntlet/Unreal/Automation/UE.IterativeCookByTheBook.cs
+++ b/Engine/Source/Programs/AutomationTool/Gauntlet/Unreal/Automation/UE.IterativeCookByTheBook.cs
@@ -13,9 +13,11 @@ namespace UE
 		private const string StatisticsShownKey = "Statistics of cooked packets are shown";
 		private const string PacketsSkippedKey =
 			"The number of cooked packets at the first pass is equal to the number of iteratively skipped packets at the second pass";
+		private const string DiffTotalKey = "The number of different packages is zero";
 		private const string PacketsStatisticsPattern = "Total Packages";
 		private const string DetailedPacketsStatisticsPattern = @"Packages Cooked: (\d+), Packages Incrementally Skipped: (\d+)";
 		private const string IterativePattern = "IsIterativeCook=true";
+		private const string DiffTotalPattern = "NumberOfDifferentPackages=0";
 
 		private List<(string Cooked, string Skipped)> PacketsStatistics;
 
@@ -31,6 +33,7 @@ namespace UE
 
 			Checker.AddValidation(IterativeCookKey, IsIterativeCook);
 			Checker.AddValidation(StatisticsShownKey, ArePacketStatisticsShown);
+			Checker.AddValidation(DiffTotalKey, IsNumberOfDifferentPackagesZero);
 		}
 
 		public override void TickTest()
@@ -87,5 +90,10 @@ namespace UE
 		{
 			return EditorLogParser.GetLogLinesContaining(IterativePattern).Any();
 		}
+
+		private bool IsNumberOfDifferentPackagesZero()
+		{
+			return EditorLogParser.GetLogLinesContaining(DiffTotalPattern).Any();
+		}
 	}
 }
