[Gauntlet] Changes log verbosity of windows full clean without an app config

diff --git a/Engine/Source/Programs/AutomationTool/Gauntlet/Framework/Base/Gauntlet.TargetDeviceDesktopCommon.cs b/Engine/Source/Programs/AutomationTool/Gauntlet/Framework/Base/Gauntlet.TargetDeviceDesktopCommon.cs
--- a/Engine/Source/Programs/AutomationTool/Gauntlet/Framework/Base/Gauntlet.TargetDeviceDesktopCommon.cs
+++ b/Engine/Source/Programs/AutomationTool/Gauntlet/Framework/Base/Gauntlet.TargetDeviceDesktopCommon.cs
@@ -122,7 +122,7 @@ namespace Gauntlet
 
 			if (AppConfig == null)
 			{
-				Log.Warning("{Platform} expects an non-null AppConfig value to reliably determine which artifact path to clear! Skipping clean", Platform);
+				Log.Info("{Platform} expects an non-null AppConfig value to reliably determine which artifact path to clear! Skipping clean of build's saved folder", Platform);
 				return;
 			}
 
