[Gauntlet] Corrects issue with DevicePool SkipInstall/FullClean values set in statis constructor being overwritten in instance constructor
Also adds a new arg to index local device names, which is useful for retaining caches in niche scenarios

diff --git a/Engine/Source/Programs/AutomationTool/Gauntlet/Framework/Devices/Gauntlet.DevicePool.cs b/Engine/Source/Programs/AutomationTool/Gauntlet/Framework/Devices/Gauntlet.DevicePool.cs
--- a/Engine/Source/Programs/AutomationTool/Gauntlet/Framework/Devices/Gauntlet.DevicePool.cs
+++ b/Engine/Source/Programs/AutomationTool/Gauntlet/Framework/Devices/Gauntlet.DevicePool.cs
@@ -219,7 +219,7 @@ namespace Gauntlet
 	/// EnumerateDevices to build a list of desired devices, which must then be reserved by calling ReserveDevices.
 	/// Once done ReleaseDevices should be called.
 	///
-	/// These reservations exist at the process level and we rely on the device implementation to provide arbitrage 
+	/// These reservations exist at the process level and we rely on the device implementation to provide arbitrage
 	/// between difference processes and machines
 	///
 	/// </summary>
@@ -250,14 +250,39 @@ namespace Gauntlet
 			}
 		}
 
+		/// <summary>
+		/// Whether or not installations should be skipped
+		/// This is primarily used within the context of UnrealSession
+		/// </summary>
 		public static bool SkipInstall;
 
+		/// <summary>
+		/// Whether or not devices should be fully cleaned
+		/// This is primarily used within the context of UnrealSession
+		/// </summary>
 		public static bool FullClean;
 
+		/// <summary>
+		/// Whether or a device reservation block is active
+		/// This is primarily used within the context of UnrealSession
+		/// when a step is executed by Horde that contains a DeviceReserve=Begin annotation
+		/// or is a step that proceeds one
+		/// </summary>
 		public static bool DeviceReservationBlock;
 
+		/// <summary>
+		/// Whether or not this particular step is an installations tep in a device reservation block
+		/// This is primarily used within the context of UnrealSession
+		/// when a job is executed by Horde that contains a DeviceReserve=Begin/Install annotation
 		public static bool? IsInstallStep;
 
+		/// <summary>
+		/// Optional index to name local devices
+		/// Occasionally useful for retaining device caches from previous tests
+		/// </summary>
+		[AutoParam(0)]
+		public int LocalDeviceIndex { get; set; }
+
 		/// <summary>
 		/// Active pool instance
 		/// </summary>
@@ -321,6 +346,8 @@ namespace Gauntlet
 		/// </summary>
 		protected DevicePool()
 		{
+			AutoParam.ApplyParamsAndDefaults(this, Globals.Params.AllArguments);
+
 			lock (LockObject)
 			{
 				// Create two local devices by default for ease of running a client and server
@@ -328,9 +355,8 @@ namespace Gauntlet
 				InitializeReservationServices();
 			}
 
-			SkipInstall = Globals.Params.ParseParams("SkipInstall", "SkipCopy", "SkipDeploy");
-			FullClean = Globals.Params.ParseParam("FullClean");
-			DeviceReservationBlock = false;
+			SkipInstall |= Globals.Params.ParseParams("SkipInstall", "SkipCopy", "SkipDeploy");
+			FullClean |= Globals.Params.ParseParam("FullClean");
 		}
 
 		#region IDisposable Support
@@ -451,7 +477,7 @@ namespace Gauntlet
 			for (int i = NumDevices; i < MaxCount; i++)
 			{
 				DeviceDefinition Def = new DeviceDefinition();
-				Def.Name = string.Format("LocalDevice{0}", i);
+				Def.Name = string.Format("LocalDevice{0}", i + LocalDeviceIndex);
 				Def.Platform = LocalPlatform;
 				UnprovisionedDevices.Add(Def);
 			}
@@ -542,7 +568,7 @@ namespace Gauntlet
 										To = To.AddDays(1);
 									}
 
-									// if From is in the future (e.g. it's 01:00 and range is 22:00-08:00) we may be in the previous days window, 
+									// if From is in the future (e.g. it's 01:00 and range is 22:00-08:00) we may be in the previous days window,
 									// so move them both back a day
 									if (From > DateTime.Now)
 									{
@@ -595,7 +621,7 @@ namespace Gauntlet
 							// When using device services, skip adding non-desktop local devices to pool if any of the services can support that platform
 							bool IsDesktop = DevicePlatform.IsInGroup(UnrealPlatformGroup.Desktop);
 							bool ReservationsEnabled = ReservationServices.Count > 0;
-							bool ServicesCanSupportThisPlatform = ReservationsEnabled 
+							bool ServicesCanSupportThisPlatform = ReservationsEnabled
 								&& ReservationServices.Where(Service => Service.CanSupportDeviceConstraint(new UnrealDeviceTargetConstraint(DevicePlatform))).Any();
 							bool DeviceServiceEnabled = ServicesCanSupportThisPlatform;
 
@@ -902,7 +928,7 @@ namespace Gauntlet
 			Dictionary<UnrealDeviceTargetConstraint, int> AvailableDeviceTypes = new Dictionary<UnrealDeviceTargetConstraint, int>();
 			Dictionary<UnrealDeviceTargetConstraint, int> TotalDeviceTypes = new Dictionary<UnrealDeviceTargetConstraint, int>();
 
-			// Do these "how many available" checks every time because the DevicePool provisions on demand so while it may think it has N machines, 
+			// Do these "how many available" checks every time because the DevicePool provisions on demand so while it may think it has N machines,
 			// some of them may fail to be provisioned and we could end up with none!
 
 			// See how many of these types are in device pool (mostly to supply informative info if we can't meet these)
@@ -987,7 +1013,7 @@ namespace Gauntlet
 		/// </summary>
 		/// <param name="Def"></param>
 		/// <returns></returns>
-		public ITargetDevice CreateAndRegisterDeviceFromDefinition(DeviceDefinition Def, UnrealDeviceTargetConstraint Constraint, 
+		public ITargetDevice CreateAndRegisterDeviceFromDefinition(DeviceDefinition Def, UnrealDeviceTargetConstraint Constraint,
 			IDeviceReservationService Service = null, IDeviceFactory Factory = null)
 		{
 			ITargetDevice NewDevice = null;
@@ -1287,7 +1313,7 @@ namespace Gauntlet
 		}
 
 		#region Legacy Implementation
-	
+
 		[Obsolete("Will be removed in a future release")]
 		public string DeviceURL
 		{
diff --git a/Engine/Source/Programs/AutomationTool/Gauntlet/Framework/Devices/Gauntlet.DeviceReservationService.cs b/Engine/Source/Programs/AutomationTool/Gauntlet/Framework/Devices/Gauntlet.DeviceReservationService.cs
--- a/Engine/Source/Programs/AutomationTool/Gauntlet/Framework/Devices/Gauntlet.DeviceReservationService.cs
+++ b/Engine/Source/Programs/AutomationTool/Gauntlet/Framework/Devices/Gauntlet.DeviceReservationService.cs
@@ -118,13 +118,27 @@ namespace Gauntlet
 						.First();
 					GetJobStepResponse Step = JobBatch.Steps.Where(Step => Step.Id == StepId).First();
 
-					if (Step.Annotations.TryGetValue("DeviceReserve", out string Value) && Value.Equals("Begin", StringComparison.InvariantCultureIgnoreCase))
+					if (Step.Annotations.TryGetValue("DeviceReserve", out string Value))
 					{
-						Log.Verbose("Detected running step as the beginning of a reservation block. Any devices used will be cleaned");
-						DevicePool.IsInstallStep = true;
-						DevicePool.DeviceReservationBlock = true;
-						DevicePool.SkipInstall = false;
-						DevicePool.FullClean = true;
+						bool InstallStep = false;
+						if (Value.Equals("Begin", StringComparison.InvariantCultureIgnoreCase))
+						{
+							Log.Verbose("Detected running step as the beginning of a reservation block. Any devices used will be cleaned");
+							InstallStep = true;
+						}
+						else if (Value.Equals("Install", StringComparison.InvariantCultureIgnoreCase))
+						{
+							Log.Verbose("Detected running step an explicit install step within a reservation block. Any devices used will be cleaned");
+							InstallStep = true;
+						}
+
+						if (InstallStep)
+						{
+							DevicePool.IsInstallStep = true;
+							DevicePool.DeviceReservationBlock = true;
+							DevicePool.SkipInstall = false;
+							DevicePool.FullClean = true;
+						}
 					}
 					else
 					{
@@ -142,6 +156,7 @@ namespace Gauntlet
 								DevicePool.DeviceReservationBlock = true;
 								DevicePool.SkipInstall = true;
 								DevicePool.FullClean = false;
+								break;
 							}
 						}
 					}
