Cloud artifacts - Use a more specific name for the uploaded utraces to avoid overwriting them if multiple zen invocations happen in the same job

diff --git a/Engine/Source/Programs/AutomationTool/BuildGraph/Tasks/CreateCloudArtifactTask.cs b/Engine/Source/Programs/AutomationTool/BuildGraph/Tasks/CreateCloudArtifactTask.cs
--- a/Engine/Source/Programs/AutomationTool/BuildGraph/Tasks/CreateCloudArtifactTask.cs
+++ b/Engine/Source/Programs/AutomationTool/BuildGraph/Tasks/CreateCloudArtifactTask.cs
@@ -483,7 +483,7 @@ namespace AutomationTool.Tasks
 			if (_parameters.EnableTracing)
 			{
 				// this is automatically uploaded as it's in the Saved folder
-				traceFile = FileReference.Combine(Unreal.RootDirectory, "Engine/Programs/AutomationTool/Saved/Logs/build-upload.utrace");
+				traceFile = FileReference.Combine(Unreal.RootDirectory, $"Engine/Programs/AutomationTool/Saved/Logs/build-upload-{buildId}-{partId}.utrace");
 			}
 
 			// set part name to default by default - should be exposed as a option to allow for uploading separate parts
diff --git a/Engine/Source/Programs/AutomationTool/BuildGraph/Tasks/RetrieveCloudArtifactTask.cs b/Engine/Source/Programs/AutomationTool/BuildGraph/Tasks/RetrieveCloudArtifactTask.cs
--- a/Engine/Source/Programs/AutomationTool/BuildGraph/Tasks/RetrieveCloudArtifactTask.cs
+++ b/Engine/Source/Programs/AutomationTool/BuildGraph/Tasks/RetrieveCloudArtifactTask.cs
@@ -393,7 +393,7 @@ namespace AutomationTool.Tasks
 			if (_parameters.EnableTracing)
 			{
 				// this is automatically uploaded as it's in the Saved folder
-				traceFile = FileReference.Combine(Unreal.RootDirectory, "Engine/Programs/AutomationTool/Saved/Logs/build-download.utrace");
+				traceFile = FileReference.Combine(Unreal.RootDirectory, $"Engine/Programs/AutomationTool/Saved/Logs/build-download-{artifact}.utrace");
 			}
 
 			BuildDownload.DownloadBuild(host, outputDir, ns, bucketId, artifact.ToString(), accessToken, allowMultipart: _parameters.AllowMultipart, assumeHttp2: String.Equals(httpVersion, "http2-only", StringComparison.OrdinalIgnoreCase), boostWorkers: _parameters.BoostWorkers, traceFile: traceFile);
