Update ReportItemizedExecutableCode to ignore 'core' so oodle2 is correctly found.

Avoids duplicate 'core' modules.

#rnx

diff --git a/Engine/Source/Programs/AutomationTool/Scripts/ReportItemizedExecutableCode.cs b/Engine/Source/Programs/AutomationTool/Scripts/ReportItemizedExecutableCode.cs
--- a/Engine/Source/Programs/AutomationTool/Scripts/ReportItemizedExecutableCode.cs
+++ b/Engine/Source/Programs/AutomationTool/Scripts/ReportItemizedExecutableCode.cs
@@ -129,7 +129,7 @@ namespace ReportItemizedExecutableCode.Automation
 				return SourcePath;
 			}
 
-			string[] NamesThatCannotBeModuleNames = { "..", "src", "lib", "float" };
+			string[] NamesThatCannotBeModuleNames = { "..", "src", "lib", "float", "core" };
 
 			for (int IdxElem = PathElements.Length - 2; IdxElem >= 0; --IdxElem)
 			{
