AutomationUtils: Call WaitForExit() to ensure async output handlers are flushed. Spectulative fix.

#rnx

diff --git a/Engine/Source/Programs/AutomationTool/AutomationUtils/ProcessUtils.cs b/Engine/Source/Programs/AutomationTool/AutomationUtils/ProcessUtils.cs
--- a/Engine/Source/Programs/AutomationTool/AutomationUtils/ProcessUtils.cs
+++ b/Engine/Source/Programs/AutomationTool/AutomationUtils/ProcessUtils.cs
@@ -503,10 +503,14 @@ namespace AutomationTool
 		public void WaitForExit()
 		{
 			bool bProcTerminated = false;
-			// Make sure the process objeect is valid.
+			// Make sure the process object is valid.
 			lock (ProcSyncObject)
 			{
 				bProcTerminated = (Proc == null) || Proc.HasExited;
+				if (Proc != null && bProcTerminated)
+				{
+					Proc.WaitForExit();
+				}
 			}
 			// Keep checking if we got all output messages until the process terminates.
 			Stopwatch Watch = Stopwatch.StartNew();
@@ -534,6 +538,10 @@ namespace AutomationTool
 				lock (ProcSyncObject)
 				{
 					bProcTerminated = (Proc == null) || Proc.HasExited;
+					if (Proc != null && bProcTerminated)
+					{
+						Proc.WaitForExit();
+					}
 				}
 				if (!bProcTerminated)
 				{
@@ -551,21 +559,16 @@ namespace AutomationTool
 			}
 			if (!(bStdOutSignalReceived && bStdErrSignalReceived))
 			{
-				Logger.LogInformation("Waited for {0:n2}s for output of {AppName}, some output may be missing; we gave up.", Watch.Elapsed.TotalSeconds, AppName);
+				Logger.LogInformation("Waited for {Seconds:n2}s for output of {AppName}, some output may be missing; we gave up.", Watch.Elapsed.TotalSeconds, AppName);
 			}
 
 			// Double-check if the process terminated
 			lock (ProcSyncObject)
 			{
-				bProcTerminated = (Proc == null) || Proc.HasExited;
-
 				if (Proc != null)
 				{
-					if (!bProcTerminated)
-					{
-						// The process did not terminate yet but we've read all output messages, wait until the process terminates
-						Proc.WaitForExit();
-					}
+					// wait until the process terminates
+					Proc.WaitForExit();
 
 					ExitCode = Proc.ExitCode;
 				}
@@ -722,6 +725,7 @@ namespace AutomationTool
 					}
 					else
 					{
+						ProcToKill.WaitForExit();
 						ExitCode = ProcToKill.ExitCode;
 						Logger.LogDebug("Process {ProcToKillName} successfully exited.", ProcToKillName);
 						OnProcessExited();
