[Backout] - CL43556455
#rnx
Original CL Desc
-----------------------------------------------------------------
AutomationUtils: Call WaitForExit() to ensure async output handlers are flushed. Spectulative fix.

#rnx

diff --git a/Engine/Source/Programs/AutomationTool/AutomationUtils/ProcessUtils.cs b/Engine/Source/Programs/AutomationTool/AutomationUtils/ProcessUtils.cs
--- a/Engine/Source/Programs/AutomationTool/AutomationUtils/ProcessUtils.cs
+++ b/Engine/Source/Programs/AutomationTool/AutomationUtils/ProcessUtils.cs
@@ -503,14 +503,10 @@ namespace AutomationTool
 		public void WaitForExit()
 		{
 			bool bProcTerminated = false;
-			// Make sure the process object is valid.
+			// Make sure the process objeect is valid.
 			lock (ProcSyncObject)
 			{
 				bProcTerminated = (Proc == null) || Proc.HasExited;
-				if (Proc != null && bProcTerminated)
-				{
-					Proc.WaitForExit();
-				}
 			}
 			// Keep checking if we got all output messages until the process terminates.
 			Stopwatch Watch = Stopwatch.StartNew();
@@ -538,10 +534,6 @@ namespace AutomationTool
 				lock (ProcSyncObject)
 				{
 					bProcTerminated = (Proc == null) || Proc.HasExited;
-					if (Proc != null && bProcTerminated)
-					{
-						Proc.WaitForExit();
-					}
 				}
 				if (!bProcTerminated)
 				{
@@ -559,16 +551,21 @@ namespace AutomationTool
 			}
 			if (!(bStdOutSignalReceived && bStdErrSignalReceived))
 			{
-				Logger.LogInformation("Waited for {Seconds:n2}s for output of {AppName}, some output may be missing; we gave up.", Watch.Elapsed.TotalSeconds, AppName);
+				Logger.LogInformation("Waited for {0:n2}s for output of {AppName}, some output may be missing; we gave up.", Watch.Elapsed.TotalSeconds, AppName);
 			}
 
 			// Double-check if the process terminated
 			lock (ProcSyncObject)
 			{
+				bProcTerminated = (Proc == null) || Proc.HasExited;
+
 				if (Proc != null)
 				{
-					// wait until the process terminates
-					Proc.WaitForExit();
+					if (!bProcTerminated)
+					{
+						// The process did not terminate yet but we've read all output messages, wait until the process terminates
+						Proc.WaitForExit();
+					}
 
 					ExitCode = Proc.ExitCode;
 				}
@@ -725,7 +722,6 @@ namespace AutomationTool
 					}
 					else
 					{
-						ProcToKill.WaitForExit();
 						ExitCode = ProcToKill.ExitCode;
 						Logger.LogDebug("Process {ProcToKillName} successfully exited.", ProcToKillName);
 						OnProcessExited();
