Fixing the issue with the copy to clipboard missing project parameter

diff --git a/Engine/Source/Programs/AutomationTool/Scripts/OpenEditor.Automation.cs b/Engine/Source/Programs/AutomationTool/Scripts/OpenEditor.Automation.cs
--- a/Engine/Source/Programs/AutomationTool/Scripts/OpenEditor.Automation.cs
+++ b/Engine/Source/Programs/AutomationTool/Scripts/OpenEditor.Automation.cs
@@ -43,20 +43,9 @@ namespace AutomationTool
 			bNoProject = false;
 		}
 
-		public override ExitCode Execute()
+		protected IEnumerable<string> GetParamList()
 		{
-			string EditorPath = string.Empty;
-			if (Path.IsPathRooted(UnrealEditorApp) && File.Exists(UnrealEditorApp))
-			{
-				EditorPath = UnrealEditorApp;
-			}
-			else
-			{
-				EditorPath = HostPlatform.Current.GetUnrealExePath(UnrealEditorApp);
-			}
-
 			string EditorArgs = "";
-
 			ProjectName = (string.IsNullOrEmpty(ProjectName) || !Path.IsPathRooted(ProjectName)) ? ParseParamValue("project", ProjectName) : ProjectName;
 			IEnumerable<string> ParamList = null;
 
@@ -86,10 +75,23 @@ namespace AutomationTool
 				ParamList = this.Params
 								.Where(P => P.StartsWith("project=", StringComparison.OrdinalIgnoreCase) == false);
 			}
+			
+			return new[] { EditorArgs }.Concat(ParamList);
+		}
 
-			ParamList = new[] { EditorArgs }.Concat(ParamList);
+		public override ExitCode Execute()
+		{
+			string EditorPath = string.Empty;
+			if (Path.IsPathRooted(UnrealEditorApp) && File.Exists(UnrealEditorApp))
+			{
+				EditorPath = UnrealEditorApp;
+			}
+			else
+			{
+				EditorPath = HostPlatform.Current.GetUnrealExePath(UnrealEditorApp);
+			}
 
-			bool bLaunched = RunProceess.RunUntrackedProcess(EditorPath, string.Join(" -", ParamList));
+			bool bLaunched = RunProceess.RunUntrackedProcess(EditorPath, string.Join(" -", GetParamList()));
 
 			return bLaunched ? ExitCode.Success : ExitCode.Error_UATLaunchFailure;
 		}
