Fixed a deadlock in the WaitForExit ProcessUtils code for UAT.

#rnx

diff --git a/Engine/Source/Programs/AutomationTool/AutomationUtils/ProcessUtils.cs b/Engine/Source/Programs/AutomationTool/AutomationUtils/ProcessUtils.cs
--- a/Engine/Source/Programs/AutomationTool/AutomationUtils/ProcessUtils.cs
+++ b/Engine/Source/Programs/AutomationTool/AutomationUtils/ProcessUtils.cs
@@ -503,11 +503,17 @@ namespace AutomationTool
 		public void WaitForExit()
 		{
 			bool bProcTerminated = false;
-			// Make sure the process objeect is valid.
+			// Make sure the process object is valid.
 			lock (ProcSyncObject)
 			{
 				bProcTerminated = (Proc == null) || Proc.HasExited;
 			}
+
+			if (bProcTerminated)
+			{
+				Proc?.WaitForExit();
+			}
+
 			// Keep checking if we got all output messages until the process terminates.
 			Stopwatch Watch = Stopwatch.StartNew();
 			int MaxWaitUntilMessagesReceived = 60 * 1000;
@@ -535,6 +541,12 @@ namespace AutomationTool
 				{
 					bProcTerminated = (Proc == null) || Proc.HasExited;
 				}
+
+				if (bProcTerminated)
+				{
+					Proc?.WaitForExit();
+				}
+
 				if (!bProcTerminated)
 				{
 					// Timeout starts when process has terminated
@@ -551,22 +563,16 @@ namespace AutomationTool
 			}
 			if (!(bStdOutSignalReceived && bStdErrSignalReceived))
 			{
-				Logger.LogInformation("Waited for {0:n2}s for output of {AppName}, some output may be missing; we gave up.", Watch.Elapsed.TotalSeconds, AppName);
+				Logger.LogInformation("Waited for {Seconds:n2}s for output of {AppName}, some output may be missing; we gave up.", Watch.Elapsed.TotalSeconds, AppName);
 			}
 
-			// Double-check if the process terminated
+			// wait until the process terminates
+			Proc?.WaitForExit();
+
 			lock (ProcSyncObject)
 			{
-				bProcTerminated = (Proc == null) || Proc.HasExited;
-
 				if (Proc != null)
 				{
-					if (!bProcTerminated)
-					{
-						// The process did not terminate yet but we've read all output messages, wait until the process terminates
-						Proc.WaitForExit();
-					}
-
 					ExitCode = Proc.ExitCode;
 				}
 			}
@@ -722,6 +728,7 @@ namespace AutomationTool
 					}
 					else
 					{
+						ProcToKill.WaitForExit();
 						ExitCode = ProcToKill.ExitCode;
 						Logger.LogDebug("Process {ProcToKillName} successfully exited.", ProcToKillName);
 						OnProcessExited();
