Fix overriding config argument defaulting to true so that passing -arg=false to Gauntlet argument will actually disable the flag

#rnx

diff --git a/Engine/Source/Programs/AutomationTool/Gauntlet/Framework/Gauntlet.AutoParam.cs b/Engine/Source/Programs/AutomationTool/Gauntlet/Framework/Gauntlet.AutoParam.cs
--- a/Engine/Source/Programs/AutomationTool/Gauntlet/Framework/Gauntlet.AutoParam.cs
+++ b/Engine/Source/Programs/AutomationTool/Gauntlet/Framework/Gauntlet.AutoParam.cs
@@ -89,12 +89,12 @@ namespace Gauntlet
 
 		/// <summary>
 		/// Checks whether Args contains a -Param statement, if so returns true else
-		/// returns Default
+		/// checks if the passed argument is different from true or 1
 		/// </summary>
 		/// <param name="Param"></param>
 		/// <param name="Args"></param>
 		/// <returns></returns>
-		static protected bool SwitchExists(string Param, string[] Args)
+		static protected object CoerceToBool(string Param, string[] Args)
 		{
 			foreach (string Arg in Args)
 			{
@@ -109,8 +109,14 @@ namespace Gauntlet
 				{
 					return true;
 				}
+
+				if (StringArg.StartsWith($"{Param}="))
+				{
+					StringArg = StringArg.Substring(Param.Length + 1);
+					return !StringArg.Equals("false", StringComparison.InvariantCultureIgnoreCase) && !StringArg.Equals("0");
+				}
 			}
-			return false;
+			return null;
 		}
 
 		/// <summary>
@@ -270,13 +276,10 @@ namespace Gauntlet
 							// Go through all params used to refer to this member
 							foreach (string Name in ParamNames)
 							{
-								// if default is a bool then just check if the switch exists
 								if (MemberType == typeof(bool))
 								{
-									if (SwitchExists(Name, Args))
-									{
-										NewValue = true;
-									}
+									// if default is a bool then try to coerce to a boolean
+									NewValue = CoerceToBool(Name, Args);
 								}
 								else
 								{
