If a cooked output directory is absent but we need to create iostore container files, emit a clear error message instead of proceeding and then failing with a null object reference latler because we lacked package store information.

diff --git a/Engine/Source/Programs/AutomationTool/Scripts/CopyBuildToStagingDirectory.Automation.cs b/Engine/Source/Programs/AutomationTool/Scripts/CopyBuildToStagingDirectory.Automation.cs
--- a/Engine/Source/Programs/AutomationTool/Scripts/CopyBuildToStagingDirectory.Automation.cs
+++ b/Engine/Source/Programs/AutomationTool/Scripts/CopyBuildToStagingDirectory.Automation.cs
@@ -1875,6 +1875,10 @@ namespace AutomationScripts
 							}
 						}
 					}
+					else if (ShouldCreateIoStoreContainerFiles(Params, SC))
+					{
+						throw new AutomationException(String.Format("A package store manifest or project store is required when staging to IoStore. Expected to find cooked directory {0}. Ensure that cooking was successful.", SC.PlatformCookDir));
+					}
 
 					SC.StageVulkanValidationLayerFiles(Params, StagedFileType.NonUFS, DirectoryReference.Combine(SC.EngineRoot, "Binaries", "ThirdParty", "Vulkan", SC.PlatformDir), StageFilesSearch.AllDirectories);
 
