CloudArtifacts will no longer sanitize name property as Jupiter doesn't need it. HordeArtifactName will now get the sanitization.

#rnx

diff --git a/Engine/Source/Programs/AutomationTool/BuildGraph/Tasks/CreateCloudArtifactTask.cs b/Engine/Source/Programs/AutomationTool/BuildGraph/Tasks/CreateCloudArtifactTask.cs
--- a/Engine/Source/Programs/AutomationTool/BuildGraph/Tasks/CreateCloudArtifactTask.cs
+++ b/Engine/Source/Programs/AutomationTool/BuildGraph/Tasks/CreateCloudArtifactTask.cs
@@ -194,9 +194,10 @@ namespace AutomationTool.Tasks
 		/// <param name="tagNameToFileSet">Mapping from tag names to the set of files they include.</param>
 		public override async Task ExecuteAsync(JobContext job, HashSet<FileReference> buildProducts, Dictionary<string, HashSet<FileReference>> tagNameToFileSet)
 		{
-			string name = SanitizeBucketValue(_parameters.Name);
+			string name = _parameters.Name;
 			ArtifactType type = new ArtifactType(_parameters.Type);
-			ArtifactName hordeArtifactName = String.IsNullOrEmpty(_parameters.HordeArtifactName) ? new ArtifactName(name) : new ArtifactName(_parameters.HordeArtifactName);
+			string sanitizedHordeArtifactParam = SanitizeBucketValue(String.IsNullOrEmpty(_parameters.HordeArtifactName) ? name : _parameters.HordeArtifactName);
+			ArtifactName hordeArtifactName = new ArtifactName(sanitizedHordeArtifactParam);
 			List<string> metadataParam = (_parameters.Metadata ?? String.Empty).Split(';', StringSplitOptions.RemoveEmptyEntries | StringSplitOptions.TrimEntries).ToList();
 
 			Dictionary<string, object> metadata = new Dictionary<string, object>();
