Validate that we have a name set for the CreateCloudArtifact task, either in the Name or HordeArtifactName parameter.

diff --git a/Engine/Source/Programs/AutomationTool/BuildGraph/Tasks/CreateCloudArtifactTask.cs b/Engine/Source/Programs/AutomationTool/BuildGraph/Tasks/CreateCloudArtifactTask.cs
--- a/Engine/Source/Programs/AutomationTool/BuildGraph/Tasks/CreateCloudArtifactTask.cs
+++ b/Engine/Source/Programs/AutomationTool/BuildGraph/Tasks/CreateCloudArtifactTask.cs
@@ -194,7 +194,12 @@ namespace AutomationTool.Tasks
 		/// <param name="tagNameToFileSet">Mapping from tag names to the set of files they include.</param>
 		public override async Task ExecuteAsync(JobContext job, HashSet<FileReference> buildProducts, Dictionary<string, HashSet<FileReference>> tagNameToFileSet)
 		{
-			string name = _parameters.Name;
+			if (String.IsNullOrEmpty(_parameters.Name) && String.IsNullOrEmpty(_parameters.HordeArtifactName))
+			{
+				throw new AutomationException("No name set in either Name or HordeArtifactName");
+			}
+
+			string name = String.IsNullOrEmpty(_parameters.Name) ? _parameters.HordeArtifactName : _parameters.Name;
 			ArtifactType type = new ArtifactType(_parameters.Type);
 			string sanitizedHordeArtifactParam = SanitizeBucketValue(String.IsNullOrEmpty(_parameters.HordeArtifactName) ? name : _parameters.HordeArtifactName);
 			ArtifactName hordeArtifactName = new ArtifactName(sanitizedHordeArtifactParam);
