[Gauntlet] Ignore case when matching flavors of builds

diff --git a/Engine/Source/Programs/AutomationTool/Gauntlet/Unreal/BuildSource/Gauntlet.UnrealBuildSource.cs b/Engine/Source/Programs/AutomationTool/Gauntlet/Unreal/BuildSource/Gauntlet.UnrealBuildSource.cs
--- a/Engine/Source/Programs/AutomationTool/Gauntlet/Unreal/BuildSource/Gauntlet.UnrealBuildSource.cs
+++ b/Engine/Source/Programs/AutomationTool/Gauntlet/Unreal/BuildSource/Gauntlet.UnrealBuildSource.cs
@@ -67,7 +67,7 @@ namespace Gauntlet
 			}
 		}
 
-		public UnrealBuildSource(string InProjectName, FileReference InProjectPath, DirectoryReference InUnrealPath, bool InUsesSharedBuildType, string BuildReference) 
+		public UnrealBuildSource(string InProjectName, FileReference InProjectPath, DirectoryReference InUnrealPath, bool InUsesSharedBuildType, string BuildReference)
 		{
 			InitBuildSource(InProjectName, InProjectPath, InUnrealPath, InUsesSharedBuildType, BuildReference, null);
 		}
@@ -369,7 +369,7 @@ namespace Gauntlet
 				}
 
 				string Flavor = string.IsNullOrEmpty(NewBuild.Flavor) ? "None" : NewBuild.Flavor;
-				Log.Info("Adding {Platform} {BuildType} | Configuration: {Configuration} | Flavor: {Flavor} | Flags: {Flags} | PreferenceOrder: {PreferenceOrder}", 
+				Log.Info("Adding {Platform} {BuildType} | Configuration: {Configuration} | Flavor: {Flavor} | Flags: {Flags} | PreferenceOrder: {PreferenceOrder}",
 					NewBuild.Platform, NewBuild.GetType().Name, NewBuild.Configuration, Flavor, NewBuild.Flags, NewBuild.PreferenceOrder);
 				DiscoveredBuilds[NewBuild.Platform].Add(NewBuild);
 			}
@@ -397,19 +397,19 @@ namespace Gauntlet
 				// check for the editor
 				string EditorExe  = Path.Combine(InUnrealPath.FullName, GetRelativeExecutablePath(UnrealTargetRole.Editor, BuildHostPlatform.Current.Platform, InConfiguration));
 				string OverlayExe = Path.Combine(InUnrealPath.FullName, GetModuleOverlayExecutablePath(UnrealTargetRole.Editor, BuildHostPlatform.Current.Platform));
-				
+
 				if (Utils.SystemHelpers.ApplicationExists(EditorExe))
 				{
 					EditorBuild NewBuild = new EditorBuild(EditorExe, InConfiguration);
 
 					return NewBuild;
 				}
-				
+
 				if (Utils.SystemHelpers.ApplicationExists(OverlayExe))
 				{
 					return new EditorBuild(OverlayExe, InConfiguration);
 				}
-				
+
 				Log.Info("No editor binaries found at {0}. Unable to create an editor build source.", EditorExe);
 			}
 			else
@@ -525,12 +525,12 @@ namespace Gauntlet
 
 			IEnumerable<IBuild> PlatformBuilds = DiscoveredBuilds[InPlatform.Value];
 
-			IEnumerable<IBuild> MatchingBuilds = PlatformBuilds.Where((B) => 
+			IEnumerable<IBuild> MatchingBuilds = PlatformBuilds.Where((B) =>
 			{
 				if (B.CanSupportRole(InRole)
 					&& B.Configuration == InConfiguration
 					&& (B.Flags & InFlags) == InFlags
-					&& (B.Flavor == InFlavor))
+					&& (B.Flavor.Equals(InFlavor, StringComparison.OrdinalIgnoreCase)))
 				{
 					return true;
 				}
@@ -615,7 +615,7 @@ namespace Gauntlet
 				Issues.ForEach(S => Log.Error(S));
 				return null;
 			}
-			
+
 			UnrealTestConfiguration TestConfig = Role.Options as UnrealTestConfiguration;
 
 			UnrealAppConfig Config = new UnrealAppConfig();
@@ -723,7 +723,7 @@ namespace Gauntlet
 		private string GenerateProcessedCommandLine(string InCommandLine)
 		{
 
-			// Break down Commandline into individual tokens 
+			// Break down Commandline into individual tokens
 			Dictionary<string, string> CommandlineTokens = new Dictionary<string, string>(StringComparer.OrdinalIgnoreCase);
 			// turn Name(p1,etc) into a collection of Name|(p1,etc) groups
 			MatchCollection Matches = Regex.Matches(InCommandLine, @"(?<option>-?[\w\d.:!\[\]\/\\?-]+)(=(?<value>(""([^""]*)"")|(\S+)))?");
@@ -824,7 +824,7 @@ namespace Gauntlet
 					{
 						string Message = string.Format("The project config is overriding build targets.\n"
 								+ "But no suitable editor build for {0} configuration found from target file.\n"
-								+ "{1}", TargetConfiguration, Ex.Message); 
+								+ "{1}", TargetConfiguration, Ex.Message);
 
 						if (BuildName.Equals("Editor", StringComparison.OrdinalIgnoreCase))
 						{
