Fix MultiServerNode connection logic to work with non-localhost addresses.
Add "Matchmake" option to LaunchMultiServer UAT script that suppresses the default use of local addresses.
Enable ro.InitRemoteServerIdBeforeUObjectInit by default.

diff --git a/Engine/Plugins/Runtime/MultiServerReplication/Source/MultiServerReplication/Private/MultiServerNode.cpp b/Engine/Plugins/Runtime/MultiServerReplication/Source/MultiServerReplication/Private/MultiServerNode.cpp
--- a/Engine/Plugins/Runtime/MultiServerReplication/Source/MultiServerReplication/Private/MultiServerNode.cpp
+++ b/Engine/Plugins/Runtime/MultiServerReplication/Source/MultiServerReplication/Private/MultiServerNode.cpp
@@ -6,6 +6,7 @@
 #include "MultiServerBeaconClient.h"
 #include "MultiServerPeerConnection.h"
 #include "MultiServerReplicationTypes.h"
+#include "Algo/Sort.h"
 #include "Engine/NetConnection.h"
 #include "Engine/World.h"
 #include "Engine/NetDriver.h"
@@ -128,7 +129,14 @@ bool UMultiServerNode::RegisterServer(const FMultiServerNodeCreateParams& Params
 	}
 	else
 	{
-		for (const FString& PeerAddress : Params.PeerAddresses)
+		// Assume every server wil receive the same list of addresses. Sort it so all servers can deterministically figure out who else they need to connect to.
+		TArray<FString> SortedPeerAddresses = Params.PeerAddresses;
+		Algo::Sort(SortedPeerAddresses);
+
+		// If our local IP wasn't passed in, assume localhost.
+		FString ListenIp = Params.ListenIp.IsEmpty() ? TEXT("127.0.0.1") : Params.ListenIp;
+
+		for (const FString& PeerAddress : SortedPeerAddresses)
 		{
 			FURL PeerURL(nullptr, ToCStr(PeerAddress), ETravelType::TRAVEL_Absolute);
 			
@@ -138,16 +146,10 @@ bool UMultiServerNode::RegisterServer(const FMultiServerNodeCreateParams& Params
 				continue;
 			}
 
-			// Only connect if the peer port is "lower" than the listening port, to prevent redundant connections.
-			// This only works when connecting to instances on the same machine, so limit it to loopback addresses.
-			// TODO: Handle this for remote addresses, or require the filtering to be done at a higher level (command line, etc).
-			if (PeerURL.Host.Equals(TEXT("localhost"), ESearchCase::IgnoreCase) ||
-				PeerURL.Host.StartsWith(TEXT("127."), ESearchCase::IgnoreCase))
+			// Only connect to peers that are before this instance in the list to prevent redundant connections.
+			if (PeerURL.Host.Equals(ListenIp) && PeerURL.Port == Params.ListenPort)
 			{
-				if (Params.ListenPort != 0 && PeerURL.Port >= Params.ListenPort)
-				{
-					continue;
-				}
+				break;
 			}
 
 			if (!PeerAddress.IsEmpty())
diff --git a/Engine/Plugins/Runtime/MultiServerReplication/Source/MultiServerReplication/Public/MultiServerNode.h b/Engine/Plugins/Runtime/MultiServerReplication/Source/MultiServerReplication/Public/MultiServerNode.h
--- a/Engine/Plugins/Runtime/MultiServerReplication/Source/MultiServerReplication/Public/MultiServerNode.h
+++ b/Engine/Plugins/Runtime/MultiServerReplication/Source/MultiServerReplication/Public/MultiServerNode.h
@@ -27,6 +27,9 @@ struct FMultiServerNodeCreateParams
 	/** String identifier of this node. Must be unique among all nodes that will connect to each other. */
 	FString LocalPeerId;
 
+	/** The local IP address of this node, only used so it can find itself in the PeerAddresses list. */
+	FString ListenIp;
+
 	/** The port on which this node will listen for new connections. */ 
 	uint16 ListenPort = 0;
 
diff --git a/Engine/Source/Programs/AutomationTool/Scripts/LaunchMultiServer.Automation.cs b/Engine/Source/Programs/AutomationTool/Scripts/LaunchMultiServer.Automation.cs
--- a/Engine/Source/Programs/AutomationTool/Scripts/LaunchMultiServer.Automation.cs
+++ b/Engine/Source/Programs/AutomationTool/Scripts/LaunchMultiServer.Automation.cs
@@ -287,6 +287,9 @@ Example running 4 clients locally connecting to single proxy server w/ a non-loc
 		[CommandLine, Help("CommonArgs", "Specify arguments that are run by all processes")]
 		public string CommonArgs = "";
 
+		[CommandLine, Help("Matchmake", "Project will use a matchmaking service to connect servers instead of passing the -MultiServerPeers and -ProxyGameServers options to servers.")]
+		public bool Matchmake = false;
+
 		// We should only parse the commandline arguments once, since doing so multiple times will add entries to the List<> args.
 		bool bParsedCommandLineArguments = false;
 
@@ -450,7 +453,12 @@ Example running 4 clients locally connecting to single proxy server w/ a non-loc
 			ServerCommonArgs += NoTimeouts ? " -notimeouts" : "";
 			ServerCommonArgs += IsAttachingDebugger(ProcessLaunchType.Server) ? " -WaitForDebuggerNoBreak" : "";
 			ServerCommonArgs += $" -MultiServerNumServers={NumTotalServers}";
-			ServerCommonArgs += (NonLocalServers.Length > 0) ? $" -MultiServerPeers={string.Join(',', NonLocalServers)}" : "";
+
+			if (!Matchmake)
+			{
+				ServerCommonArgs += (NonLocalServers.Length > 0) ? $" -MultiServerPeers={string.Join(',', NonLocalServers)}" : "";
+				ServerCommonArgs += " -MultiServerLocalHost";
+			}
 
 			for (int ServerId = 1; ServerId <= NumLocalServers; ++ServerId)
 			{
@@ -460,7 +468,7 @@ Example running 4 clients locally connecting to single proxy server w/ a non-loc
 				Logger.LogInformation("Starting Dedicated MultiServer for Game Port {0}", ServerGamePort);
 
 				string ServerArgs = ServerCommonArgs.Replace("{ServerId}", ServerId.ToString(), StringComparison.OrdinalIgnoreCase);
-				ServerArgs += String.Format(" -MultiServerLocalHost -port={0} -log=MultiServer-{1}.log", ServerGamePort, ServerId);
+				ServerArgs += String.Format(" -port={0} -log=MultiServer-{1}.log", ServerGamePort, ServerId);
 
 				string AppTitle = $"Server ID {ServerId}";
 				ExecuteApp(ProcessLaunchType.Server, AppTitle, ServerArgs);
@@ -507,7 +515,7 @@ Example running 4 clients locally connecting to single proxy server w/ a non-loc
 			CommonProxyServerArgs += IsAttachingDebugger(ProcessLaunchType.ProxyServer) ? " -WaitForDebuggerNoBreak" : "";
 			CommonProxyServerArgs += " -port={ProxyServerPort}"; // will be Replaced() below.
 
-			if (AllGameServerAddresses.Length > 0)
+			if (!Matchmake && AllGameServerAddresses.Length > 0)
 			{
 				CommonProxyServerArgs += $" -ProxyGameServers={string.Join(',', AllGameServerAddresses)}";
 			}
diff --git a/Engine/Source/Runtime/CoreUObject/Private/UObject/RemoteObject.cpp b/Engine/Source/Runtime/CoreUObject/Private/UObject/RemoteObject.cpp
--- a/Engine/Source/Runtime/CoreUObject/Private/UObject/RemoteObject.cpp
+++ b/Engine/Source/Runtime/CoreUObject/Private/UObject/RemoteObject.cpp
@@ -25,7 +25,7 @@
 
 DEFINE_LOG_CATEGORY(LogRemoteObject);
 
-int32 GInitRemoteServerIdBeforeUObjectInit = 1;
+int32 GInitRemoteServerIdBeforeUObjectInit = 0;
 static FAutoConsoleVariableRef CVarInitRemoteServerIdBeforeUObjectInit(
 	TEXT("ro.InitRemoteServerIdBeforeUObjectInit"),
 	GInitRemoteServerIdBeforeUObjectInit,
