[Gauntlet] Ignore case when matching flavors of builds while running dev builds

diff --git a/Engine/Source/Programs/AutomationTool/Gauntlet/Unreal/BuildSource/Gauntlet.UnrealBuildSource.cs b/Engine/Source/Programs/AutomationTool/Gauntlet/Unreal/BuildSource/Gauntlet.UnrealBuildSource.cs
--- a/Engine/Source/Programs/AutomationTool/Gauntlet/Unreal/BuildSource/Gauntlet.UnrealBuildSource.cs
+++ b/Engine/Source/Programs/AutomationTool/Gauntlet/Unreal/BuildSource/Gauntlet.UnrealBuildSource.cs
@@ -549,7 +549,7 @@ namespace Gauntlet
 				{
 					if (B.CanSupportRole(InRole)
 						&& (B.Flags & InFlags) == InFlags
-						&& (B.Flavor == InFlavor))
+						&& (B.Flavor.Equals(InFlavor, StringComparison.OrdinalIgnoreCase)))
 					{
 						Log.Warning("Build did not have configuration {0} for {1}, but selecting due to presence of -dev flag",
 							InConfiguration, InPlatform);
