UAT: BuildCookRun fix to include basebuild .sig files when generating patches.
Copy .sig files to the Releases/<ReleaseName> dir when running BCR with -createreleaseversion=<ReleaseName>.
Copy .sig files from the Releases/<ReleaseName> dir when running BCR with -basedonreleaseversion=<ReleaseName> -generatepatch.

diff --git a/Engine/Source/Programs/AutomationTool/Scripts/CopyBuildToStagingDirectory.Automation.cs b/Engine/Source/Programs/AutomationTool/Scripts/CopyBuildToStagingDirectory.Automation.cs
--- a/Engine/Source/Programs/AutomationTool/Scripts/CopyBuildToStagingDirectory.Automation.cs
+++ b/Engine/Source/Programs/AutomationTool/Scripts/CopyBuildToStagingDirectory.Automation.cs
@@ -4504,6 +4504,10 @@ namespace AutomationScripts
 
 					InternalUtils.SafeCreateDirectory(Path.GetDirectoryName(ReleaseVersionPath));
 					InternalUtils.SafeCopyFile(OutputLocation.FullName, ReleaseVersionPath);
+					if (InternalUtils.SafeFileExists(Path.ChangeExtension(OutputLocation.FullName, ".sig")))
+					{
+						InternalUtils.SafeCopyFile(Path.ChangeExtension(OutputLocation.FullName, ".sig"), Path.ChangeExtension(ReleaseVersionPath, ".sig"));
+					}
 					if (bShouldCreateIoStoreContainerFiles)
 					{
 						InternalUtils.SafeCopyFile(Path.ChangeExtension(OutputLocation.FullName, ".utoc"), Path.ChangeExtension(ReleaseVersionPath, ".utoc"));
@@ -4644,6 +4648,7 @@ namespace AutomationScripts
 						{
 							HashSet<string> IncludedExtensions = new HashSet<string>();
 							IncludedExtensions.Add(OutputFilenameExtension);
+							IncludedExtensions.Add(".sig");
 							if (bShouldCreateIoStoreContainerFiles)
 							{
 								IncludedExtensions.Add(".ucas");
