Unix socket aware fetching of cooked file list from ZenStore.

#rnx

diff --git a/Engine/Source/Programs/AutomationTool/Scripts/CopyBuildToStagingDirectory.Automation.cs b/Engine/Source/Programs/AutomationTool/Scripts/CopyBuildToStagingDirectory.Automation.cs
--- a/Engine/Source/Programs/AutomationTool/Scripts/CopyBuildToStagingDirectory.Automation.cs
+++ b/Engine/Source/Programs/AutomationTool/Scripts/CopyBuildToStagingDirectory.Automation.cs
@@ -27,6 +27,7 @@ using System.Text.Json.Serialization;
 using System.Text.Json.Nodes;
 using AutomationUtils;
 using System.Net;
+using System.Net.Sockets;
 
 namespace AutomationScripts
 {
@@ -915,16 +916,62 @@ namespace AutomationScripts
 			return Field.AsHash();
 		}
 
-		private static List<ZenCookedFile> ReadZenCookedFilesFromZenServer(string Host, int Port, string ProjectId, string OplogId, ProjectParams Params, DeploymentContext SC, bool bAutoLaunch, string PackageStoreFileArgName, string PackageStoreFileArgValue)
+		private static HttpClient CreateZenHttpClient(DeploymentContext SC)
+		{
+			if (SC.PackageStoreData == null || SC.PackageStoreData.ZenServerStore == null || !SC.PackageStoreData.ZenServerStore.HostName.EndsWith(".sock"))
+			{
+				return new HttpClient();
+			}
+
+			string UnixSocketPath = SC.PackageStoreData.ZenServerStore.HostName;
+			var SocketsHandler = new SocketsHttpHandler {
+				ConnectCallback = async (Context, Token) => {
+					var Socket = new Socket(AddressFamily.Unix, SocketType.Stream, ProtocolType.IP);
+					var EndPoint = new UnixDomainSocketEndPoint(UnixSocketPath);
+					await Socket.ConnectAsync(EndPoint);
+					return new NetworkStream(Socket, ownsSocket:true);
+				}
+			};
+			return new HttpClient(SocketsHandler);
+		}
+
+		private static string GetZenHostName(DeploymentContext SC)
+		{
+			if (SC.PackageStoreData == null || SC.PackageStoreData.ZenServerStore == null)
+			{
+				return "localhost:8558";
+			}
+
+			var ZenServer = SC.PackageStoreData.ZenServerStore;
+
+			string HostName = ZenServer.HostName;
+			if (HostName.EndsWith(".sock"))
+			{
+				HostName = "unixsocket";
+			}
+			
+			if (ZenServer.HostPort > 0)
+			{
+				HostName += ":" + ZenServer.HostPort;
+			}
+
+			return HostName;
+		}
+
+		private static List<ZenCookedFile> ReadZenCookedFilesFromZenServer(ProjectParams Params, DeploymentContext SC, bool bAutoLaunch, string PackageStoreFileArgName, string PackageStoreFileArgValue)
 		{
 			List<ZenCookedFile> ZenCookedFiles = new List<ZenCookedFile>();
 
+			string ProjectId = SC.PackageStoreData.ZenServerStore.ProjectId;
+			string OplogId = SC.PackageStoreData.ZenServerStore.OplogId;
+
 			bool bAttemptAutoLaunchOnFailure = bAutoLaunch;
 			HttpResponseMessage HttpGetResult = null;
 			for (; ; )
 			{
-				HttpClient HttpClient = new HttpClient();
-				string Uri = string.Format("http://{0}:{1}/prj/{2}/oplog/{3}/entries", Host, Port, ProjectId, OplogId);
+				string HostName = GetZenHostName(SC);
+				HttpClient HttpClient = CreateZenHttpClient(SC);
+				string Uri = string.Format("http://{0}/prj/{1}/oplog/{2}/entries", HostName, ProjectId, OplogId);
 				Uri += "?trim_by_referencedset=true";
 
 				using var Request = new HttpRequestMessage(HttpMethod.Get, Uri);
@@ -949,7 +996,7 @@ namespace AutomationScripts
 					}
 					else
 					{
-						throw new AutomationException(String.Format("Failed sending oplog request to Zen at {0}:{1}. Ensure that the server is running.", Host, Port));
+						throw new AutomationException(String.Format("Failed sending oplog request to Zen at {0}. Ensure that the server is running.", HostName));
 					}
 				}
 			}
@@ -1032,7 +1079,7 @@ namespace AutomationScripts
 					Logger.LogInformation("Reading oplog from Zen...");
 					ZenServerStoreData ZenServer = ParsedProjectStore.ZenServer;
 					SC.PackageStoreData.ZenServerStore = ZenServer;
-					SC.PackageStoreData.ZenCookedFiles = ReadZenCookedFilesFromZenServer(ZenServer.HostName, ZenServer.HostPort, ZenServer.ProjectId, ZenServer.OplogId, Params, SC, true, "ProjectStore", ProjectStoreFile.FullName);
+					SC.PackageStoreData.ZenCookedFiles = ReadZenCookedFilesFromZenServer(Params, SC, true, "ProjectStore", ProjectStoreFile.FullName);
 					return;
 				}
 			}
@@ -1330,7 +1377,7 @@ namespace AutomationScripts
 				{
 					CookedFiles.AddRange(DirectoryReference.EnumerateFiles(SC.PlatformCookDir, "*.projectstore", SearchOption.TopDirectoryOnly).ToList());
 				}
-				HttpClient Client = new HttpClient();
+				HttpClient Client = CreateZenHttpClient(SC);
 				List<Task> TaskList = new List<Task>();
 				foreach (ZenCookedFile CookedFile in SC.PackageStoreData.ZenCookedFiles)
 				{
@@ -4942,7 +4989,7 @@ namespace AutomationScripts
 				{
 					HttpResponseMessage HttpGetResult = null;
 
-					HttpClient HttpClient = new HttpClient();
+					HttpClient HttpClient = CreateZenHttpClient(SC);
 					using var Request = new HttpRequestMessage(HttpMethod.Get, string.Format("http://{0}:{1}/prj/{2}/oplog/{3}/{4}", SC.PackageStoreData.ZenServerStore.HostName, SC.PackageStoreData.ZenServerStore.HostPort, SC.PackageStoreData.ZenServerStore.ProjectId, SC.PackageStoreData.ZenServerStore.OplogId, ZenFile.ChunkId.ToString()));
 					Request.Headers.Add("Accept", "application/octet-stream");
 					try
