[EditorAutomation]

-Allow other automation test implementors to specify an idle log timeout
-Increase the Creative CPP Tests log timeout since they can be a long period wiating during boot where no tests run which can cause tests to timeout despite having a higher maxduration time than the log timeout

#rnx

diff --git a/Engine/Source/Programs/AutomationTool/Gauntlet/Unreal/Automation/UE.Automation.cs b/Engine/Source/Programs/AutomationTool/Gauntlet/Unreal/Automation/UE.Automation.cs
--- a/Engine/Source/Programs/AutomationTool/Gauntlet/Unreal/Automation/UE.Automation.cs
+++ b/Engine/Source/Programs/AutomationTool/Gauntlet/Unreal/Automation/UE.Automation.cs
@@ -463,7 +463,7 @@ namespace UE
 
 		private DateTime LastAutomationEntryTime = DateTime.MinValue;
 
-		private float IdleTimeoutSec = 30 * 60;
+		protected float IdleTimeoutSec = 30 * 60;
 
 		/// Maximum of events display per test
 		protected virtual int MaxEventsDisplayPerTest { get; set; } = 10;
