Improved AFS file transfer - retries and resume for file push

#android

diff --git a/Engine/Plugins/Runtime/AndroidFileServer/Source/AndroidFileServer/Android/java/com/epicgames/unreal/RemoteFileManager.java b/Engine/Plugins/Runtime/AndroidFileServer/Source/AndroidFileServer/Android/java/com/epicgames/unreal/RemoteFileManager.java
--- a/Engine/Plugins/Runtime/AndroidFileServer/Source/AndroidFileServer/Android/java/com/epicgames/unreal/RemoteFileManager.java
+++ b/Engine/Plugins/Runtime/AndroidFileServer/Source/AndroidFileServer/Android/java/com/epicgames/unreal/RemoteFileManager.java
@@ -34,6 +34,7 @@ import java.nio.file.Path;
 import java.nio.file.Paths;
 import java.nio.file.SimpleFileVisitor;
 import java.nio.file.StandardCopyOption;
+import java.nio.file.StandardOpenOption;
 import java.nio.file.attribute.BasicFileAttributes;
 import java.nio.file.attribute.PosixFilePermission;
 import java.nio.file.attribute.PosixFilePermissions;
@@ -53,7 +54,7 @@ import java.net.SocketException;
 public class RemoteFileManager implements Runnable
 {
 	private String TAG = "UEFS";
-	private static final int SERVER_VERSION = 100;
+	private static final int SERVER_VERSION = 101;
 	private static final int ACCEPT_TIMEOUT = 2000;		// milliseconds
 
 	private static final int MAX_OVERFLOW_FILES = 32;
@@ -79,6 +80,7 @@ public class RemoteFileManager implements Runnable
 	private static final int Command_FileRead = 24;
 	private static final int Command_FileWrite = 25;
 	private static final int Command_FileWriteCompressed = 26;
+	private static final int Command_FileWriteHandshake = 27;
 
 	private String externalPackageDir;
 	private String internalPackageDir;
@@ -866,14 +868,13 @@ public class RemoteFileManager implements Runnable
 			if (filesize == 0 && path.toString().startsWith("/proc/"))
 			{
 				ClearResponse();
-				try {
-					BufferedReader reader = new BufferedReader(new FileReader(new File(path.toString())));
+				try (BufferedReader reader = new BufferedReader(new FileReader(new File(path.toString()))))
+				{
 					String line;
 					while ((line = reader.readLine()) != null)
 					{
 						AppendResponse(line + "\n");
 					}
-					reader.close();
 				}
 				catch (IOException ie)
 				{
@@ -931,12 +932,11 @@ public class RemoteFileManager implements Runnable
 
 		while (true)
 		{
-			try
+			try (OutputStream writer = Files.newOutputStream(path))
 			{
 				int bufferSize = 65536 * 16;
 				byte[] buffer = new byte[bufferSize];
 
-				OutputStream writer = Files.newOutputStream(path);
 				while (filesize > 0)
 				{
 					int chunk = filesize < bufferSize ? (int)filesize : bufferSize;
@@ -954,7 +954,6 @@ public class RemoteFileManager implements Runnable
 						dataRemaining -= readBytes;
 					}
 				}
-				writer.close();
 				break;
 			}
 			catch (NoSuchFileException fe)
@@ -980,6 +979,143 @@ public class RemoteFileManager implements Runnable
 //		SendResponse("true");
 	}
 
+	public void CommandFileWriteHandshake(byte[] timestamp, String params)
+	{
+		String fixedFilename = FixupPathLastDir(params, false);
+		Path path = Paths.get(fixedFilename);
+//Log.d(TAG, "FileWriteHandshake: " + path.toString());
+		long filesize = dataRemaining;
+
+		if (Files.isDirectory(path))
+		{
+			Log.d(TAG, "FileWriteHandshake failed: " + path.toString());
+			SendSize(-1);
+			return;
+		}
+
+		// make sure the directory exists if file doesn't
+		int lastIndex = path.getNameCount();
+		Path basedir = path.subpath(0, lastIndex - 1);
+		if (!Files.isDirectory(basedir))
+		{
+			if (!DirCreateRecursive(basedir.toString()))
+			{
+				SendSize(-1);
+				return;
+			}
+		}
+
+		int bufferSize = 65536 * 16;
+		byte[] buffer = new byte[bufferSize];
+
+		// check for inprogress file to resume
+		Path inprogressDataPath = Paths.get(fixedFilename+"_inprogress_.data");
+		Path inprogressTimePath = Paths.get(fixedFilename+"_inprogress_.time");
+		long currentSize = 0;
+		try (InputStream reader = Files.newInputStream(inprogressTimePath))
+		{
+			currentSize = Files.size(inprogressDataPath);
+
+			// validate timestamp to confirm attempting to resume same file
+			int readBytes = reader.read(buffer, 0, 8);
+
+			if (readBytes != 8 || buffer[0] != timestamp[0] || buffer[1] != timestamp[1] && buffer[2] != timestamp[2] || buffer[3] != timestamp[3] ||
+				buffer[4] != timestamp[4] || buffer[5] != timestamp[5] || buffer[6] != timestamp[6] || buffer[7] != timestamp[7])
+			{
+				currentSize = 0;
+			}
+		}
+		catch (Exception e)
+		{
+			currentSize = 0;
+		}
+
+		// write timestamp file if new
+		if (currentSize == 0)
+		{
+			try (OutputStream writer = Files.newOutputStream(inprogressTimePath))
+			{
+				writer.write(timestamp, 0, 8);
+			}
+			catch (Exception e)
+			{
+				Log.d(TAG, "FileWriteHandshake failed: " + path.toString(), e);
+				SendSize(-1);
+				return;
+			}
+		}
+
+		// request data from currentSize (0 for full file)
+		dataRemaining = filesize - currentSize;
+		SendSize(currentSize);
+
+		OutputStream writer = null;
+		try
+		{
+			if (currentSize == 0)
+			{
+				writer = Files.newOutputStream(inprogressDataPath);
+			}
+			else
+			{
+				writer = Files.newOutputStream(inprogressDataPath, StandardOpenOption.APPEND);
+			}
+			
+			while (dataRemaining > 0)
+			{
+				int chunk = dataRemaining < bufferSize ? (int)dataRemaining : bufferSize;
+				int readBytes = inputStream.read(buffer, 0, chunk);
+				if (readBytes == -1)
+				{
+					writer.close();
+					ioException = true;
+					SendResponse("false");
+					return;
+				}
+				if (readBytes > 0)
+				{
+					writer.write(buffer, 0, readBytes);
+					dataRemaining -= readBytes;
+				}
+			}
+			writer.close();
+		}
+		catch (Exception e)
+		{
+			// clean up so start over next time
+			try
+			{
+				if (writer != null)
+				{
+					writer.close();
+				}
+				Files.delete(inprogressTimePath);
+				Files.delete(inprogressDataPath);
+			}
+			catch (Exception e2)
+			{
+			}
+
+			IgnoreData(dataRemaining);
+			Log.d(TAG, "FileWriteHandshake failed: " + path.toString(), e);
+			SendResponse("false");
+			return;
+		}
+
+		// finalize file (delete timestamp and move in progress data to correct filename)
+		try
+		{
+			Files.delete(inprogressTimePath);
+			Files.move(inprogressDataPath, path, StandardCopyOption.REPLACE_EXISTING);
+		}
+		catch (Exception e)
+		{
+			SendResponse("false");
+			return;
+		}
+		SendResponse("true");
+	}
+
 	public void CommandFileWriteCompressedOodle(String params)
 	{
 		Path path = Paths.get(FixupPathLastDir(params.substring(1), false));
@@ -1373,6 +1509,26 @@ public class RemoteFileManager implements Runnable
 				if (ioException) break;
 				CommandFileWriteCompressed(params);
 				return true;
+
+			case Command_FileWriteHandshake:
+				{
+					byte[] timestampData = new byte[8];
+					int readBytes = 0;
+					try
+					{
+						readBytes = inputStream.read(timestampData, 0, 8);
+					}
+					catch (IOException e)
+					{
+						ioException = true;
+					}
+					if (ioException || readBytes != 8) break;
+					dataRemaining -= 8;
+					params = readString();
+					if (ioException) break;
+					CommandFileWriteHandshake(timestampData, params);
+				}
+				return true;
 		}
 		return false;
 	}
diff --git a/Engine/Source/Programs/AutomationTool/Android/AndroidFileClient.cs b/Engine/Source/Programs/AutomationTool/Android/AndroidFileClient.cs
--- a/Engine/Source/Programs/AutomationTool/Android/AndroidFileClient.cs
+++ b/Engine/Source/Programs/AutomationTool/Android/AndroidFileClient.cs
@@ -35,7 +35,9 @@ namespace AutomationTool
 		private const bool TimeBatching = false;
 
 		private const int DefaultPort = 57099;
-		private const int ServerVersion = 100;
+		private const int ServerVersion_Min = 100;
+		private const int ServerVersion_Max = 101;
+		private const int ServerVersion_Handshake = 101;
 		private const int FILE_BUFFERSIZE = 65536 * 16;
 		private const int SOCKET_READBUFFER = 625 * 1024;
 		private const int SOCKET_SENDBUFFER = 625 * 1024;
@@ -46,6 +48,11 @@ namespace AutomationTool
 		private int ServerPort;
 		private int HostPort;
 
+		private string ActivePackageName = "";
+		private string ActiveToken = "";
+		private string ActiveIPAddress = "127.0.0.1";
+		private int ActiveVersion = 0;
+
 		private static OptimalADB adb = new OptimalADB();
 		static string LastListenDevice = "";
 		static bool bHave_ss = true;
@@ -83,6 +90,7 @@ namespace AutomationTool
 		private static int Command_FileRead = 24;
 		private static int Command_FileWrite = 25;
 		private static int Command_FileWriteCompressed = 26;
+		private static int Command_FileWriteHandshake = 27;
 
 		// ------------------ Helpers to deal with ADB ------------------ 
 
@@ -1071,8 +1079,8 @@ namespace AutomationTool
 				int byteRecv = SocketReceive(false, sizePacket);
 				if (byteRecv == 8)
 				{
-					Result = sizePacket[0] + (sizePacket[1] << 8) + (sizePacket[2] << 16) + (sizePacket[3] << 24) +
-							(sizePacket[4] << 32) + ((long)sizePacket[5] << 40) + ((long)sizePacket[6] << 48) + ((long)sizePacket[7] << 56);
+					Result = (long)sizePacket[0] + ((long)sizePacket[1] << 8) + ((long)sizePacket[2] << 16) + ((long)sizePacket[3] << 24) +
+							((long)sizePacket[4] << 32) + ((long)sizePacket[5] << 40) + ((long)sizePacket[6] << 48) + ((long)sizePacket[7] << 56);
 				}
 			}
 			catch (Exception e)
@@ -1083,24 +1091,11 @@ namespace AutomationTool
 			return Result;
 		}
 
-		private string GetStringResult(int Command, string Params)
+		private string GetStringResult()
 		{
 			string Result = null;
-
-			if (ClientSocket == null)
-			{
-				return Result;
-			}
-
-			byte[] paramMessage = Encoding.UTF8.GetBytes(Params + "\0");
-			long paramSize = paramMessage.Length;
-			byte[] commandMessage = CommandPacket(Command, paramSize);
-
 			try
 			{
-				int bytesSent = SocketSend(false, commandMessage);
-				bytesSent = SocketSend(false, paramMessage);
-
 				int resultSize = (int)ResultSize();
 				if (resultSize > 0)
 				{
@@ -1125,6 +1120,33 @@ namespace AutomationTool
 			return Result;
 		}
 
+		private string GetStringResult(int Command, string Params)
+		{
+			string Result = null;
+
+			if (ClientSocket == null)
+			{
+				return Result;
+			}
+
+			byte[] paramMessage = Encoding.UTF8.GetBytes(Params + "\0");
+			long paramSize = paramMessage.Length;
+			byte[] commandMessage = CommandPacket(Command, paramSize);
+
+			try
+			{
+				int bytesSent = SocketSend(false, commandMessage);
+				bytesSent = SocketSend(false, paramMessage);
+				Result = GetStringResult();
+			}
+			catch (Exception e)
+			{
+				Logger.LogError("Unexpected Exception: {Arg0}", e.ToString());
+				CloseConnection();
+			}
+			return Result;
+		}
+
 		private bool GetBoolResult(int Command, string Params)
 		{
 			string Result = GetStringResult(Command, Params);
@@ -1245,6 +1267,13 @@ namespace AutomationTool
 			return ServerPort;
 		}
 
+		public void SetActive(String PackageName, String Token, String Address)
+		{
+			ActivePackageName = PackageName;
+			ActiveToken = Token;
+			ActiveIPAddress = Address;
+		}
+
 		public bool OpenConnection(String Address = "127.0.0.1")
 		{
 			int UsePort = ServerPort;
@@ -1297,8 +1326,9 @@ namespace AutomationTool
 					if (bytesRecv == resultSize)
 					{
 						int version = message[0] + (message[1] << 8);
-						if (version == ServerVersion)
+						if (version >= ServerVersion_Min && version <= ServerVersion_Max)
 						{
+							ActiveVersion = version;
 							Logger.LogInformation("Connected to RemoteFileManager");
 							return true;
 						}
@@ -1936,6 +1966,11 @@ namespace AutomationTool
 					{
 						if (OpenConnection(IPAddress))
 						{
+							// remember if needed to reestablish connection for retries
+							ActivePackageName = PackageName;
+							ActiveToken = Token;
+							ActiveIPAddress = IPAddress;
+
 							return true;
 						}
 					}
@@ -1961,6 +1996,11 @@ namespace AutomationTool
 
 				if (OpenConnection(IPAddress))
 				{
+					// remember if needed to reestablish connection for retries
+					ActivePackageName = PackageName;
+					ActiveToken = Token;
+					ActiveIPAddress = IPAddress;
+
 					return true;
 				}
 
@@ -2220,28 +2260,61 @@ namespace AutomationTool
 			// with network error, try reconnect
 			// with reconnect error, retry reconnect
 			// with reconnect success, try filesystem IO process
+			int TryNum = 0;
 			bool HasNetError = false;
-			for (int TryNum = 0; TryNum < RetriesMax; ++TryNum)
+			while (true)
 			{
-				if (HasNetError)
+				if (!HasNetError)
 				{
-					Logger.LogInformation($"Retry {TryNum + 1}/{RetriesMax} writing file {DestPath} because of network error");
-					Thread.Sleep(RetryMillisecondsDelay);
-					HasNetError = !OpenConnection();
-					continue;
+					(bool Result, HasNetError) = FileWriteInternal(SourcePath, DestPath, bLog);
+					if (!HasNetError)
+					{
+						return Result;
+					}
 				}
 
-				(bool Result, HasNetError) = FileWriteInternal(SourcePath, DestPath, bLog);
-				if (!HasNetError)
+				if (TryNum++ == RetriesMax)
 				{
-					return Result;
+					Logger.LogInformation($"Too many retries writing {DestPath}; giving up");
+					break;
 				}
+
+				Logger.LogInformation($"Retry {TryNum}/{RetriesMax} writing file {DestPath} because of network error");
+				if (!ClientSocket.Connected)
+				{
+					try
+					{
+						ClientSocket.Close();
+						ClientSocket.Dispose();
+					}
+					catch (Exception)
+					{
+					}
+					ClientSocket = null;
+				}
+				Thread.Sleep(RetryMillisecondsDelay);
+
+				bool bUSB;
+				bool bWifi;
+				string WifiAddress;
+				if (GetListenStatus(Device, ServerPort, out bUSB, out bWifi, out WifiAddress))
+				{
+					bool bIsUSB = (ActiveIPAddress == null || ActiveIPAddress == "" || ActiveIPAddress == "127.0.0.1");
+					if ((bUSB && bIsUSB) || (bWifi && !bIsUSB))
+					{
+						HasNetError = !OpenConnection();
+						continue;
+					}
+					HasNetError = !StartServer(ActivePackageName, ActiveToken, ActiveIPAddress);
+					continue;
+				}
+				HasNetError = !OpenConnection();
 			}
 
 			return false;
 		}
 
-		private (bool Result, bool HasNetworkError) FileWriteInternal(string SourcePath, string DestPath, int bLog = 0)
+		private (bool Result, bool HasNetworkError) FileWriteInternal_Normal(string SourcePath, string DestPath, int bLog = 0)
 		{
 			bool Result = false;
 			bool HasNetError = false;
@@ -2301,6 +2374,7 @@ namespace AutomationTool
 					}
 					else
 					{
+						// legacy file write support
 						// reduce the number of SocketSends if not batching
 						byte[] commandMessage = new byte[8 + headerSize];
 						commandMessage[0] = (byte)(Command_FileWrite & 255);
@@ -2344,6 +2418,122 @@ namespace AutomationTool
 			return (Result, HasNetError);
 		}
 
+		private (bool Result, bool HasNetworkError) FileWriteInternal(string SourcePath, string DestPath, int bLog = 0)
+		{
+			if (BatchBuffer != null || ActiveVersion < ServerVersion_Handshake)
+			{
+				return FileWriteInternal_Normal(SourcePath, DestPath, bLog);
+			}
+
+			bool Result = false;
+			bool HasNetError = false;
+			try
+			{
+				using (FileStream fileStream = new FileStream(SourcePath, FileMode.Open, FileAccess.Read, FileShare.None, 4096, FileOptions.SequentialScan))
+				{
+					// deal with LastBaseDir substitution
+					DestPath = OptimizePath(DestPath);
+					if (bLog > 0)
+					{
+						Logger.LogInformation("{bLog}> Writing {DestPath}", bLog, DestPath);
+					}
+
+					long fileSize = fileStream.Length;
+					if (fileSize <= FILE_BUFFERSIZE)
+					{
+						// just start over as full transfer if smaller file
+						return FileWriteInternal_Normal(SourcePath, DestPath, bLog);
+					}
+
+					// get timestamp of file to use for validating resume
+					long Timestamp = File.GetLastWriteTimeUtc(SourcePath).ToBinary();
+
+					int BytesRead;
+					int BytesSent = 0;
+					byte[] buffer = new byte[FILE_BUFFERSIZE];
+
+					// send file size and timestamp
+					byte[] paramMessage = Encoding.UTF8.GetBytes(DestPath + "\0");
+					int headerSize = paramMessage.Length;
+					long paramSize = headerSize + 8 + fileSize;
+
+					// reduce the number of SocketSends
+					byte[] commandMessage = new byte[8 + 8 + headerSize];
+					commandMessage[0] = (byte)(Command_FileWriteHandshake & 255);
+					commandMessage[1] = (byte)((Command_FileWriteHandshake >> 8) & 255);
+					commandMessage[2] = (byte)(paramSize & 255);
+					commandMessage[3] = (byte)((paramSize >> 8) & 255);
+					commandMessage[4] = (byte)((paramSize >> 16) & 255);
+					commandMessage[5] = (byte)((paramSize >> 24) & 255);
+					commandMessage[6] = (byte)((paramSize >> 32) & 255);
+					commandMessage[7] = (byte)((paramSize >> 40) & 255);
+					commandMessage[8] = (byte)(Timestamp & 255);
+					commandMessage[9] = (byte)((Timestamp >> 8) & 255);
+					commandMessage[10] = (byte)((Timestamp >> 16) & 255);
+					commandMessage[11] = (byte)((Timestamp >> 24) & 255);
+					commandMessage[12] = (byte)((Timestamp >> 32) & 255);
+					commandMessage[13] = (byte)((Timestamp >> 40) & 255);
+					commandMessage[14] = (byte)((Timestamp >> 48) & 255);
+					commandMessage[15] = (byte)((Timestamp >> 56) & 255);
+					Buffer.BlockCopy(paramMessage, 0, commandMessage, 16, headerSize);
+					int bytesSent = SocketSend(false, commandMessage);
+
+					// wait for return (error, full file, or place to resume transfer)
+					long resultSize = ResultSize();
+					if (resultSize == -1)
+					{
+						// error with file
+						Result = false;
+					}
+					else if (resultSize == 0)
+					{
+						// send full file
+						Stats_FilesWrite++;
+						while ((BytesRead = fileStream.Read(buffer, 0, FILE_BUFFERSIZE)) != 0)
+						{
+							BytesSent = SocketSend(true, buffer, BytesRead);
+							Stats_CompressedBytesSent += BytesSent;
+						}
+						fileStream.Close();
+						Result = GetStringResult() == "true";
+					}
+					else
+					{
+						// may already be the same size
+						if (resultSize < fileSize)
+						{
+							// send from end of file on device (append)
+							Stats_FilesWrite++;
+							fileStream.Seek(resultSize, SeekOrigin.Begin);
+							while ((BytesRead = fileStream.Read(buffer, 0, FILE_BUFFERSIZE)) != 0)
+							{
+								BytesSent = SocketSend(true, buffer, BytesRead);
+								Stats_CompressedBytesSent += BytesSent;
+							}
+						}
+						fileStream.Close();
+						Result = GetStringResult() == "true";
+					}
+				}
+			}
+			catch (IOException)
+			{
+				// file not found
+			}
+			catch (SocketException e)
+			{
+				HasNetError = true;
+				Logger.LogWarning("{bLog}> Network Exception: {Arg1}", bLog, e.ToString());
+				CloseConnection();
+			}
+			catch (Exception e)
+			{
+				Logger.LogError("{bLog}> Unexpected Exception: {Arg1}", bLog, e.ToString());
+				CloseConnection();
+			}
+			return (Result, HasNetError);
+		}
+
 		public bool FileWriteCompressed(string SourcePath, string DestPath, int bLog = 0)
 		{
 			Boolean Result = false;
diff --git a/Engine/Source/Programs/AutomationTool/Android/AndroidPlatform.Automation.cs b/Engine/Source/Programs/AutomationTool/Android/AndroidPlatform.Automation.cs
--- a/Engine/Source/Programs/AutomationTool/Android/AndroidPlatform.Automation.cs
+++ b/Engine/Source/Programs/AutomationTool/Android/AndroidPlatform.Automation.cs
@@ -2791,6 +2791,11 @@ public class AndroidPlatform : Platform
 					}
 				}
 			}
+			else
+			{
+				// remember for retries
+				client.SetActive(PackageName, AFSToken, IPAddress);
+			}
 
 			// verify we connected to the right server
 			string DevicePackageName = client.Query("^packagename");
diff --git a/Engine/Source/Programs/UnrealAndroidFileTool/Program.cs b/Engine/Source/Programs/UnrealAndroidFileTool/Program.cs
--- a/Engine/Source/Programs/UnrealAndroidFileTool/Program.cs
+++ b/Engine/Source/Programs/UnrealAndroidFileTool/Program.cs
@@ -18,7 +18,7 @@ namespace UnrealAndroidFileTool
 {
 	class Program
 	{
-		private static string AppVersion = "1.0.7";
+		private static string AppVersion = "1.0.8";
 		private static ILogger Logger = Log.Logger;
 
 		static void ShowHelp(string command)
@@ -1144,6 +1144,8 @@ namespace UnrealAndroidFileTool
 			AndroidFileClient? client = new AndroidFileClient(Device, Port);
 			if (client != null)
 			{
+				bool bSetActive = false;
+
 				if (!client.OpenConnection(IPAddress))
 				{
 					if (PackageName == "")
@@ -1159,6 +1161,7 @@ namespace UnrealAndroidFileTool
 							Logger.LogError("Unable to connect to " + Device);
 							return null;
 						}
+						bSetActive = true;
 					}
 				}
 
@@ -1177,11 +1180,18 @@ namespace UnrealAndroidFileTool
 					client.TerminateServer();
 
 					Logger.LogInformation("Trying to start file server {0}", PackageName);
-					if (!client.StartServer(PackageName))
+					if (!client.StartServer(PackageName, Token, IPAddress))
 					{
 						Logger.LogError("Failed to start server");
 						return null;
 					}
+					bSetActive = true;
+				}
+
+				// remember active state for retries
+				if (!bSetActive)
+				{
+					client.SetActive(PackageName, Token, IPAddress);
 				}
 				return client;
 			}
diff --git a/Engine/Source/Programs/UnrealAndroidFileTool/UnrealAndroidFileTool.csproj b/Engine/Source/Programs/UnrealAndroidFileTool/UnrealAndroidFileTool.csproj
--- a/Engine/Source/Programs/UnrealAndroidFileTool/UnrealAndroidFileTool.csproj
+++ b/Engine/Source/Programs/UnrealAndroidFileTool/UnrealAndroidFileTool.csproj
@@ -71,7 +71,7 @@
   </ItemGroup>
 
   <ItemGroup>
-    <PackageReference Include="Microsoft.CodeAnalysis.CSharp" Version="4.2.0" />
+    <PackageReference Include="Microsoft.CodeAnalysis.CSharp" Version="4.11.0" />
     <PackageReference Include="Microsoft.Win32.Registry" Version="5.0.0" />
     <PackageReference Include="OpenTracing" Version="0.12.1" />
     <PackageReference Include="System.Management" Version="4.7.0" />
