[UAT]
* Removed crazy spam in logs about every time something checked if a file exists

diff --git a/Engine/Source/Programs/AutomationTool/AutomationUtils/CommandUtils.cs b/Engine/Source/Programs/AutomationTool/AutomationUtils/CommandUtils.cs
--- a/Engine/Source/Programs/AutomationTool/AutomationUtils/CommandUtils.cs
+++ b/Engine/Source/Programs/AutomationTool/AutomationUtils/CommandUtils.cs
@@ -792,7 +792,7 @@ namespace AutomationTool
 		{
 			// need to remove the quotes before checking to see if it exists
 			var NormalizedFilename = ConvertSeparators(PathSeparator.Default, FileName).Replace("\"", "");
-			return InternalUtils.SafeFileExists(NormalizedFilename);
+			return InternalUtils.SafeFileExists(NormalizedFilename, true);
 		}
 
 		/// <summary>
