[Gauntlet] Fix exception message formatting in TargetDeviceIOS.ctor

diff --git a/Engine/Source/Programs/AutomationTool/Gauntlet/Platform/IOS/Gauntlet.TargetDeviceIOS.cs b/Engine/Source/Programs/AutomationTool/Gauntlet/Platform/IOS/Gauntlet.TargetDeviceIOS.cs
--- a/Engine/Source/Programs/AutomationTool/Gauntlet/Platform/IOS/Gauntlet.TargetDeviceIOS.cs
+++ b/Engine/Source/Programs/AutomationTool/Gauntlet/Platform/IOS/Gauntlet.TargetDeviceIOS.cs
@@ -15,30 +15,30 @@ using UnrealBuildTool;
 using static AutomationTool.CommandUtils;
 
 /** IOS General Notes:
- * 
+ *
  *	Functionality:
- *	
+ *
  *		Tooling for iOS automation is in a tumultuous state.
  *		Historically, ios-deploy was used to manage all aspects of automation.
  *		However, recent changes in Apple's system framework have made ios-deploy obsolete for launching apps for iOS versions 17+ (See for details: [REDACTED_PROFILE_URL]/ios-deploy/issues/588)
  *		Alternative tools offer most - but not all - of the features needed to install, launch, monitor, and debug an app.
  *		Because of this, Gauntlet iOS automation is restricted to running on a Mac host and mixes and matches which tools are used for a given operation.
- *		
+ *
  *		For now, TargetDeviceIOS uses the following tools to support automation:
  *			ios-deploy			- Third party CLI tool, still used for legacy operations and copying files to/from device - [REDACTED_PROFILE_URL]/ios-deploy
  *			devicectl			- Apple's core device framework. Has issues with copying files FROM the device, but handles just about everything else - run 'xcrun devicectl' for details
- *			
+ *
  *		In general:
  *			IF you are using xcode versions 15 or below AND iOS version 16 or below, ios-deploy will be primarily used
  *			IF you are using xcode version 16+, devicectl will be primarily used
  *			IF you are using trying to test on iOS versions 17+, you MUST upgrade to XCode 16 which contains critical app monitoring features. (See https://forums.developer.apple.com/forums/thread/756393)
-			
+
 		*** IMPORTANT ****
 		Testing iOS 17+ and iOS 16 or lower in a single pass is not supported. Doing so can result in undefined behavior.
 		If you need to do this, separate the testing into separate passes targeting each device.
- *			
+ *
  *	Other:
- *	
+ *
  *		- Builds (including local developer builds used with -dev) must already be signed. If built with UBT, signing is already part of the process. Automated re-signing may be added in the future.
  *		- Installing builds requires your device to be added to a mobile provision, and for your build to have been signed with an embedded mobile provision containing the device (See https://developer.apple.com/documentation/xcode/distributing-your-app-to-registered-devices)
  *		- If your app installs content after launch, consider using a bulk build to prevent timeouts. This also lets you more rapidly test -dev apps because code and content are separated
@@ -64,7 +64,7 @@ namespace Gauntlet
 				this.UUID = UUID;
 				this.IOSMajorVersion = IOSMajorVersion;
 				this.bConnected = bConnected;
-			}		
+			}
 		}
 
 		/// <summary>
@@ -176,7 +176,7 @@ namespace Gauntlet
 					Log.Info("Attempting to connect to {UUID}", UUID);
 					if (!Connect())
 					{
-						throw new DeviceException("Failed to connect to {DeviceName}. Ensure the device is properly connected to the host.", UUID);
+						throw new DeviceException("Failed to connect to {0}. Ensure the device is properly connected to the host.", UUID);
 					}
 				}
 				else
@@ -199,7 +199,7 @@ namespace Gauntlet
 					Log.Info("Attempting to connect to {DeviceName}", Name);
 					if(!Connect())
 					{
-						throw new DeviceException("Failed to connect to {DeviceName}", UUID);
+						throw new DeviceException("Failed to connect to {0}", UUID);
 					}
 				}
 			}
@@ -330,7 +330,7 @@ namespace Gauntlet
 		public bool Reboot()
 		{
 			// Rebooting tends to lock the device and break connections requiring manual intervention to unlock the device
-			return true; 
+			return true;
 		}
 
 		public Dictionary<EIntendedBaseCopyDirectory, string> GetPlatformDirectoryMappings()
@@ -422,7 +422,7 @@ namespace Gauntlet
 			// The output of the --list command will print all files in that path
 			// so if the output contains the path, it exists
 			// if no output is provided, the path does not exist
-			
+
 			IProcessResult Result = ExecuteIOSDeployCommand(ListEngineSaved);
 			Result.WaitForExit();
 
@@ -434,7 +434,7 @@ namespace Gauntlet
 			{
 				throw new DeviceException("Listing Engine artifacts for {0} exited with code {1}:\n{2}", BundleID, Result.ExitCode, Result.Output);
 			}
-			
+
 			if (Result.Output.Contains(EngineSavedDirectory))
 			{
 				Result = ExecuteIOSDeployCommand(CleanEngine);
@@ -442,19 +442,19 @@ namespace Gauntlet
 				if (Result.ExitCode != 0)
 				{
 					throw new DeviceException("Cleaning Engine artifacts for {0} exited with code {1}:\n{2}", BundleID, Result.ExitCode, Result.Output);
-				}	
+				}
 			}
 
 			string ListProjectSaved = string.Format("--bundle_id {0} --list={1}", BundleID, ProjectSavedDirectory);
 			string CleanProject = string.Format("--bundle_id {0} --rmtree {1}", BundleID, ProjectSavedDirectory);
-			
+
 			Result = ExecuteIOSDeployCommand(ListProjectSaved);
 			Result.WaitForExit();
 			if (Result.ExitCode != 0)
 			{
 				throw new DeviceException("Listing Project artifacts for {0} exited with code {1}:\n{2}", BundleID, Result.ExitCode, Result.Output);
 			}
-			
+
 			// The same applies here for the --list
 			if (Result.Output.Contains(ListProjectSaved))
 			{
@@ -530,12 +530,12 @@ namespace Gauntlet
 			using(ScopedSuspendECErrorParsing ErrorSuspension = new())
 			{
 				if(UseDeviceCtl)
-				{ 
+				{
 					Result = ExecuteDevicectlCommand(string.Format("device install app {0} -v", BuildPath), MaxInstallTime, AdditionalOptions: ERunOptions.NoStdOutRedirect);
 				}
 				else
 				{
-					Result = ExecuteIOSDeployCommand(string.Format("-b \"{0}\"", BuildPath), MaxInstallTime, AdditionalOptions: ERunOptions.NoStdOutRedirect);			
+					Result = ExecuteIOSDeployCommand(string.Format("-b \"{0}\"", BuildPath), MaxInstallTime, AdditionalOptions: ERunOptions.NoStdOutRedirect);
 				}
 			}
 
@@ -578,7 +578,7 @@ namespace Gauntlet
 								Directory.Delete(TempBuildCache, true);
 							}
 							Directory.CreateDirectory(TempBuildCache);
-									
+
 							Log.Info("Copying {Source} to {Dest}", SourceBulkPath, DestinationBulkPath);
 							SystemHelpers.CopyDirectory(SourceBulkPath, DestinationBulkPath, SystemHelpers.CopyOptions.Copy);
 
@@ -695,7 +695,7 @@ namespace Gauntlet
 				string CopyCommand = string.Format("--bundle_id {0} --upload={1} --to {2}",
 					BundleID, SourceFile, DestinationFile);
 				IProcessResult Result = ExecuteIOSDeployCommand(CopyCommand, 120);
-	
+
 				if(Result.ExitCode != 0)
 				{
 					if (Result.ExitCode != 0)
@@ -711,12 +711,12 @@ namespace Gauntlet
 		{
 			IProcessResult CopyResult;
 
-			if(!Destination.StartsWith("/Documents", StringComparison.OrdinalIgnoreCase) 
+			if(!Destination.StartsWith("/Documents", StringComparison.OrdinalIgnoreCase)
 			&& !Destination.StartsWith("Documents", StringComparison.OrdinalIgnoreCase))
 			{
-				Destination = Destination.StartsWith('/') 
-					? "/Documents" + Destination 
-					: "/Documents/" + Destination;  
+				Destination = Destination.StartsWith('/')
+					? "/Documents" + Destination
+					: "/Documents/" + Destination;
 			}
 
 			if (UseDeviceCtl && !UseIOSDeployForCopy)
@@ -769,7 +769,7 @@ namespace Gauntlet
 			IProcessResult CommandLineCopyResult;
 			if(UseDeviceCtl)
 			{
-				string CopyCommand = string.Format("device copy to --device {0} --source {1} --destination {2} --domain-type appDataContainer --domain-identifier {3}", 
+				string CopyCommand = string.Format("device copy to --device {0} --source {1} --destination {2} --domain-type appDataContainer --domain-identifier {3}",
 					UUID, TempCommandLine, Path.Combine("Documents", "uecommandline.txt"), IOSInstall.PackageName);
 				CommandLineCopyResult = ExecuteDevicectlCommand(CopyCommand, UseDeviceID: false);
 			}
@@ -952,7 +952,7 @@ namespace Gauntlet
 			{
 				Log.Info("{DeploymentExecutable} executing '{Command}'", Executable, CommandLine);
 			}
-			
+
 			ERunOptions RunOptions = ERunOptions.NoWaitForExit | AdditionalOptions;
 			RunOptions |= Log.IsVeryVerbose
 				? ERunOptions.AllowSpew
