Group cook diff report callstacks by [REDACTED_PERSON] first

#rnx

diff --git a/Engine/Source/Programs/AutomationTool/Scripts/CookDiffReportHelper.cs b/Engine/Source/Programs/AutomationTool/Scripts/CookDiffReportHelper.cs
--- a/Engine/Source/Programs/AutomationTool/Scripts/CookDiffReportHelper.cs
+++ b/Engine/Source/Programs/AutomationTool/Scripts/CookDiffReportHelper.cs
@@ -41,6 +41,7 @@ namespace AutomationTool
 			int NumResults = MaxResults;
 			HashSet<Diff> UsedDiffs = [];
 
+			PropertyNamePriorityReport.GenerateReport(Detail, ref NumResults, UsedDiffs, MaxFiles);
 			CallstackPriorityReport.GenerateReport(Detail, ref NumResults, UsedDiffs, MaxFiles);
 			PackagePriorityReport.GenerateReport(Detail, ref NumResults, UsedDiffs);
 
@@ -223,6 +224,135 @@ namespace AutomationTool
 		}
 
 
+		private class PropertyNamePriorityReport
+		{
+			class CallstackInfo
+			{
+				public string Callstack;
+				public HashSet<string> Filenames = [];
+			};
+
+			class PropertyNameInfo
+			{
+				public Dictionary<int,CallstackInfo> CallstackMap = []; // key = callstack hash
+				public HashSet<string> Filenames = [];
+				public Dictionary<string,int> ClassCounts = [];
+				public List<Diff> Diffs = [];
+				public int DiffBytes = 0;
+				public int SizeMismatches = 0;
+			};
+
+			public static void GenerateReport(DetailRoot Detail, ref int MaxItems, HashSet<Diff> UsedDiffs, int MaxFiles)
+			{
+				Dictionary<string,PropertyNameInfo> InfoMap = [];
+
+				foreach (Package Package in Detail.Packages)
+				{
+					foreach (Section Section in Package.Sections)
+					{
+						foreach (Diff Diff in Section.Diffs)
+						{
+							if (string.IsNullOrEmpty(Diff.PropertyName))
+							{
+								continue;
+							}
+
+							if (UsedDiffs.Contains(Diff))
+							{
+								continue;
+							}
+
+							string Callstack = GetCleanCallstack(Diff.Callstack);
+							if (string.IsNullOrEmpty(Callstack))
+							{
+								continue;
+							}
+
+							if (!InfoMap.TryGetValue(Diff.PropertyName, out PropertyNameInfo Info))
+							{
+								Info = new();
+								InfoMap.Add(Diff.PropertyName, Info );
+							}
+
+							int CallstackHash = Callstack.GetHashCode();
+							if (!Info.CallstackMap.TryGetValue(CallstackHash, out CallstackInfo CallstackInfo))
+							{
+								CallstackInfo = new();
+								CallstackInfo.Callstack = Callstack;
+								Info.CallstackMap.Add(CallstackHash, CallstackInfo);
+							}
+							CallstackInfo.Filenames.Add(Section.Filename ?? Package.Filename);
+
+							Info.DiffBytes += (Diff.Count * Diff.Size);
+							Info.Filenames.Add(Section.Filename ?? Package.Filename);
+							Info.Diffs.Add(Diff);
+
+							Info.ClassCounts.TryAdd(Package.ClassName, 0);
+							Info.ClassCounts[Package.ClassName]++;
+
+							if (Section.SizeMismatch)
+							{
+								Info.SizeMismatches++;
+							}
+						}
+					}
+				}
+
+				int NumItems = Math.Min(MaxItems, InfoMap.Count);
+				MaxItems -= NumItems;
+
+
+				foreach (PropertyNameInfo Info in InfoMap.Values)
+				{
+					Info.CallstackMap.OrderByDescending( Pair => Pair.Value.Filenames.Count ); // @todo: may need heuristic here
+				}
+				InfoMap.OrderByDescending( Pair => Pair.Value.Filenames.Count ); // @todo: may need heuristic here
+
+				foreach (KeyValuePair<string,PropertyNameInfo> InfoPair in InfoMap.Take(NumItems))
+				{
+					PropertyNameInfo Info = InfoPair.Value;
+					CallstackInfo CallstackInfo = Info.CallstackMap.First().Value;
+					int NumFiles = Info.Filenames.Count;
+					Logger.LogWarning("Serialization property affecting {files} files, totalling at least {bytes} bytes of nondeterminism{sizemsg}", NumFiles, Info.DiffBytes, (Info.SizeMismatches > 0) ? $" - {Info.SizeMismatches} files have nondeterministic sizes" : "" );
+					Logger.LogInformation("\tSerialized Property: {property}", InfoPair.Key);
+					Logger.LogInformation("\tHighest-Impact Callstack Affecting {count} files ({other} other callstacks):", CallstackInfo.Filenames.Count, Info.CallstackMap.Count-1);
+					foreach (string CallstackLine in CallstackInfo.Callstack.EnumerateLines())
+					{
+						if (!CallstackLine.StartsWith("Serialized Property:")) // don't show this as it's already been shown
+						{
+							Logger.LogInformation("\t\t{line}", CallstackLine);
+						}
+					}
+					Logger.LogInformation("\tAffected classes:");
+					foreach (KeyValuePair<string,int> ClassCountPair in Info.ClassCounts)
+					{
+						Logger.LogInformation("\t\t{class} x {count}", ClassCountPair.Key, ClassCountPair.Value);
+					}
+					
+					if (NumFiles > MaxFiles)
+					{
+						Logger.LogInformation("\tAffected files (showing first {max}):", MaxFiles);
+						NumFiles = MaxFiles;
+					}
+					else
+					{
+						Logger.LogInformation("\tAffected files:");
+					}
+					foreach (string Filename in Info.Filenames.Take(NumFiles))
+					{
+						Logger.LogInformation("\t\t{filename}", GetCleanPackageFilename(Filename));
+					}
+					Logger.LogInformation("");
+
+					foreach (Diff Diff in Info.Diffs)
+					{
+						UsedDiffs.Add(Diff);
+					}
+				}
+			}
+		}
+
+
 
 		private class CallstackPriorityReport
 		{
@@ -238,7 +368,7 @@ namespace AutomationTool
 
 			public static void GenerateReport(DetailRoot Detail, ref int MaxItems, HashSet<Diff> UsedDiffs, int MaxFiles)
 			{
-				Dictionary<int,CallstackInfo> InfoMap = [];
+				Dictionary<int,CallstackInfo> InfoMap = []; // key = callstack hash
 
 				foreach (Package Package in Detail.Packages)
 				{
@@ -246,6 +376,11 @@ namespace AutomationTool
 					{
 						foreach (Diff Diff in Section.Diffs)
 						{
+							if (UsedDiffs.Contains(Diff))
+							{
+								continue;
+							}
+
 							string Callstack = GetCleanCallstack(Diff.Callstack);
 							if (string.IsNullOrEmpty(Callstack))
 							{
