Removed EditorTutorials localization target

This has been unused since UE5

#rnx

diff --git a/Engine/Config/BaseEditor.ini b/Engine/Config/BaseEditor.ini
--- a/Engine/Config/BaseEditor.ini
+++ b/Engine/Config/BaseEditor.ini
@@ -9,7 +9,6 @@
 [Internationalization]
 ; These are the paths for localizing the editor itself, the localized game content settings are in *Game.ini
 +LocalizationPaths=../../../Engine/Content/Localization/Editor
-+LocalizationPaths=../../../Engine/Content/Localization/EditorTutorials
 +LocalizationPaths=../../../Engine/Content/Localization/Keywords
 +LocalizationPaths=../../../Engine/Content/Localization/Category
 +ToolTipLocalizationPaths=../../../Engine/Content/Localization/ToolTips
diff --git a/Engine/Config/Localization/EditorTutorials.ini b/Engine/Config/Localization/EditorTutorials.ini
deleted file mode 100644
diff --git a/Engine/Source/Programs/AutomationTool/CookedEditor/CookedEditor.Automation.cs b/Engine/Source/Programs/AutomationTool/CookedEditor/CookedEditor.Automation.cs
--- a/Engine/Source/Programs/AutomationTool/CookedEditor/CookedEditor.Automation.cs
+++ b/Engine/Source/Programs/AutomationTool/CookedEditor/CookedEditor.Automation.cs
@@ -465,7 +465,7 @@ public class MakeCookedEditor : BuildCommand
 		{
 			List<string> CulturesToStage = Project.GetCulturesToStage(Params, ConfigHelper.GameConfig);
 
-			string[] EditorLocalizationTargetsToStage = { "Category", "Engine", "Editor", "EditorTutorials", "Keywords", "PropertyNames", "ToolTips" };
+			string[] EditorLocalizationTargetsToStage = { "Category", "Engine", "Editor", "Keywords", "PropertyNames", "ToolTips" };
 			foreach (string EditorLocalizationTargetToStage in EditorLocalizationTargetsToStage)
 			{
 				// Note: We skip "ShouldStageLocalizationTarget" below as games may disable certain targets that they don't need at runtime (eg, Engine), but all of these targets are still needed for an editor!
