Optionally expect image comparison json metadata

#rnx

diff --git a/Engine/Source/Programs/AutomationTool/Gauntlet/Framework/Utils/Gauntlet.HordeReport.cs b/Engine/Source/Programs/AutomationTool/Gauntlet/Framework/Utils/Gauntlet.HordeReport.cs
--- a/Engine/Source/Programs/AutomationTool/Gauntlet/Framework/Utils/Gauntlet.HordeReport.cs
+++ b/Engine/Source/Programs/AutomationTool/Gauntlet/Framework/Utils/Gauntlet.HordeReport.cs
@@ -2024,20 +2024,20 @@ namespace Gauntlet
 			/// </summary>
 			public class ImageComparisonFiles : ArtifactFiles, ILogArtifact
 			{
-				public ImageComparisonFiles(string InName, Dictionary<string, string> InFiles, string InRootPath) : base(TestLogValueTypes.imageCompare, InName, InFiles, InRootPath)
+				public ImageComparisonFiles(string InName, Dictionary<string, string> InFiles, string InRootPath, bool AddMetadataFiles = false) : base(TestLogValueTypes.imageCompare, InName, InFiles, InRootPath)
 				{
 					string FilePath = Files.GetValueOrDefault(TestLogValueTypes.unapproved);
-					if (!string.IsNullOrEmpty(FilePath))
+					if (AddMetadataFiles && !string.IsNullOrEmpty(FilePath))
 					{
 						Files.Add(TestLogValueTypes.unapprovedMetadata, Path.ChangeExtension(FilePath, ".json"));
 					}
 					FilePath = Files.GetValueOrDefault(TestLogValueTypes.approved);
-					if (!string.IsNullOrEmpty(FilePath))
+					if (AddMetadataFiles && !string.IsNullOrEmpty(FilePath))
 					{
 						Files.Add(TestLogValueTypes.approvedMetadata, Path.ChangeExtension(FilePath, ".json"));
 					}
 					FilePath = Files.GetValueOrDefault(TestLogValueTypes.difference);
-					if (!string.IsNullOrEmpty(FilePath))
+					if (AddMetadataFiles && !string.IsNullOrEmpty(FilePath))
 					{
 						Files.Add(TestLogValueTypes.differenceReport, Path.Combine(Path.GetDirectoryName(FilePath), "Report.json"));
 					}
@@ -2281,7 +2281,7 @@ namespace Gauntlet
 									string EntryName = ArtifactEntry.Name;
 									if (ArtifactEntry.Type == "Comparison")
 									{
-										ImageComparisonFiles ComparisonFiles = new ImageComparisonFiles(EntryName, ArtifactEntry.Files, InReportPath);
+										ImageComparisonFiles ComparisonFiles = new ImageComparisonFiles(EntryName, ArtifactEntry.Files, InReportPath, true);
 										string Format = InTestEntry.Event.Message.Replace(EntryName, $"{{{ComparisonFiles.Token}}}");
 										Stream.AddImageComparison(EventTime, LogLevel, Format, ComparisonFiles);
 
