fix incorrect usage of OrderByDescending - was discarding the return value

#rnx

diff --git a/Engine/Source/Programs/AutomationTool/Scripts/CookDiffReportHelper.cs b/Engine/Source/Programs/AutomationTool/Scripts/CookDiffReportHelper.cs
--- a/Engine/Source/Programs/AutomationTool/Scripts/CookDiffReportHelper.cs
+++ b/Engine/Source/Programs/AutomationTool/Scripts/CookDiffReportHelper.cs
@@ -301,17 +301,12 @@ namespace AutomationTool
 				int NumItems = Math.Min(MaxItems, InfoMap.Count);
 				MaxItems -= NumItems;
 
-
-				foreach (PropertyNameInfo Info in InfoMap.Values)
-				{
-					Info.CallstackMap.OrderByDescending( Pair => Pair.Value.Filenames.Count ); // @todo: may need heuristic here
-				}
-				InfoMap.OrderByDescending( Pair => Pair.Value.Filenames.Count ); // @todo: may need heuristic here
-
-				foreach (KeyValuePair<string,PropertyNameInfo> InfoPair in InfoMap.Take(NumItems))
+				IOrderedEnumerable<KeyValuePair<string,PropertyNameInfo>> OrderedInfoMap = InfoMap.OrderByDescending( Pair => Pair.Value.Filenames.Count ); // @todo: may need heuristic here
+				foreach (KeyValuePair<string,PropertyNameInfo> InfoPair in OrderedInfoMap.Take(NumItems))
 				{
 					PropertyNameInfo Info = InfoPair.Value;
-					CallstackInfo CallstackInfo = Info.CallstackMap.First().Value;
+					IOrderedEnumerable<KeyValuePair<int,CallstackInfo>> OrderedCallstackMap = Info.CallstackMap.OrderByDescending( Pair => Pair.Value.Filenames.Count ); // @todo: may need heuristic here
+					CallstackInfo CallstackInfo = OrderedCallstackMap.First().Value; 
 					int NumFiles = Info.Filenames.Count;
 					Logger.LogWarning("Serialization property affecting {files} files, totalling at least {bytes} bytes of nondeterminism{sizemsg}", NumFiles, Info.DiffBytes, (Info.SizeMismatches > 0) ? $" - {Info.SizeMismatches} files have nondeterministic sizes" : "" );
 					Logger.LogInformation("\tSerialized Property: {property}", InfoPair.Key);
@@ -413,9 +408,10 @@ namespace AutomationTool
 				int NumItems = Math.Min(MaxItems, InfoMap.Count);
 				MaxItems -= NumItems;
 
-				InfoMap.OrderByDescending( Pair => Pair.Value.Filenames.Count ); // @todo: may need heuristic here
+				
 
-				foreach (KeyValuePair<int,CallstackInfo> InfoPair in InfoMap.Take(NumItems))
+				IOrderedEnumerable<KeyValuePair<int,CallstackInfo>> OrderedInfoMap = InfoMap.OrderByDescending( Pair => Pair.Value.Filenames.Count ); // @todo: may need heuristic here
+				foreach (KeyValuePair<int,CallstackInfo> InfoPair in OrderedInfoMap.Take(NumItems))
 				{
 					CallstackInfo Info = InfoPair.Value;
 					int NumFiles = Info.Filenames.Count;
@@ -457,9 +453,8 @@ namespace AutomationTool
 		{
 			public static void GenerateReport(DetailRoot Detail, ref int MaxItems, HashSet<Diff> UsedDiffs)
 			{
-				Detail.Packages.OrderByDescending( P => CalculateImportance(P, UsedDiffs) );
-
-				foreach (Package Package in Detail.Packages)
+				IOrderedEnumerable<Package> OrderedPackages = Detail.Packages.OrderByDescending( P => CalculateImportance(P, UsedDiffs) );
+				foreach (Package Package in OrderedPackages)
 				{
 					MaxItems--;
 					if (MaxItems < 0)
