PGO Profile for Android

Updated automation code to support AFS for `CopyAdditionalFiles`.

#rnx

diff --git a/Engine/Source/Programs/AutomationTool/Gauntlet/Platform/Android/Gauntlet.AndroidPGO.cs b/Engine/Source/Programs/AutomationTool/Gauntlet/Platform/Android/Gauntlet.AndroidPGO.cs
--- a/Engine/Source/Programs/AutomationTool/Gauntlet/Platform/Android/Gauntlet.AndroidPGO.cs
+++ b/Engine/Source/Programs/AutomationTool/Gauntlet/Platform/Android/Gauntlet.AndroidPGO.cs
@@ -4,6 +4,7 @@ using System;
 using System.IO;
 using System.Text;
 using System.Collections.Generic;
+using System.Runtime.InteropServices;
 using System.Diagnostics;
 using AutomationTool;
 using UnrealBuildTool;
@@ -15,10 +16,22 @@ namespace Gauntlet
 	/// </summary>
 	internal class AndroidPGOPlatform : IPGOPlatform
 	{
-		private readonly string AndroidLlvmProfdataPath = Path.Combine(Environment.ExpandEnvironmentVariables("%NDKROOT%"), @"toolchains\llvm\prebuilt\windows-x86_64\bin\llvm-profdata.exe");
+		private readonly string AndroidLlvmProfdataPath;
 		string LocalOutputDirectory;
 		string LocalProfDataFile;
 
+		public AndroidPGOPlatform()
+		{
+			if (RuntimeInformation.IsOSPlatform(OSPlatform.Windows))
+			{
+				AndroidLlvmProfdataPath = Path.Combine(Environment.ExpandEnvironmentVariables("%NDKROOT%"), @"toolchains\llvm\prebuilt\windows-x86_64\bin\llvm-profdata.exe");
+			}
+			else
+			{
+				AndroidLlvmProfdataPath = Path.Combine(Environment.ExpandEnvironmentVariables("%NDKROOT%"), "toolchains/llvm/prebuilt/darwin-x86_64/bin/llvm-profdata");
+			}
+		}
+
 		public UnrealTargetPlatform GetPlatform()
 		{
 			return UnrealTargetPlatform.Android;
diff --git a/Engine/Source/Programs/AutomationTool/Gauntlet/Platform/Android/Gauntlet.TargetDeviceAndroid.cs b/Engine/Source/Programs/AutomationTool/Gauntlet/Platform/Android/Gauntlet.TargetDeviceAndroid.cs
--- a/Engine/Source/Programs/AutomationTool/Gauntlet/Platform/Android/Gauntlet.TargetDeviceAndroid.cs
+++ b/Engine/Source/Programs/AutomationTool/Gauntlet/Platform/Android/Gauntlet.TargetDeviceAndroid.cs
@@ -113,6 +113,8 @@ namespace Gauntlet
 
 		private string NoPlayProtectSetting = null;
 
+		private AndroidAppInstall InstallCache = null;
+
 		public TargetDeviceAndroid(string InDeviceName = "", AndroidDeviceData DeviceData = null, string InCachePath = null)
 		{
 			AutoParam.ApplyParamsAndDefaults(this, Globals.Params.AllArguments);
@@ -414,9 +416,11 @@ namespace Gauntlet
 			DeviceLogPath = string.Format("{0}/Logs/{1}.log", Build.UsesPublicLogs ? DeviceExternalFilesSavedPath : DeviceExternalStorageSavedPath, AppConfig.ProjectName);
 			DeviceArtifactPath = Build.UsesExternalFilesDir ? DeviceExternalFilesSavedPath : DeviceExternalStorageSavedPath;
 
-			PopulateDirectoryMappings(Path.GetDirectoryName(DeviceArtifactPath));
+			PopulateDirectoryMappings(AppConfig);
+
+			InstallCache = new AndroidAppInstall(AppConfig, this, AppConfig.ProjectName, Build.AndroidPackageName, AppConfig.CommandLine, AppConfig.ProjectFile, AppConfig.Build.Configuration);
 
-			return new AndroidAppInstall(AppConfig, this, AppConfig.ProjectName, Build.AndroidPackageName, AppConfig.CommandLine, AppConfig.ProjectFile, AppConfig.Build.Configuration);
+			return InstallCache;
 		}
 
 		public void CopyAdditionalFiles(IEnumerable<UnrealFileToCopy> FilesToCopy)
@@ -428,7 +432,8 @@ namespace Gauntlet
 				{
 					FileInfo SourceFile = new FileInfo(FileToCopy.SourceFileLocation);
 					SourceFile.IsReadOnly = false;
-					CopyFileToDevice(null, SourceFile.FullName, DestinationFile); // todo: add AFS support
+
+					CopyFileToDevice(InstallCache.AndroidPackageName, SourceFile.FullName, DestinationFile, InstallCache.ProjectFile, InstallCache.Configuration);
 				}
 				else
 				{
@@ -475,6 +480,27 @@ namespace Gauntlet
 				RunAdbDeviceCommand(string.Format("shell mkdir -p {0}/", DeviceExternalFilesSavedPath));
 			}
 
+			if (!Globals.IsRunningDev && UseAFS(Install.AppConfig))
+			{
+				AndroidBuild Build = Install.AppConfig.Build as AndroidBuild;
+				if (Path.GetFileName(Build.SourceApkPath).StartsWith("AFS_"))
+				{
+					Log.Info("Installing non-AFS apk before Launching the {0}", Build.AndroidPackageName);
+
+					string QuotedSourcePath = Build.SourceApkPath.Replace("AFS_", "");
+					if (QuotedSourcePath.Contains(" "))
+					{
+						QuotedSourcePath = '"' + QuotedSourcePath + '"';
+					}
+
+					IProcessResult AdbResult = RunAdbDeviceCommand(string.Format("install -r {0}", QuotedSourcePath));
+					if (AdbResult.ExitCode != 0)
+					{
+						throw new DeviceException("Failed to install non-AFS APK {0}. Error {1}", QuotedSourcePath, AdbResult.Output);
+					}
+				}
+			}
+
 			// start the app on device!
 			string CommandLine = "shell am start -W -S -n " + Install.AndroidPackageName + "/" + LaunchActivity;
 			IProcessResult Process = RunAdbDeviceCommand(CommandLine, false, true);
@@ -483,24 +509,39 @@ namespace Gauntlet
 		}
 		#endregion
 
-		public void PopulateDirectoryMappings(string ProjectDir)
+		public void PopulateDirectoryMappings(UnrealAppConfig AppConfig)
 		{
 			LocalDirectoryMappings.Clear();
 
-			ProjectDir = ProjectDir.Replace('\\', '/');
-			if (!ProjectDir.EndsWith('/'))
+			if (UseAFS(AppConfig))
 			{
-				ProjectDir += '/';
+				LocalDirectoryMappings.Add(EIntendedBaseCopyDirectory.Build, "^build");
+				LocalDirectoryMappings.Add(EIntendedBaseCopyDirectory.Binaries, "^binaries");
+				LocalDirectoryMappings.Add(EIntendedBaseCopyDirectory.Config, "^saved/Config");
+				LocalDirectoryMappings.Add(EIntendedBaseCopyDirectory.Content, "^content");
+				LocalDirectoryMappings.Add(EIntendedBaseCopyDirectory.Demos, "^saved/Demos");
+				LocalDirectoryMappings.Add(EIntendedBaseCopyDirectory.PersistentDownloadDir, "^saved/PersistentDownloadDir");
+				LocalDirectoryMappings.Add(EIntendedBaseCopyDirectory.Profiling, "^saved/Profiling");
+				LocalDirectoryMappings.Add(EIntendedBaseCopyDirectory.Saved, "^saved");
+			}
+			else
+			{
+				string ProjectDir = Path.GetDirectoryName(DeviceArtifactPath);
+				ProjectDir = ProjectDir.Replace('\\', '/');
+				if (!ProjectDir.EndsWith('/'))
+				{
+					ProjectDir += '/';
+				}
+				
+				LocalDirectoryMappings.Add(EIntendedBaseCopyDirectory.Build, ProjectDir + "Build");
+				LocalDirectoryMappings.Add(EIntendedBaseCopyDirectory.Binaries, ProjectDir + "Binaries");
+				LocalDirectoryMappings.Add(EIntendedBaseCopyDirectory.Config, ProjectDir + "Saved/Config");
+				LocalDirectoryMappings.Add(EIntendedBaseCopyDirectory.Content, ProjectDir + "Content");
+				LocalDirectoryMappings.Add(EIntendedBaseCopyDirectory.Demos, ProjectDir + "Saved/Demos");
+				LocalDirectoryMappings.Add(EIntendedBaseCopyDirectory.PersistentDownloadDir, ProjectDir + "Saved/PersistentDownloadDir");
+				LocalDirectoryMappings.Add(EIntendedBaseCopyDirectory.Profiling, ProjectDir + "Saved/Profiling");
+				LocalDirectoryMappings.Add(EIntendedBaseCopyDirectory.Saved, ProjectDir + "Saved");
 			}
-
-			LocalDirectoryMappings.Add(EIntendedBaseCopyDirectory.Build, ProjectDir + "Build");
-			LocalDirectoryMappings.Add(EIntendedBaseCopyDirectory.Binaries, ProjectDir + "Binaries");
-			LocalDirectoryMappings.Add(EIntendedBaseCopyDirectory.Config, ProjectDir + "Saved/Config");
-			LocalDirectoryMappings.Add(EIntendedBaseCopyDirectory.Content, ProjectDir + "Content");
-			LocalDirectoryMappings.Add(EIntendedBaseCopyDirectory.Demos, ProjectDir + "Saved/Demos");
-			LocalDirectoryMappings.Add(EIntendedBaseCopyDirectory.PersistentDownloadDir, ProjectDir + "Saved/PersistentDownloadDir");
-			LocalDirectoryMappings.Add(EIntendedBaseCopyDirectory.Profiling, ProjectDir + "Saved/Profiling");
-			LocalDirectoryMappings.Add(EIntendedBaseCopyDirectory.Saved, ProjectDir + "Saved");
 		}
 
 		public static bool UsingAndroidFileServer(FileReference RawProjectPath, UnrealTargetConfiguration TargetConfiguration, out bool bEnablePlugin, out string AFSToken, out bool bIsShipping, out bool bIncludeInShipping, out bool bAllowExternalStartInShipping)
@@ -802,6 +843,15 @@ namespace Gauntlet
 			// we write this out to a dependency file in the demote dir and check it each time.
 			// current file time
 			DateTime LocalModifiedTime = File.GetLastWriteTime(SourcePath);
+			
+			// If not using AFS and the APK starts with AFS_, rename to use the normal one
+			if (!Globals.IsRunningDev && IsAPK && !UseAFS)
+			{
+				if (Path.GetFileName(SourcePath).StartsWith("AFS_"))
+				{
+					SourcePath = SourcePath.Replace("AFS_", "");
+				}
+			}
 
 			string QuotedSourcePath = SourcePath;
 			if (SourcePath.Contains(" "))
@@ -1310,7 +1360,7 @@ namespace Gauntlet
 				if (LocalDirectoryMappings.Count == 0)
 				{
 					Console.WriteLine("Populating Directory");
-					PopulateDirectoryMappings(Path.GetDirectoryName(DeviceArtifactPath));
+					PopulateDirectoryMappings(AppConfig);
 				}
 				foreach (UnrealFileToCopy FileToCopy in AppConfig.FilesToCopy)
 				{
