Adjust Editor BootTest metadata to fall under BootTest metadata

#rnx

diff --git a/Engine/Source/Programs/AutomationTool/Gauntlet/Unreal/Automation/UE.BootTest.cs b/Engine/Source/Programs/AutomationTool/Gauntlet/Unreal/Automation/UE.BootTest.cs
--- a/Engine/Source/Programs/AutomationTool/Gauntlet/Unreal/Automation/UE.BootTest.cs
+++ b/Engine/Source/Programs/AutomationTool/Gauntlet/Unreal/Automation/UE.BootTest.cs
@@ -45,6 +45,9 @@ namespace UE
 		{
 		}
 
+		protected override string HordeReportTestKey => $"UE.BootTest({Suite})";
+		protected override string HordeReportTestName => "BootTest";
+
 		/// <summary>
 		/// Returns the configuration description for this test
 		/// </summary>
