Allow the -ProxyGameServers command-line argument to accept a port range.  This helps with launching multiple servers on a single host, so you can just use 127.0.0.1:[REDACTED_PHONE] for 25 servers.

diff --git a/Engine/Plugins/Runtime/MultiServerReplication/Source/MultiServerReplication/Private/MultiServerProxy.cpp b/Engine/Plugins/Runtime/MultiServerReplication/Source/MultiServerReplication/Private/MultiServerProxy.cpp
--- a/Engine/Plugins/Runtime/MultiServerReplication/Source/MultiServerReplication/Private/MultiServerProxy.cpp
+++ b/Engine/Plugins/Runtime/MultiServerReplication/Source/MultiServerReplication/Private/MultiServerProxy.cpp
@@ -631,6 +631,57 @@ FGameServerConnectionState* UProxyNetDriver::GetGameServerConnection(int32 Index
 	return nullptr;
 }
 
+// Given a comma-separated list of strings (IpAddress1:StartPort-EndPort,IpAddress2:StartPort-EndPort)
+// Return the list of actual servers addresses: { IpAddress1:StartPort, ..., IpAddress1:EndPort, IpAddress2:StartPort, ... }
+// This allows us to use short-hand when connecting to multiple servers.
+TArray<FString> ParseServerAddressList(const FString& InServerAddressesStr)
+{
+	TArray<FString> InAddresses;
+	InServerAddressesStr.ParseIntoArray(InAddresses, TEXT(","), true);
+
+	TArray<FString> OutAddresses;
+	for (const FString& InAddress : InAddresses)
+	{
+		// See if the address comes with a port (e.g. IpAddress:PortRange)
+		int32 PortSeparatorIdx = INDEX_NONE;
+		if (InAddress.FindLastChar(TEXT(':'), PortSeparatorIdx))
+		{
+			FString AddressPart = InAddress.Left(PortSeparatorIdx);
+			FString PortPart = InAddress.Mid(PortSeparatorIdx + 1);
+
+			// See if the PortRange is actually a port range (:StartPort-EndPort) or just (:Port)
+			int32 RangeSeparatorIdx = INDEX_NONE;
+			if (PortPart.FindLastChar(TEXT('-'), RangeSeparatorIdx))
+			{
+				FString StartPortView = PortPart.Left(RangeSeparatorIdx);
+				FString EndPortView = PortPart.Mid(RangeSeparatorIdx + 1);
+				int32 StartPort = FCString::Atoi(*StartPortView);
+				int32 EndPort = FCString::Atoi(*EndPortView);
+
+				if (StartPort > 0 && EndPort >= StartPort)
+				{
+					for (int32 Port = StartPort; Port <= EndPort; ++Port)
+					{
+						OutAddresses.Add(FString::Printf(TEXT("%s:%d"), *AddressPart, Port));
+					}
+				}
+			}
+			else
+			{
+				// Just put the address in as typed (we don't need to verify the port)
+				OutAddresses.Emplace(InAddress);
+			}
+		}
+		else
+		{
+			// Just put the address in as typed (no port specified, let the system determine what to do)
+			OutAddresses.Emplace(InAddress);
+		}
+	}
+
+	return OutAddresses;
+}
+
 bool UProxyNetDriver::InitBase(bool bInitAsClient, FNetworkNotify* InNotify, const FURL& URL, bool bReuseAddressAndPort, FString& Error)
 {
 	check(!bInitAsClient);
@@ -645,12 +696,11 @@ bool UProxyNetDriver::InitBase(bool bInitAsClient, FNetworkNotify* InNotify, con
 	FString GameServerAddresses;
 	if (FParse::Value(FCommandLine::Get(), TEXT("-ProxyGameServers="), GameServerAddresses, false))
 	{
-		TArray<FString> Addresses;
-		GameServerAddresses.ParseIntoArray(Addresses, TEXT(","), true);
+		TArray<FString> AllGameServerAddresses = ParseServerAddressList(GameServerAddresses);
 
-		for (const FString& Address : Addresses)
+		for (const FString& Address : AllGameServerAddresses)
 		{
-			FURL GameServerURL { nullptr, *Address, ETravelType::TRAVEL_Absolute };
+			FURL GameServerURL{ nullptr, *Address, ETravelType::TRAVEL_Absolute };
 			if (ensureMsgf(GameServerURL.Valid, TEXT("Attempting to register an invalid ProxyGameServer URL. Input: %s. Parsed URL: %s"), *Address, *GameServerURL.ToString()))
 			{
 				RegisterGameServer(GameServerURL);
diff --git a/Engine/Source/Programs/AutomationTool/Scripts/LaunchMultiServer.Automation.cs b/Engine/Source/Programs/AutomationTool/Scripts/LaunchMultiServer.Automation.cs
--- a/Engine/Source/Programs/AutomationTool/Scripts/LaunchMultiServer.Automation.cs
+++ b/Engine/Source/Programs/AutomationTool/Scripts/LaunchMultiServer.Automation.cs
@@ -351,7 +351,7 @@ Example running 4 clients locally connecting to single proxy server w/ a non-loc
 				if (ClientValue.ToLower() == "direct")
 				{
 					Logger.LogInformation("Starting direct client instances connecting to MultiServer instances");
-					StartProcessesForClients(ProcessLaunchType.DirectClient, LocalGameServerAddresses, CommonArgs);
+					StartProcessesForClients(ProcessLaunchType.DirectClient, AllGameServerAddresses, CommonArgs);
 				}
 				else if (ClientValue.ToLower() == "proxy")
 				{
@@ -446,7 +446,7 @@ Example running 4 clients locally connecting to single proxy server w/ a non-loc
 			}
 
 			int NumTotalServers = NumLocalServers + NonLocalServers.Length;
-			string[] LocalGameServerAddresses = new string[NumTotalServers];
+			string[] LocalGameServerAddresses = new string[1];
 
 			// This is optional to the multi-server process, these are in addition to the MultiServerNumPeers argument
 			string ServerCommonArgs = $" -server {AdditionalServerArguments} -NODEBUGOUTPUT";
@@ -463,7 +463,6 @@ Example running 4 clients locally connecting to single proxy server w/ a non-loc
 			for (int ServerId = 1; ServerId <= NumLocalServers; ++ServerId)
 			{
 				int ServerGamePort = BaseGameListenPort + ServerId;
-				LocalGameServerAddresses[ServerId - 1] = $"127.0.0.1:{ServerGamePort}";
 
 				Logger.LogInformation("Starting Dedicated MultiServer for Game Port {0}", ServerGamePort);
 
@@ -474,6 +473,12 @@ Example running 4 clients locally connecting to single proxy server w/ a non-loc
 				ExecuteApp(ProcessLaunchType.Server, AppTitle, ServerArgs);
 			}
 
+			// If there's a range of addresses, give us a range
+			if (NumLocalServers > 1)
+				LocalGameServerAddresses[0] = $"127.0.0.1:{BaseGameListenPort + 1}-{BaseGameListenPort + NumLocalServers}";
+			else
+				LocalGameServerAddresses[0] = $"127.0.0.1:{BaseGameListenPort + 1}";
+
 			return LocalGameServerAddresses;
 		}
 
@@ -481,28 +486,40 @@ Example running 4 clients locally connecting to single proxy server w/ a non-loc
 		/// Start the Client processes
 		/// </summary>
 		/// <param name="ClientType">What type of Client we're launching which will help set the titles properly</param>
-		/// <param name="GameServerAddresses">The game servers (or proxies) that a client will connect to.  There will be ClientsPerServer launched per entry</param>
+		/// <param name="GameServerAddresses">The game servers (or proxies) that a client will connect to.  The addresses may contain a port range (e.g. IpAddress:StartPort-EndPort).</param>
 		/// <param name="AdditionalClientArguments">What additional parameters we launch the client with</param>
 		/// <param name="ClientsPerServer">The number of clients to launch per-GameServerAddress entry</param>
 		virtual protected void StartProcessesForClients(ProcessLaunchType ClientType, string[] GameServerAddresses, string AdditionalClientArguments, int ClientsPerServer = 1)
 		{
 			for (int Idx= 0; Idx < GameServerAddresses.Length; ++Idx)
 			{
-				string GameServerAddress = GameServerAddresses[Idx];
+				string GameServerAddressWithPortRange = GameServerAddresses[Idx];
+				string GameServerBaseAddress;
+				int StartPort, EndPort;
 
-				string ClientArgs = $" -game {AdditionalClientArguments}";
-				ClientArgs += $" {GameServerAddress}";
-				ClientArgs += $" -log={ClientType}-{{ClientId}}.log -windowed -SaveWinPos={{ClientId}}";
+				bool bValidAddress = ParseAddressAndPortRange(out GameServerBaseAddress, out StartPort, out EndPort, GameServerAddressWithPortRange);
+				if (!bValidAddress)
+				{
+					Logger.LogError("Invalid Game Server Address specified {GameServerAddress} should be in the format of IpAddress:Port or IpAddress:StartPort-EndPort", GameServerAddressWithPortRange);
+					return;
+				}
 
-				// Potentially launch multiple clients per server address
-				for (int InstanceNum = 0; InstanceNum < ClientsPerServer; ++InstanceNum)
+				for (int Port = StartPort; Port <= EndPort; ++Port)
 				{
-					int ClientId = (Idx + 1 + InstanceNum);
-					string InstanceArgs = ClientArgs.Replace("{ClientId}", ClientId.ToString(), StringComparison.OrdinalIgnoreCase);
+					string ClientArgs = $" -game {AdditionalClientArguments}";
+					ClientArgs += $" {GameServerBaseAddress}:{Port}";
+					ClientArgs += $" -log={ClientType}-{{ClientId}}.log -windowed -SaveWinPos={{ClientId}}";
 
-					string InstanceDetails = (ClientsPerServer > 1) ? $"(Srv {Idx+1})" : string.Empty;
-					string AppTitle = $"{ClientType} {ClientId} {InstanceDetails}";
-					ExecuteApp(ClientType, AppTitle, InstanceArgs);
+					// Potentially launch multiple clients per server address
+					for (int InstanceNum = 0; InstanceNum < ClientsPerServer; ++InstanceNum)
+					{
+						int ClientId = (Idx + 1 + InstanceNum);
+						string InstanceArgs = ClientArgs.Replace("{ClientId}", ClientId.ToString(), StringComparison.OrdinalIgnoreCase);
+
+						string InstanceDetails = (ClientsPerServer > 1) ? $"(Srv {Idx + 1})" : string.Empty;
+						string AppTitle = $"{ClientType} {ClientId} {InstanceDetails}";
+						ExecuteApp(ClientType, AppTitle, InstanceArgs);
+					}
 				}
 			}
 		}
@@ -517,7 +534,7 @@ Example running 4 clients locally connecting to single proxy server w/ a non-loc
 
 			if (!Matchmake && AllGameServerAddresses.Length > 0)
 			{
-				CommonProxyServerArgs += $" -ProxyGameServers={string.Join(',', AllGameServerAddresses)}";
+				CommonProxyServerArgs += $" -ProxyGameServers={string.Join(',', AllGameServerAddresses).TrimEnd(',')}";
 			}
 
 			if (ProxyClientPrimary > 0)
@@ -630,6 +647,54 @@ Example running 4 clients locally connecting to single proxy server w/ a non-loc
 			}
 		}
 
+		/// <summary>
+		/// Helper function allows us to break-out a potential input of "address", "address:port", or "address:startPort-endPort".
+		/// </summary>
+		/// <param name="IpAddress">The base ip address of the passed-in addressToParse.</param>
+		/// <param name="StartPort">The start port of the port range in addressToParse. If no port specified, zero.</param>
+		/// <param name="EndPort">The end port of the port range in addressToParse. If none specified, return startPort.</param>
+		/// <param name="AddressToParse">The input address to parse of the form "address", "address:port", or "address:startPort-endPort"</param>
+		/// <returns>True if a successful parse is made.  False if it is somehow malformed.</returns>
+		protected bool ParseAddressAndPortRange(out string IpAddress, out int StartPort, out int EndPort, string AddressToParse)
+		{
+			IpAddress = string.Empty;
+			StartPort = 0;
+			EndPort = 0;
+
+			if (string.IsNullOrEmpty(AddressToParse))
+				return false;
+
+			string[] AddressParts = AddressToParse.Split(':');
+			IpAddress = AddressParts[0];
+
+			// See if we've also specified a port (if not we're done, ports stay at zero)
+			if (AddressParts.Length < 2)
+				return true;
+
+			// Check for a port range (e.g. :9001-9005) or a single port (:9001)
+			string[] PortRange = AddressParts[1].Split('-');
+			if (!int.TryParse(PortRange[0], out StartPort))
+				return false;
+
+			// If we specified a port range (9001-9005) then parse the endPort
+			if (PortRange.Length > 1)
+			{
+				if (!int.TryParse(PortRange[1], out EndPort))
+				{
+					EndPort = StartPort;
+					return false;
+				}
+
+				return (EndPort > StartPort);
+			}
+			else
+			{
+				EndPort = StartPort;
+			}
+
+			return true;
+		}
+
 		/// <summary>
 		/// Are we requesting a debugger be attached to this particular process?
 		/// </summary>
