Correct issue with file locking in UnrealAndroidFileTool for writes

#android

diff --git a/Engine/Source/Programs/AutomationTool/Android/AndroidFileClient.cs b/Engine/Source/Programs/AutomationTool/Android/AndroidFileClient.cs
--- a/Engine/Source/Programs/AutomationTool/Android/AndroidFileClient.cs
+++ b/Engine/Source/Programs/AutomationTool/Android/AndroidFileClient.cs
@@ -2433,18 +2433,20 @@ namespace AutomationTool
 				{
 					// deal with LastBaseDir substitution
 					DestPath = OptimizePath(DestPath);
-					if (bLog > 0)
-					{
-						Logger.LogInformation("{bLog}> Writing {DestPath}", bLog, DestPath);
-					}
 
 					long fileSize = fileStream.Length;
 					if (fileSize <= FILE_BUFFERSIZE)
 					{
 						// just start over as full transfer if smaller file
+						fileStream.Close();
 						return FileWriteInternal_Normal(SourcePath, DestPath, bLog);
 					}
 
+					if (bLog > 0)
+					{
+						Logger.LogInformation("{bLog}> Writing {DestPath}", bLog, DestPath);
+					}
+
 					// get timestamp of file to use for validating resume
 					long Timestamp = File.GetLastWriteTimeUtc(SourcePath).ToBinary();
 
diff --git a/Engine/Source/Programs/UnrealAndroidFileTool/Program.cs b/Engine/Source/Programs/UnrealAndroidFileTool/Program.cs
--- a/Engine/Source/Programs/UnrealAndroidFileTool/Program.cs
+++ b/Engine/Source/Programs/UnrealAndroidFileTool/Program.cs
@@ -18,7 +18,7 @@ namespace UnrealAndroidFileTool
 {
 	class Program
 	{
-		private static string AppVersion = "1.0.8";
+		private static string AppVersion = "1.0.9";
 		private static ILogger Logger = Log.Logger;
 
 		static void ShowHelp(string command)
