UAT: BuildCMakeLib: Added missing backend classes to most platforms to support Ninja build system. [Part 1/2]

diff --git a/Engine/Platforms/VisionOS/Source/Programs/AutomationTool/Scripts/BuildCMakeLib_VisionOS.Automation.cs b/Engine/Platforms/VisionOS/Source/Programs/AutomationTool/Scripts/BuildCMakeLib_VisionOS.Automation.cs
--- a/Engine/Platforms/VisionOS/Source/Programs/AutomationTool/Scripts/BuildCMakeLib_VisionOS.Automation.cs
+++ b/Engine/Platforms/VisionOS/Source/Programs/AutomationTool/Scripts/BuildCMakeLib_VisionOS.Automation.cs
@@ -21,3 +21,16 @@ class MakefileTargetPlatform_VisionOS : MakefileTargetPlatform_IOS
 	}
 
 }
+
+class NinjaBuildTargetPlatform_VisionOS : NinjaBuildTargetPlatform_IOS
+{
+	public override string PlatformOrGroupName => nameof(UnrealTargetPlatform.VisionOS);
+	public override bool IsPlatformExtension => true;
+	protected override string CMakeSystemName => "VisionOS";
+	protected override string SdkName => Architecture.Contains("simulator") || Architecture == "x86_64" ? "xrsimulator" : "xros";
+
+	public NinjaBuildTargetPlatform_VisionOS(string Architecture) : base(Architecture)
+	{
+	}
+
+}
diff --git a/Engine/Source/Programs/AutomationTool/Scripts/BuildCMakeLib.Automation.cs b/Engine/Source/Programs/AutomationTool/Scripts/BuildCMakeLib.Automation.cs
--- a/Engine/Source/Programs/AutomationTool/Scripts/BuildCMakeLib.Automation.cs
+++ b/Engine/Source/Programs/AutomationTool/Scripts/BuildCMakeLib.Automation.cs
@@ -442,7 +442,10 @@ public sealed class BuildCMakeLib : BuildCommand
 
 		public override string CMakeGeneratorName => "Ninja";
 
-		public override string MakeCommand => FileReference.Combine(Unreal.RootDirectory, "Engine", "Extras", "ThirdPartyNotUE", "ninja-build", "ninja.exe").FullName;
+		public override string MakeCommand =>
+			BuildHostPlatform.Current.Platform.IsInGroup(UnrealPlatformGroup.Microsoft)
+			? FileReference.Combine(Unreal.RootDirectory, "Engine", "Extras", "ThirdPartyNotUE", "ninja-build", "ninja.exe").FullName
+			: "ninja";		
 	}
 
 	public abstract class VSTargetPlatform : TargetPlatform
@@ -831,6 +834,15 @@ class MakefileTargetPlatform_Win64 : BuildCMakeLib.MakefileTargetPlatform
 	public override bool IsPlatformExtension => false;
 }
 
+class NinjaBuildTargetPlatform_Win64 : BuildCMakeLib.NinjaBuildTargetPlatform
+{
+	public override string PlatformOrGroupName => nameof(UnrealTargetPlatform.Win64);
+	public override string DebugDatabaseExtension => "pdb";
+	public override string DynamicLibraryExtension => "dll";
+	public override string StaticLibraryExtension => "lib";
+	public override bool IsPlatformExtension => false;
+}
+
 class NMakeTargetPlatform_Unix : BuildCMakeLib.NMakeTargetPlatform
 {
 	public override string PlatformOrGroupName => nameof(UnrealPlatformGroup.Unix);
@@ -861,6 +873,21 @@ class MakefileTargetPlatform_Unix : BuildCMakeLib.MakefileTargetPlatform
 	}
 }
 
+class NinjaBuildTargetPlatform_Unix : BuildCMakeLib.NinjaBuildTargetPlatform
+{
+	public override string PlatformOrGroupName => nameof(UnrealPlatformGroup.Unix);
+	public override string StaticLibraryExtension => "a";
+	public override string VariantDirectory => Architecture;
+	public override bool IsPlatformExtension => false;
+
+	private readonly string Architecture;
+
+	public NinjaBuildTargetPlatform_Unix(string Architecture)
+	{
+		this.Architecture = Architecture;
+	}
+}
+
 class NMakeTargetPlatform_Linux : NMakeTargetPlatform_Unix
 {
 	public override string PlatformOrGroupName => nameof(UnrealTargetPlatform.Linux);
@@ -881,6 +908,16 @@ class MakefileTargetPlatform_Linux : MakefileTargetPlatform_Unix
 	}
 }
 
+class NinjaBuildTargetPlatform_Linux : NinjaBuildTargetPlatform_Unix
+{
+	public override string PlatformOrGroupName => nameof(UnrealTargetPlatform.Linux);
+
+	public NinjaBuildTargetPlatform_Linux(string Architecture)
+		: base(Architecture)
+	{
+	}
+}
+
 class MakefileTargetPlatform_LinuxArm64 : MakefileTargetPlatform_Unix
 {
 	public override string PlatformOrGroupName => nameof(UnrealTargetPlatform.LinuxArm64);
@@ -891,6 +928,16 @@ class MakefileTargetPlatform_LinuxArm64 : MakefileTargetPlatform_Unix
 	}
 }
 
+class NinjaBuildTargetPlatform_LinuxArm64 : NinjaBuildTargetPlatform_Unix
+{
+	public override string PlatformOrGroupName => nameof(UnrealTargetPlatform.LinuxArm64);
+
+	public NinjaBuildTargetPlatform_LinuxArm64(string Architecture)
+		: base(Architecture)
+	{
+	}
+}
+
 class XcodeTargetPlatform_Mac : BuildCMakeLib.XcodeTargetPlatform
 {
 	public override string PlatformOrGroupName => nameof(UnrealTargetPlatform.Mac);
@@ -932,6 +979,27 @@ class MakefileTargetPlatform_Mac : BuildCMakeLib.MakefileTargetPlatform
 	}
 }
 
+class NinjaBuildTargetPlatform_Mac : BuildCMakeLib.NinjaBuildTargetPlatform
+{
+	public override string PlatformOrGroupName => nameof(UnrealTargetPlatform.Mac);
+	public override string StaticLibraryExtension => "a";
+	public override string VariantDirectory => Architecture;
+	public override bool IsPlatformExtension => false;
+
+	private readonly string Architecture;
+
+	public NinjaBuildTargetPlatform_Mac(string Architecture)
+	{
+		this.Architecture = Architecture;
+	}
+
+	public override string GetCMakeSetupArguments(BuildCMakeLib.TargetLib TargetLib, string TargetConfiguration)
+	{
+		return base.GetCMakeSetupArguments(TargetLib, TargetConfiguration)
+			+ string.Format(" -DCMAKE_OSX_ARCHITECTURES={0}", Architecture);
+	}
+}
+
 public class XcodeTargetPlatform_IOS : BuildCMakeLib.XcodeTargetPlatform
 {
 	public override string PlatformOrGroupName => nameof(UnrealTargetPlatform.IOS);
@@ -984,6 +1052,43 @@ public class MakefileTargetPlatform_IOS : BuildCMakeLib.MakefileTargetPlatform
 
 }
 
+public class NinjaBuildTargetPlatform_IOS : BuildCMakeLib.NinjaBuildTargetPlatform
+{
+	public override string PlatformOrGroupName => nameof(UnrealTargetPlatform.IOS);
+	public override string StaticLibraryExtension => "a";
+	public override string VariantDirectory => Architecture;
+	public override bool IsPlatformExtension => false;
+	protected virtual string CMakeSystemName => "iOS";
+	protected virtual string SdkName => Architecture == "x86_64" || Architecture == "iossimulator" ? "iphonesimulator" : "iphoneos";
+	protected virtual string ClangArchitecture => Architecture.Contains("x86_64") ? "x86_64" : "arm64";
+
+	protected readonly string Architecture;
+
+	public NinjaBuildTargetPlatform_IOS(string Architecture)
+	{
+		this.Architecture = Architecture;
+	}
+
+	public override string GetCMakeSetupArguments(BuildCMakeLib.TargetLib TargetLib, string TargetConfiguration)
+	{
+		return base.GetCMakeSetupArguments(TargetLib, TargetConfiguration)
+			+ $" -DCMAKE_SYSTEM_NAME={CMakeSystemName}"
+			+ $" -DCMAKE_OSX_ARCHITECTURES={ClangArchitecture}"
+			+ $" -DCMAKE_OSX_SYSROOT={GetSysRoot()}";
+	}
+
+	private string GetSysRoot()
+	{
+		IProcessResult Result = Run("xcrun", string.Format("-sdk {0} --show-sdk-path", SdkName));
+		string SysRoot = (Result.Output ?? "").Trim();
+		if (Result.ExitCode != 0 || !DirectoryExists(SysRoot))
+		{
+			throw new AutomationException("Failed to locate {3} SDK \"{0}\":{1}{2}", SdkName, Environment.NewLine, SysRoot, CMakeSystemName);
+		}
+		return SysRoot;
+	}
+}
+
 class XcodeTargetPlatform_TVOS : XcodeTargetPlatform_IOS
 {
 	public override string PlatformOrGroupName => nameof(UnrealTargetPlatform.TVOS);
@@ -999,7 +1104,17 @@ class MakefileTargetPlatform_TVOS : MakefileTargetPlatform_IOS
 	public MakefileTargetPlatform_TVOS(string Architecture) : base(Architecture)
 	{
 	}
+}
 
+class NinjaBuildTargetPlatform_TVOS : NinjaBuildTargetPlatform_IOS
+{
+	public override string PlatformOrGroupName => nameof(UnrealTargetPlatform.TVOS);
+	protected override string CMakeSystemName => "tvOS";
+	protected override string SdkName => Architecture == "x86_64" || Architecture == "tvossimulator" || Architecture == "iossimulator" ? "appletvsimulator" : "appletvos";
+
+	public NinjaBuildTargetPlatform_TVOS(string Architecture) : base(Architecture)
+	{
+	}
 }
 
 class NMakeTargetPlatform_Android : BuildCMakeLib.NMakeTargetPlatform
@@ -1032,6 +1147,21 @@ class MakefileTargetPlatform_Android : BuildCMakeLib.MakefileTargetPlatform
 	}
 }
 
+class NinjaBuildTargetPlatform_Android : BuildCMakeLib.NinjaBuildTargetPlatform
+{
+	public override string PlatformOrGroupName => nameof(UnrealTargetPlatform.Android);
+	public override string StaticLibraryExtension => "a";
+	public override string VariantDirectory => Architecture;
+	public override bool IsPlatformExtension => false;
+
+	private readonly string Architecture;
+
+	public NinjaBuildTargetPlatform_Android(string Architecture)
+	{
+		this.Architecture = Architecture;
+	}
+}
+
 class VS2019TargetPlatform_Android : BuildCMakeLib.VS2019TargetPlatform
 {
 	public override string PlatformOrGroupName => nameof(UnrealTargetPlatform.Android);
