Improve logging for instances when Gauntlet says that app exited with "code 0" but then sets exit code to -1

#rnx

diff --git a/Engine/Source/Programs/AutomationTool/Gauntlet/Unreal/Base/Gauntlet.UnrealTestNode.cs b/Engine/Source/Programs/AutomationTool/Gauntlet/Unreal/Base/Gauntlet.UnrealTestNode.cs
--- a/Engine/Source/Programs/AutomationTool/Gauntlet/Unreal/Base/Gauntlet.UnrealTestNode.cs
+++ b/Engine/Source/Programs/AutomationTool/Gauntlet/Unreal/Base/Gauntlet.UnrealTestNode.cs
@@ -1846,7 +1846,7 @@ namespace Gauntlet
 			}*/
 
 			// The process is gone but we don't know why. This is likely bad and signifies an unhandled or undiagnosed error
-			ExitReason = "app exited with code 0";
+			ExitReason = "app exited for unhandled reason";
 			ExitCode = -1;
 			return UnrealProcessResult.Unknown;
 		}
