Allow ignoring ensures in report

#rnx

diff --git a/Engine/Source/Programs/AutomationTool/Gauntlet/Unreal/Base/Gauntlet.UnrealTestConfiguration.cs b/Engine/Source/Programs/AutomationTool/Gauntlet/Unreal/Base/Gauntlet.UnrealTestConfiguration.cs
--- a/Engine/Source/Programs/AutomationTool/Gauntlet/Unreal/Base/Gauntlet.UnrealTestConfiguration.cs
+++ b/Engine/Source/Programs/AutomationTool/Gauntlet/Unreal/Base/Gauntlet.UnrealTestConfiguration.cs
@@ -716,6 +716,12 @@ namespace Gauntlet
 		[AutoParam(false)]
 		public bool FailOnEnsures { get; set; }
 
+		/// <summary>
+		/// Whether ensures should be ignored from report
+		/// </summary>
+		[AutoParam(false)]
+		public bool IgnoreEnsures { get; set; }
+
 		/// <summary>
 		/// Whether warnings are shown in the summary
 		/// </summary>
diff --git a/Engine/Source/Programs/AutomationTool/Gauntlet/Unreal/Base/Gauntlet.UnrealTestNode.cs b/Engine/Source/Programs/AutomationTool/Gauntlet/Unreal/Base/Gauntlet.UnrealTestNode.cs
--- a/Engine/Source/Programs/AutomationTool/Gauntlet/Unreal/Base/Gauntlet.UnrealTestNode.cs
+++ b/Engine/Source/Programs/AutomationTool/Gauntlet/Unreal/Base/Gauntlet.UnrealTestNode.cs
@@ -1872,11 +1872,14 @@ namespace Gauntlet
 				EventList.Add(FatalEvent);
 			}
 
-			// Create events for any ensures
-			foreach (UnrealLog.CallstackMessage Ensure in InLog.Ensures)
+			if (!GetCachedConfiguration().IgnoreEnsures)
 			{
-				UnrealTestEvent EnsureEvent = new UnrealTestEvent(EventSeverity.Warning, Ensure.Message, Enumerable.Empty<string>(), Ensure);
-				EventList.Add(EnsureEvent);
+				// Create events for any ensures
+				foreach (UnrealLog.CallstackMessage Ensure in InLog.Ensures)
+				{
+					UnrealTestEvent EnsureEvent = new UnrealTestEvent(EventSeverity.Warning, Ensure.Message, Enumerable.Empty<string>(), Ensure);
+					EventList.Add(EnsureEvent);
+				}
 			}
 
 			HashSet<string> MonitoredCategorySet = new HashSet<string>();
