Fixed inconsistent exception logging in Gauntlet.TestExecutor.cs

diff --git a/Engine/Source/Programs/AutomationTool/Gauntlet/Framework/Gauntlet.TestExecutor.cs b/Engine/Source/Programs/AutomationTool/Gauntlet/Framework/Gauntlet.TestExecutor.cs
--- a/Engine/Source/Programs/AutomationTool/Gauntlet/Framework/Gauntlet.TestExecutor.cs
+++ b/Engine/Source/Programs/AutomationTool/Gauntlet/Framework/Gauntlet.TestExecutor.cs
@@ -259,7 +259,7 @@ namespace Gauntlet
 									&& (ex.InnerException as WebException).Response != null
 									&& ((ex.InnerException as WebException).Response as HttpWebResponse).StatusCode == HttpStatusCode.Conflict)
 								{
-									Log.Info(ex.Message);
+									Log.Info(ex.ToString());
 									// In case of conflict server response, let the test execution loop handles it (resolved through timeout option).
 								}
 								else
@@ -602,8 +602,8 @@ namespace Gauntlet
 			}
 			catch (Exception Ex)
 			{
-				Log.Error(KnownLogEvents.Gauntlet_TestEvent, "Test {Name} threw an exception during launch. Skipping test. Ex: {Exception}\n{StackTrace}", Name, Ex.Message, Ex.StackTrace);
-			}			
+				Log.Error(KnownLogEvents.Gauntlet_TestEvent, "Test {Name} threw an exception during launch. Skipping test. Ex: {Exception}", Name, Ex);
+			}
 
 			return false;			
 		}
@@ -760,7 +760,7 @@ namespace Gauntlet
 			}
 			catch (Exception Ex)
 			{
-				TestInfo.CancellationReason = string.Format("Test {0} threw an exception. Cancelling. Ex: {1}\n{2}", TestInfo.TestNode.Name, Ex.Message, Ex.StackTrace);
+				TestInfo.CancellationReason = string.Format("Test {0} threw an exception. Cancelling. Ex: {1}", TestInfo.TestNode.Name, Ex);
 				TestInfo.FinalResult = TestResult.Cancelled;
 			}
 		
@@ -867,7 +867,8 @@ namespace Gauntlet
 						}
 						catch (System.Exception ex)
 						{
-							Log.Error(KnownLogEvents.Gauntlet_TestEvent, "Test {Name} threw an exception while cleaning up. Ex: {Exception}", TestInfo.TestNode.Name, ex.Message);
+							Log.Error(KnownLogEvents.Gauntlet_TestEvent, "Test {Name} threw an exception while cleaning up. Ex: {Exception}",
+								TestInfo.TestNode.Name, ex);
 						}
 					}
 
@@ -878,13 +879,13 @@ namespace Gauntlet
 				{
 					if (TestIsRunning)
 					{
-						Log.Warning(KnownLogEvents.Gauntlet_TestEvent, "Cancelled Test {Name} threw an exception while stopping. Ex: {Exception}\n{StackTrace}", 
-							TestInfo.TestNode.Name, ex.Message, ex.StackTrace);
+						Log.Warning(KnownLogEvents.Gauntlet_TestEvent, "Cancelled Test {Name} threw an exception while stopping. Ex: {Exception}", 
+							TestInfo.TestNode.Name, ex);
 					}
 					else
 					{
-						Log.Error(KnownLogEvents.Gauntlet_TestEvent, "Test {Name} threw an exception while stopping. Ex: {Exception}\n{StackTrace}",
-							TestInfo.TestNode.Name, ex.Message, ex.StackTrace);
+						Log.Error(KnownLogEvents.Gauntlet_TestEvent, "Test {Name} threw an exception while stopping. Ex: {Exception}",
+							TestInfo.TestNode.Name, ex);
 					}
 
 					TestInfo.FinalResult = TestResult.Failed;
