Fix SafeCopyFile to handle links.

#rnx

diff --git a/Engine/Source/Programs/AutomationTool/AutomationUtils/Utils.cs b/Engine/Source/Programs/AutomationTool/AutomationUtils/Utils.cs
--- a/Engine/Source/Programs/AutomationTool/AutomationUtils/Utils.cs
+++ b/Engine/Source/Programs/AutomationTool/AutomationUtils/Utils.cs
@@ -504,6 +504,12 @@ namespace AutomationTool
 					{
 						FileInfo SourceInfo = new FileInfo(SourceName);
 						FileInfo TargetInfo = new FileInfo(TargetName);
+						
+						if (SourceInfo.Attributes.HasFlag(FileAttributes.ReparsePoint))
+						{
+							SourceInfo = (FileInfo)SourceInfo.ResolveLinkTarget(true);
+						}
+
 						if (!bSkipSizeCheck && SourceInfo.Length != TargetInfo.Length)
 						{
 							Logger.LogInformation("Size mismatch {SourceName} = {SourceLength} to {TargetName} = {TargetLength}", SourceName, SourceInfo.Length, TargetName, TargetInfo.Length);
