BuildGraph: Make the $(ProjectDir) variable available to BuildGraph scripts in the case that the project is specified by set:ProjectName= or set:ProjectFile= paramaters to AutomationTool.exe rather than the -project= parameter.
#rnx

diff --git a/Engine/Source/Programs/AutomationTool/AutomationUtils/BuildCommand.cs b/Engine/Source/Programs/AutomationTool/AutomationUtils/BuildCommand.cs
--- a/Engine/Source/Programs/AutomationTool/AutomationUtils/BuildCommand.cs
+++ b/Engine/Source/Programs/AutomationTool/AutomationUtils/BuildCommand.cs
@@ -73,6 +73,19 @@ namespace AutomationTool
 			return ParamValue ?? Default;
 		}
 
+		/// <summary>
+		/// Parses the command's Params list for a property set in the format set:<Param>=<Value> and reads its value. 
+		/// Ex. ParseSetPropertyValue("projectName", "") will return "Lyra" for commandline "BuildGraph -set:ProjectName=Lyra"
+		/// </summary>
+		/// <param name="propertyName">PropertyName to read its value.</param>
+		/// <param name="Default">Value to return if the PropertyName is not found.</param>
+		/// <returns>Returns the value or Default if the parameter was not found.</returns>
+		[return: NotNullIfNotNull("Default")]
+		public string? ParseSetPropertyValue(string propertyName, string? Default = null)
+		{
+			return ParseParamValue(Params, $"set:{propertyName}", Default);
+		}
+
 		/// <summary>
 		/// Parses an argument.
 		/// </summary>
@@ -318,7 +331,14 @@ namespace AutomationTool
 			else
 			{
 				string originalProjectName = ParseParamValue("project", "");
-
+				if (string.IsNullOrEmpty(originalProjectName))
+				{
+					originalProjectName = ParseSetPropertyValue("projectName", "");
+				}
+				if (string.IsNullOrEmpty(originalProjectName))
+				{
+					originalProjectName = ParseSetPropertyValue("projectFile", "");
+				}
 				if (string.IsNullOrEmpty(originalProjectName))
 				{
 					return null;
