Adding ability for DLC and non chunked builds to generate Optional and OnDemand Paks.
Adding two new ProjectParams to support this:
-GenerateOptionalPak allows for generating optional paks
-GenerateOnDemandPak allows for generating optional paks

Adding ability for DLC paks to be renamed via a replacement string.

[/Script/UnrealEd.ProjectPackagingSettings]
+DLCPakNameOverride=(PakName="TheNameToReplaceFromThePak", OverrideName="TheNameToUseInstead")

In the DefaultGame.ini supports DLCs paks to be renamed.

diff --git a/Engine/Source/Programs/AutomationTool/AutomationUtils/ProjectParams.cs b/Engine/Source/Programs/AutomationTool/AutomationUtils/ProjectParams.cs
--- a/Engine/Source/Programs/AutomationTool/AutomationUtils/ProjectParams.cs
+++ b/Engine/Source/Programs/AutomationTool/AutomationUtils/ProjectParams.cs
@@ -455,6 +455,8 @@ namespace AutomationTool
 			this.SessionLabel = InParams.SessionLabel;
 			this.ProjectDescriptor = InParams.ProjectDescriptor;
 			this.Upload = InParams.Upload;
+			this.GenerateOptionalPakForNonChunkedBuild = InParams.GenerateOptionalPakForNonChunkedBuild;
+			this.GenerateOnDemandPakForNonChunkedBuild = InParams.GenerateOnDemandPakForNonChunkedBuild;
 			this.XcodeBuildOptions = InParams.XcodeBuildOptions;
 			this.bMakeMacNative = InParams.bMakeMacNative;
 		}
@@ -637,6 +639,8 @@ namespace AutomationTool
 			ParamList<string> InMapsToRebuildHLOD = null,
 			ParamList<string> TitleID = null,
 			string Upload = null,
+			bool? GenerateOptionalPakForNonChunkedBuild = null,
+			bool? GenerateOnDemandPakForNonChunkedBuild = null,
 			string XcodeBuildOptions = null,
 			bool bMakeMacNative = false
 			)
@@ -1092,6 +1096,9 @@ namespace AutomationTool
 				this.Manifests = true;
 			}
 
+			this.GenerateOptionalPakForNonChunkedBuild = GetParamValueIfNotSpecified(Command, GenerateOptionalPakForNonChunkedBuild, false, "GenerateOptionalPakForNonChunkedBuild");
+			this.GenerateOnDemandPakForNonChunkedBuild = GetParamValueIfNotSpecified(Command, GenerateOnDemandPakForNonChunkedBuild, false, "GenerateOnDemandPakForNonChunkedBuild");
+
 			this.XcodeBuildOptions = ParseParamValueIfNotSpecified(Command, XcodeBuildOptions, "xcodebuildoptions", null);
 			this.bMakeMacNative = GetParamValueIfNotSpecified(Command, null, this.bMakeMacNative, "macnative");
 
@@ -2473,6 +2480,12 @@ namespace AutomationTool
 		[Help("applyiostoreondemand", "Forces IoStoreOnDemand to be enabled for the project even if it is not set up for it")]
 		public bool ApplyIoStoreOnDemand { get; }
 
+		[Help("GenerateOptionalPakForNonChunkedBuild", "When true non chunked builds will attempt to generate an optional pak for optional bulk data.")]
+		public bool GenerateOptionalPakForNonChunkedBuild { get; }
+
+		[Help("GenerateOnDemandPakForNonChunkedBuild", "When true non chunked builds will attempt to generate an OnDemand pak. Will also force ApplyIoStoreOnDemand to true")]
+		public bool GenerateOnDemandPakForNonChunkedBuild { get; }
+
 		[Help("XcodeBuildOptions", "Extra options to pass to xcodebuild")]
 		public string XcodeBuildOptions { get; set; }
 
@@ -3475,6 +3488,8 @@ namespace AutomationTool
 				Logger.LogDebug("Trace={Trace}", Trace);
 				Logger.LogDebug("TraceHost={TraceHost}", TraceHost);
 				Logger.LogDebug("TraceFile={TraceFile}", TraceFile);
+				Logger.LogDebug($"GenerateOptionalPakForNonChunkedBuild={GenerateOptionalPakForNonChunkedBuild}");
+				Logger.LogDebug($"GenerateOnDemandPakForNonChunkedBuild={GenerateOnDemandPakForNonChunkedBuild}");
 				Logger.LogDebug("Project Params **************");
 			}
 			bLogged = true;
diff --git a/Engine/Source/Programs/AutomationTool/Scripts/CopyBuildToStagingDirectory.Automation.cs b/Engine/Source/Programs/AutomationTool/Scripts/CopyBuildToStagingDirectory.Automation.cs
--- a/Engine/Source/Programs/AutomationTool/Scripts/CopyBuildToStagingDirectory.Automation.cs
+++ b/Engine/Source/Programs/AutomationTool/Scripts/CopyBuildToStagingDirectory.Automation.cs
@@ -3147,6 +3147,8 @@ namespace AutomationScripts
 			return null;
 		}
 
+		const string OptionalBulkDataFileExtension = ".uptnl";
+
 		/// <summary>
 		/// Creates a pak file using staging context (single manifest)
 		/// </summary>
@@ -3162,24 +3164,45 @@ namespace AutomationScripts
 
 			List<PakFileRules> PakRulesList = GetPakFileRules(Params, SC);
 
-			ConcurrentBag<string> FilesToRemove = new ConcurrentBag<string>();
+			ConcurrentDictionary<string, string> DefaultPak = new ConcurrentDictionary<string, string>(Environment.ProcessorCount, UnrealPakResponseFile.Count());
+			ConcurrentDictionary<string, string> OptionalPak = new ConcurrentDictionary<string, string>(Environment.ProcessorCount, UnrealPakResponseFile.Count());
+			ConcurrentDictionary<string, string> OnDemandPak = new ConcurrentDictionary<string, string>(Environment.ProcessorCount, UnrealPakResponseFile.Count());
+			ConcurrentDictionary<string, string> OnDemandOptionalPak = new ConcurrentDictionary<string, string>(Environment.ProcessorCount, UnrealPakResponseFile.Count());
 
 			// Apply the pak file rules, this can remove things but will not override the pak file name
 			Parallel.ForEach(UnrealPakResponseFile, StagingFile =>
 			{
 				bool bExcludeFromPaks = false;
-				ApplyPakFileRules(PakRulesList, StagingFile, null, null, out bExcludeFromPaks);
+				PakFileRules? MatchingRule = ApplyPakFileRules(PakRulesList, StagingFile, null, null, out bExcludeFromPaks);
 
 				if (bExcludeFromPaks)
 				{
-					FilesToRemove.Add(StagingFile.Key);
+					return;
 				}
-			});
 
-			foreach (var FileToRemove in FilesToRemove.ToArray())
-			{
-				UnrealPakResponseFile.Remove(FileToRemove);
-			}
+				string OriginalFilename = StagingFile.Key;
+				string OrigExt = Path.GetExtension(OriginalFilename);
+
+				if (Params.GenerateOnDemandPakForNonChunkedBuild && MatchingRule.HasValue && MatchingRule.Value.bOnDemand)
+				{
+					if (Params.GenerateOptionalPakForNonChunkedBuild && OrigExt.Equals(OptionalBulkDataFileExtension))
+					{
+						OnDemandOptionalPak.TryAdd(StagingFile.Key, StagingFile.Value);
+					}
+					else
+					{
+						OnDemandPak.TryAdd(StagingFile.Key, StagingFile.Value);
+					}
+				}
+				else if (Params.GenerateOptionalPakForNonChunkedBuild && OrigExt.Equals(OptionalBulkDataFileExtension))
+				{
+					OptionalPak.TryAdd(StagingFile.Key, StagingFile.Value);
+				}
+				else
+				{
+					DefaultPak.TryAdd(StagingFile.Key, StagingFile.Value);
+				}
+			});
 
 			EncryptionAndSigning.CryptoSettings PakCryptoSettings = Params.DLCOverrideCryptoSettings;
 			if (PakCryptoSettings == null)
@@ -3188,7 +3211,23 @@ namespace AutomationScripts
 			}
 
 			List<CreatePakParams> PakInputs = new List<CreatePakParams>();
-			PakInputs.Add(new CreatePakParams(SC.ShortProjectName, UnrealPakResponseFile, Params.Compressed, Params.RehydrateAssets, Params.DLCOverrideEncryptionKeyGuid));
+			PakInputs.Add(new CreatePakParams(SC.ShortProjectName, DefaultPak.ToDictionary(), Params.Compressed, Params.RehydrateAssets, Params.DLCOverrideEncryptionKeyGuid));
+			if (!OptionalPak.IsEmpty)
+			{
+				string OptionalPakName = SC.ShortProjectName + "optional";
+				PakInputs.Add(new CreatePakParams(OptionalPakName, OptionalPak.ToDictionary(), Params.Compressed, Params.RehydrateAssets, Params.DLCOverrideEncryptionKeyGuid));
+			}
+			if (!OnDemandPak.IsEmpty)
+			{
+				string OnDemandPakName = SC.ShortProjectName + "ondemand";
+				PakInputs.Add(new CreatePakParams(OnDemandPakName, OnDemandPak.ToDictionary(), Params.Compressed, Params.RehydrateAssets, Params.DLCOverrideEncryptionKeyGuid));
+			}
+			if (!OnDemandOptionalPak.IsEmpty)
+			{
+				string OnDemandOptionalPakName = SC.ShortProjectName + "ondemandoptional";
+				PakInputs.Add(new CreatePakParams(OnDemandOptionalPakName, OnDemandOptionalPak.ToDictionary(), Params.Compressed, Params.RehydrateAssets, Params.DLCOverrideEncryptionKeyGuid));
+			}
+
 			CreatePaks(Params, SC, PakInputs, PakCryptoSettings);
 		}
 
@@ -4029,6 +4068,30 @@ namespace AutomationScripts
 
 			CommonAdditionalArgs += " " + AdditionalCompressionOptionsOnCommandLine;
 
+			Dictionary<string, string> DLCToNameOverrides = new Dictionary<string, string>();
+			{
+				List<string> DLCToNameOverridesTmp;
+				if (PlatformGameConfig.GetArray("/Script/UnrealEd.ProjectPackagingSettings", "DLCPakNameOverride", out DLCToNameOverridesTmp))
+				{
+					// Each string has the format '+DLCPakNameOverride=(PakName="TheNameToReplaceFromThePak", OverrideName="TheNameToUseInstead")'
+					// for @ regex " is written as ""
+					Regex DLCPakNameOverrideRegex = new Regex(@"\(PakName\s*=\s*(?<PakName>""*\w*""*),\s*OverrideName\s*=\s*(?<OverrideName>""*\w*""*)\)");
+
+					foreach (var PathStr in DLCToNameOverridesTmp)
+					{
+						var DLCPakOverrideMatch = DLCPakNameOverrideRegex.Match(PathStr);
+						if (DLCPakOverrideMatch.Success)
+						{
+							string DLCNameToOverride = DLCPakOverrideMatch.Groups["PakName"].Value.StripQuoteArgument();
+							string OverrideName = DLCPakOverrideMatch.Groups["OverrideName"].Value.StripQuoteArgument();
+							DLCToNameOverrides.Add(DLCNameToOverride, OverrideName);
+
+							Logger.LogInformation($"Adding DLCPak Name Override: {DLCNameToOverride} -> {OverrideName}");
+						}
+					}
+				}
+			}
+
 			bool bCopiedExistingGlobalUtoc = false;
 			List<Tuple<FileReference, StagedFileReference, string>> Outputs = new List<Tuple<FileReference, StagedFileReference, string>>();
 			foreach (CreatePakParams PakParams in PakParamsList)
@@ -4044,9 +4107,15 @@ namespace AutomationScripts
 				StagedFileReference OutputRelativeLocation;
 				if (Params.HasDLCName)
 				{
+					string DLCPakName = Params.DLCFile.GetFileNameWithoutExtension() + OutputFilename;
+					foreach (var PakToReplacement in DLCToNameOverrides)
+					{
+						DLCPakName = DLCPakName.Replace(PakToReplacement.Key, PakToReplacement.Value);
+					}
+
 					if (Params.DLCOverrideStagedSubDir != null)
 					{
-						OutputRelativeLocation = StagedFileReference.Combine(SC.RelativeProjectRootForStage, Params.DLCOverrideStagedSubDir, "Content", "Paks", Params.DLCFile.GetFileNameWithoutExtension() + OutputFilename + ".pak");
+						OutputRelativeLocation = StagedFileReference.Combine(SC.RelativeProjectRootForStage, Params.DLCOverrideStagedSubDir, "Content", "Paks", DLCPakName + ".pak");
 					}
 					else
 					{
@@ -4060,7 +4129,7 @@ namespace AutomationScripts
 							// if the plugin is external, the MakeRelativeTo above will fail, so just use the plugin name as the directory
 							PluginSubdirectory = Params.DLCFile.GetFileNameWithoutAnyExtensions();
 						}
-						OutputRelativeLocation = StagedFileReference.Combine(SC.RelativeProjectRootForStage, PluginSubdirectory, "Content", "Paks", SC.FinalCookPlatform, Params.DLCFile.GetFileNameWithoutExtension() + OutputFilename + ".pak");
+						OutputRelativeLocation = StagedFileReference.Combine(SC.RelativeProjectRootForStage, PluginSubdirectory, "Content", "Paks", SC.FinalCookPlatform, DLCPakName + ".pak");
 					}
 				}
 				else
@@ -5123,7 +5192,6 @@ namespace AutomationScripts
 					ChunkDefinitions.Add(CD);
 				}
 
-				const string OptionalBulkDataFileExtension = ".uptnl";
 				ConcurrentDictionary<string, ChunkDefinition> AdditionalChunks = new ConcurrentDictionary<string, ChunkDefinition>(StringComparer.InvariantCultureIgnoreCase);
 				ChunkDefinition DefaultChunk = ChunkDefinitions[DefaultChunkIndex];
 
