BuildGraph: $(ProjectDir): Fix error "Could not find a project file <ProjectName>." Do not use -set:ProjectName to initialize the project, only use -set:ProjectFile. The "ProjectName" can be different than the .uproject's leaf filename.
#rnx

diff --git a/Engine/Source/Programs/AutomationTool/AutomationUtils/BuildCommand.cs b/Engine/Source/Programs/AutomationTool/AutomationUtils/BuildCommand.cs
--- a/Engine/Source/Programs/AutomationTool/AutomationUtils/BuildCommand.cs
+++ b/Engine/Source/Programs/AutomationTool/AutomationUtils/BuildCommand.cs
@@ -331,10 +331,12 @@ namespace AutomationTool
 			else
 			{
 				string originalProjectName = ParseParamValue("project", "");
-				if (string.IsNullOrEmpty(originalProjectName))
-				{
-					originalProjectName = ParseSetPropertyValue("projectName", "");
-				}
+				// -set:projectName values can be different for some projects than the shortname used for -project and
+				// -set:projectFile, so do not report -set:projectName as the project.
+				//if (string.IsNullOrEmpty(originalProjectName))
+				//{
+				//	originalProjectName = ParseSetPropertyValue("projectName", "");
+				//}
 				if (string.IsNullOrEmpty(originalProjectName))
 				{
 					originalProjectName = ParseSetPropertyValue("projectFile", "");
