Backout: BuildGraph: Make the $(ProjectDir) variable available to BuildGraph scripts in the case that the project is specified by set:ProjectName= or set:ProjectFile= paramaters to AutomationTool.exe rather than the -project= parameter.
Setting $(Project) to either -set:ProjectFile or -set:ProjectName when they differ fails, because it used as the value for both when set, and different tasks expect one or the other. We can set $(ProjectDir) when they differ, but not $(Project).
#rnx

diff --git a/Engine/Source/Programs/AutomationTool/AutomationUtils/BuildCommand.cs b/Engine/Source/Programs/AutomationTool/AutomationUtils/BuildCommand.cs
--- a/Engine/Source/Programs/AutomationTool/AutomationUtils/BuildCommand.cs
+++ b/Engine/Source/Programs/AutomationTool/AutomationUtils/BuildCommand.cs
@@ -331,16 +331,7 @@ namespace AutomationTool
 			else
 			{
 				string originalProjectName = ParseParamValue("project", "");
-				// -set:projectName values can be different for some projects than the shortname used for -project and
-				// -set:projectFile, so do not report -set:projectName as the project.
-				//if (string.IsNullOrEmpty(originalProjectName))
-				//{
-				//	originalProjectName = ParseSetPropertyValue("projectName", "");
-				//}
-				if (string.IsNullOrEmpty(originalProjectName))
-				{
-					originalProjectName = ParseSetPropertyValue("projectFile", "");
-				}
+
 				if (string.IsNullOrEmpty(originalProjectName))
 				{
 					return null;
