Upgraded "Magick.NET" to 14.7.0 to resolve known high severity vulnerability from nuget

diff --git a/Engine/Source/Programs/AutomationTool/AutomationTool.csproj b/Engine/Source/Programs/AutomationTool/AutomationTool.csproj
--- a/Engine/Source/Programs/AutomationTool/AutomationTool.csproj
+++ b/Engine/Source/Programs/AutomationTool/AutomationTool.csproj
@@ -89,7 +89,7 @@
 
   <ItemGroup>
     <PackageReference Include="AWSSDK.SecurityToken" Version="3.7.300.105" />
-    <PackageReference Include="Magick.NET-Q16-HDRI-AnyCPU" Version="14.0.0" />
+    <PackageReference Include="Magick.NET-Q16-HDRI-AnyCPU" Version="14.7.0" />
     <PackageReference Include="Microsoft.Build" Version="17.11.4" ExcludeAssets="runtime" PrivateAssets="all">
       <IncludeAssets>compile; build; native; contentfiles; analyzers; buildtransitive</IncludeAssets>
     </PackageReference>
diff --git a/Engine/Source/Programs/AutomationTool/AutomationUtils/AutomationUtils.Automation.csproj b/Engine/Source/Programs/AutomationTool/AutomationUtils/AutomationUtils.Automation.csproj
--- a/Engine/Source/Programs/AutomationTool/AutomationUtils/AutomationUtils.Automation.csproj
+++ b/Engine/Source/Programs/AutomationTool/AutomationUtils/AutomationUtils.Automation.csproj
@@ -64,7 +64,7 @@
     <ProjectReference Include="..\..\UnrealBuildTool\UnrealBuildTool.csproj" PrivateAssets="All"><Private>true</Private></ProjectReference>
   </ItemGroup>
   <ItemGroup>
-    <PackageReference Include="Magick.NET-Q16-HDRI-AnyCPU" Version="14.0.0" PrivateAssets="all" />
+    <PackageReference Include="Magick.NET-Q16-HDRI-AnyCPU" Version="14.7.0" PrivateAssets="all" />
     <PackageReference Include="System.Drawing.Common" Version="8.0.10" />
     <PackageReference Include="System.Net.Http" Version="4.3.4" />
 
diff --git a/Engine/Source/Programs/AutomationTool/Gauntlet/Gauntlet.Automation.csproj b/Engine/Source/Programs/AutomationTool/Gauntlet/Gauntlet.Automation.csproj
--- a/Engine/Source/Programs/AutomationTool/Gauntlet/Gauntlet.Automation.csproj
+++ b/Engine/Source/Programs/AutomationTool/Gauntlet/Gauntlet.Automation.csproj
@@ -75,7 +75,7 @@
     <ProjectReference Include="..\Android\Android.Automation.csproj" PrivateAssets="All"><Private>false</Private></ProjectReference>
   </ItemGroup>
   <ItemGroup>
-    <PackageReference Include="Magick.NET-Q16-HDRI-AnyCPU" Version="14.0.0" PrivateAssets="all" />
+    <PackageReference Include="Magick.NET-Q16-HDRI-AnyCPU" Version="14.7.0" PrivateAssets="all" />
     <PackageReference Include="Microsoft.CSharp" Version="4.7.0" />
     <PackageReference Include="MySql.Data" Version="6.10.9" PrivateAssets="all" />
   </ItemGroup>
