BuildGraph: Make the $(ProjectDir) variable available to BuildGraph scripts in the case that the project is specified by set:ProjectName= or set:ProjectFile= paramaters to AutomationTool.exe rather than the -project= parameter.
#rnx

diff --git a/Engine/Source/Programs/AutomationTool/AutomationUtils/BuildCommand.cs b/Engine/Source/Programs/AutomationTool/AutomationUtils/BuildCommand.cs
--- a/Engine/Source/Programs/AutomationTool/AutomationUtils/BuildCommand.cs
+++ b/Engine/Source/Programs/AutomationTool/AutomationUtils/BuildCommand.cs
@@ -348,6 +348,39 @@ namespace AutomationTool
 			return projectFullPath;
 		}
 
+		/// <summary>
+		/// Searches for project's .uproject filepath, user-friendly name, and directory of .uproject filepath, from
+		/// commandline arguments (including -project= and -set:Project???= parameters). Caller should handle
+		/// possiblity that some fields are valid and others are null.
+		/// </summary>
+		public void ParseProjectIdParams(out FileReference? projectFile, out string? projectName,
+			out DirectoryReference? projectDir)
+		{
+			FileReference? projectFullPath = ParseProjectParam();
+			if (projectFullPath != null)
+			{
+				projectFile = projectFullPath;
+				projectName = projectFullPath.GetFileNameWithoutExtension();
+				projectDir = projectFullPath.Directory;
+			}
+			else
+			{
+				projectFile = null;
+				projectName = ParseSetPropertyValue("projectName", null);
+				projectDir = null;
+
+				string? projectFileString = ParseSetPropertyValue("projectFile", null);
+				if (projectFileString != null)
+				{
+					projectFile = ParseProjectString(projectFileString);
+					if (projectFile != null)
+					{
+						projectDir = projectFile.Directory;
+					}
+				}
+			}
+		}
+
 		/// <summary>
 		/// Checks that all of the required params are present, throws an exception if not
 		/// </summary>
diff --git a/Engine/Source/Programs/AutomationTool/BuildGraph/BuildGraph.cs b/Engine/Source/Programs/AutomationTool/BuildGraph/BuildGraph.cs
--- a/Engine/Source/Programs/AutomationTool/BuildGraph/BuildGraph.cs
+++ b/Engine/Source/Programs/AutomationTool/BuildGraph/BuildGraph.cs
@@ -364,12 +364,19 @@ namespace AutomationTool
 			}
 
 			// If the -project flag is given, pass useful information into the graph
-			FileReference? projectFile = ParseProjectParam();
+			ParseProjectIdParams(out FileReference? projectFile, out string? projectName,
+				out DirectoryReference? projectDir);
 			if (projectFile != null)
 			{
-				defaultProperties["ProjectName"] = projectFile.GetFileNameWithoutExtension();
 				defaultProperties["ProjectFile"] = projectFile.FullName;
-				defaultProperties["ProjectDir"] = projectFile.Directory.FullName;
+			}
+			if (projectName != null)
+			{
+				defaultProperties["ProjectName"] = projectName;
+			}
+			if (projectDir != null)
+			{ 
+				defaultProperties["ProjectDir"] = projectDir.FullName;
 			}
 
 			// Add any additional custom arguments from the command line (of the form -Set:X=Y)
