Clean temp folders used by test automation with Engine Test to avoid copying duplicated artifacts with UE.Automation tests

#rnx

diff --git a/Engine/Source/Programs/AutomationTool/Gauntlet/Unreal/Automation/UE.Automation.cs b/Engine/Source/Programs/AutomationTool/Gauntlet/Unreal/Automation/UE.Automation.cs
--- a/Engine/Source/Programs/AutomationTool/Gauntlet/Unreal/Automation/UE.Automation.cs
+++ b/Engine/Source/Programs/AutomationTool/Gauntlet/Unreal/Automation/UE.Automation.cs
@@ -978,6 +978,28 @@ namespace UE
 			return UnrealResult;
 		}
 
+		public override IEnumerable<UnrealRoleArtifacts> SaveRoleArtifacts(string OutputPath)
+		{
+			// Clean Automation/Tmp and Screenshots folders before copy to avoid duplication of screenshot comparison files
+			foreach (UnrealSessionInstance.RoleInstance RoleInstance in TestInstance.RunningRoles)
+			{
+				DirectoryInfo SrcDirectory = new DirectoryInfo(RoleInstance.AppInstance.ArtifactPath);
+				foreach (DirectoryInfo Dir in SrcDirectory.GetDirectories())
+				{
+					if (Dir.Name.Equals("Screenshots", StringComparison.OrdinalIgnoreCase))
+					{
+						try { Dir.Delete(true); } catch { Log.Verbose("Could not delete directory: {Dir}", Dir.FullName); }
+					}
+					else if(Dir.Name.Equals("Automation", StringComparison.OrdinalIgnoreCase))
+					{
+						Dir.GetDirectories("Tmp", SearchOption.TopDirectoryOnly).ToList().ForEach(T => { try { T.Delete(true); } catch { Log.Verbose("Could not delete directory: {Dir}", T.FullName); } });
+					}
+				}
+			}
+
+			return base.SaveRoleArtifacts(OutputPath);
+		}
+
 		/// <summary>
 		/// Optional function that is called on test completion and gives an opportunity to create a report
 		/// </summary>
