Fix deployment to Steam Deck that may be failing due to how arguments are passed to rsync

#13597

diff --git a/Engine/Source/Programs/AutomationTool/SteamDeck/SteamDeckSupport.cs b/Engine/Source/Programs/AutomationTool/SteamDeck/SteamDeckSupport.cs
--- a/Engine/Source/Programs/AutomationTool/SteamDeck/SteamDeckSupport.cs
+++ b/Engine/Source/Programs/AutomationTool/SteamDeck/SteamDeckSupport.cs
@@ -202,15 +202,12 @@ public static class SteamDeckSupport
 		// make a standard set of options to pass to the rsync for auth's -e option
 		string[] AuthOpts =
 		{
-			SSHPath,
-			$"-o UserKnownHostsFile='{KnownHostsPath}'",
+			$"\"\"{SSHPath}\"\"",
+			$"-o UserKnownHostsFile=\"\"{KnownHostsPath}\"\"",
 			$"-o StrictHostKeyChecking=no",
-			$"-i '{DevKitRSAPath}'",
+			$"-i \"\"{DevKitRSAPath}\"\"",
 		};
-
-		if (SSHPath.Any(Char.IsWhiteSpace))
-			AuthOpts[0] = $"'{SSHPath}'";
-
+		
 		string AuthOptions = $"-e \"{string.Join(" ", AuthOpts)}\"";
 
 		// make a set of --exclude options for anything we want to exclude
