Add an easier way for staging scripts to request a default pakfile rule for OnDemand data.

#rnx

- Move the functionality for creating a default OnDemand pakfile rule from the 'ApplyIoStoreOnDemand' option which is commonly used for debugging/testing to it's own option 'CreateDefaultOnDemandPakRule' so it can be invoked by other systems without enabling the additional baggage that 'ApplyIoStoreOnDemand' brings.
-- Enabling this option will create a pakfile rule that will produce OnDemand data and include all bulkdata files (.ubulk, .uptnl) BUT only if there are no pre-existing pakfile rules for producing OnDemand data.
- Fix ::CreatePakUsingStagingManifest to correctly identify OnDemand containers by setting 'bOnDemand' to true when calling ::CreatePakParams

diff --git a/Engine/Source/Programs/AutomationTool/AutomationUtils/ProjectParams.cs b/Engine/Source/Programs/AutomationTool/AutomationUtils/ProjectParams.cs
--- a/Engine/Source/Programs/AutomationTool/AutomationUtils/ProjectParams.cs
+++ b/Engine/Source/Programs/AutomationTool/AutomationUtils/ProjectParams.cs
@@ -455,6 +455,8 @@ namespace AutomationTool
 			this.SessionLabel = InParams.SessionLabel;
 			this.ProjectDescriptor = InParams.ProjectDescriptor;
 			this.Upload = InParams.Upload;
+			this.ApplyIoStoreOnDemand = InParams.ApplyIoStoreOnDemand;
+			this.CreateDefaultOnDemandPakRule = InParams.CreateDefaultOnDemandPakRule;
 			this.GenerateOptionalPakForNonChunkedBuild = InParams.GenerateOptionalPakForNonChunkedBuild;
 			this.GenerateOnDemandPakForNonChunkedBuild = InParams.GenerateOnDemandPakForNonChunkedBuild;
 			this.XcodeBuildOptions = InParams.XcodeBuildOptions;
@@ -639,6 +641,7 @@ namespace AutomationTool
 			ParamList<string> InMapsToRebuildHLOD = null,
 			ParamList<string> TitleID = null,
 			string Upload = null,
+			bool? CreateDefaultOnDemandPakRule = null,
 			bool? GenerateOptionalPakForNonChunkedBuild = null,
 			bool? GenerateOnDemandPakForNonChunkedBuild = null,
 			string XcodeBuildOptions = null,
@@ -1096,6 +1099,8 @@ namespace AutomationTool
 				this.Manifests = true;
 			}
 
+			this.CreateDefaultOnDemandPakRule = GetParamValueIfNotSpecified(Command, CreateDefaultOnDemandPakRule, false, "CreateDefaultOnDemandPakRule") || this.ApplyIoStoreOnDemand;
+
 			this.GenerateOptionalPakForNonChunkedBuild = GetParamValueIfNotSpecified(Command, GenerateOptionalPakForNonChunkedBuild, false, "GenerateOptionalPakForNonChunkedBuild");
 			this.GenerateOnDemandPakForNonChunkedBuild = GetParamValueIfNotSpecified(Command, GenerateOnDemandPakForNonChunkedBuild, false, "GenerateOnDemandPakForNonChunkedBuild");
 
@@ -2480,6 +2485,9 @@ namespace AutomationTool
 		[Help("applyiostoreondemand", "Forces IoStoreOnDemand to be enabled for the project even if it is not set up for it")]
 		public bool ApplyIoStoreOnDemand { get; }
 
+		[Help("createdefaultondemandpakrule", "Creates a default pakfile rule for OnDemand data if no other such rule exists. The default rule will include all bulkdata types.")]
+		public bool CreateDefaultOnDemandPakRule { get; }
+
 		[Help("GenerateOptionalPakForNonChunkedBuild", "When true non chunked builds will attempt to generate an optional pak for optional bulk data.")]
 		public bool GenerateOptionalPakForNonChunkedBuild { get; }
 
@@ -3488,8 +3496,6 @@ namespace AutomationTool
 				Logger.LogDebug("Trace={Trace}", Trace);
 				Logger.LogDebug("TraceHost={TraceHost}", TraceHost);
 				Logger.LogDebug("TraceFile={TraceFile}", TraceFile);
-				Logger.LogDebug($"GenerateOptionalPakForNonChunkedBuild={GenerateOptionalPakForNonChunkedBuild}");
-				Logger.LogDebug($"GenerateOnDemandPakForNonChunkedBuild={GenerateOnDemandPakForNonChunkedBuild}");
 				Logger.LogDebug("Project Params **************");
 			}
 			bLogged = true;
diff --git a/Engine/Source/Programs/AutomationTool/Scripts/CopyBuildToStagingDirectory.Automation.cs b/Engine/Source/Programs/AutomationTool/Scripts/CopyBuildToStagingDirectory.Automation.cs
--- a/Engine/Source/Programs/AutomationTool/Scripts/CopyBuildToStagingDirectory.Automation.cs
+++ b/Engine/Source/Programs/AutomationTool/Scripts/CopyBuildToStagingDirectory.Automation.cs
@@ -3020,7 +3020,7 @@ namespace AutomationScripts
 				}
 			}
 
-			if (Params.ApplyIoStoreOnDemand)
+			if (Params.CreateDefaultOnDemandPakRule)
 			{
 				// If there are no pak rules for an 'OnDemand' container we will need to create one via CreateIoStoreOnDemandPakRules
 				if (RulesList.FindIndex(x => x.bOnDemand == true) == -1)
@@ -3220,12 +3220,12 @@ namespace AutomationScripts
 			if (!OnDemandPak.IsEmpty)
 			{
 				string OnDemandPakName = SC.ShortProjectName + "ondemand";
-				PakInputs.Add(new CreatePakParams(OnDemandPakName, OnDemandPak.ToDictionary(), Params.Compressed, Params.RehydrateAssets, Params.DLCOverrideEncryptionKeyGuid));
+				PakInputs.Add(new CreatePakParams(OnDemandPakName, OnDemandPak.ToDictionary(), Params.Compressed, Params.RehydrateAssets, Params.DLCOverrideEncryptionKeyGuid, bOnDemand : true));
 			}
 			if (!OnDemandOptionalPak.IsEmpty)
 			{
 				string OnDemandOptionalPakName = SC.ShortProjectName + "ondemandoptional";
-				PakInputs.Add(new CreatePakParams(OnDemandOptionalPakName, OnDemandOptionalPak.ToDictionary(), Params.Compressed, Params.RehydrateAssets, Params.DLCOverrideEncryptionKeyGuid));
+				PakInputs.Add(new CreatePakParams(OnDemandOptionalPakName, OnDemandOptionalPak.ToDictionary(), Params.Compressed, Params.RehydrateAssets, Params.DLCOverrideEncryptionKeyGuid, bOnDemand : true));
 			}
 
 			CreatePaks(Params, SC, PakInputs, PakCryptoSettings);
