Fix Gauntlet LogFile path on desktop when UserDir is specified and IsBuildMachine is false

Gauntlet tests couldn't read log lines from LogBufferReader before this fix.

diff --git a/Engine/Source/Programs/AutomationTool/Gauntlet/Framework/Base/Gauntlet.TargetDeviceDesktopCommon.cs b/Engine/Source/Programs/AutomationTool/Gauntlet/Framework/Base/Gauntlet.TargetDeviceDesktopCommon.cs
--- a/Engine/Source/Programs/AutomationTool/Gauntlet/Framework/Base/Gauntlet.TargetDeviceDesktopCommon.cs
+++ b/Engine/Source/Programs/AutomationTool/Gauntlet/Framework/Base/Gauntlet.TargetDeviceDesktopCommon.cs
@@ -438,7 +438,9 @@ namespace Gauntlet
 				// -abslog makes sure Unreal dynamically update the log window when using -log
 				if (!CommandUtils.IsBuildMachine && (!string.IsNullOrEmpty(LogFile) || AppConfig.CommandLine.Contains("-log")))
 				{
-					string LogFolder = string.IsNullOrEmpty(LogFile) ? Path.Combine(Device.LocalCachePath, "Logs") : Path.GetDirectoryName(LogFile);
+					string LogFolder = string.IsNullOrEmpty(LogFile) ?
+						(string.IsNullOrEmpty(UserDir) ? Path.Combine(Device.LocalCachePath, "Logs") : Path.Combine(UserDir, "Saved", "Logs")) :
+						Path.GetDirectoryName(LogFile);
 
 					if (!Directory.Exists(LogFolder))
 					{
