Support Artifacts link, Embedded HTML and URL link in test events
+ Refactor Horde job step summary report

#rnx

diff --git a/Engine/Plugins/Performance/AutomatedPerfTesting/Build/Scripts/AutomatedPerfTestNode.cs b/Engine/Plugins/Performance/AutomatedPerfTesting/Build/Scripts/AutomatedPerfTestNode.cs
--- a/Engine/Plugins/Performance/AutomatedPerfTesting/Build/Scripts/AutomatedPerfTestNode.cs
+++ b/Engine/Plugins/Performance/AutomatedPerfTesting/Build/Scripts/AutomatedPerfTestNode.cs
@@ -12,11 +12,11 @@ using UnrealBuildBase;	// for Unreal.RootDirectory
 using UnrealBuildTool;	// for UnrealTargetPlatform
 
 using static AutomationTool.CommandUtils;
-using UE;
-using System.Configuration;
 using System.Diagnostics;
 using System.Reflection;
 using System.Text;
+using static Gauntlet.HordeReport.AutomatedTestSessionData;
+using AutomatedTestSessionData = Gauntlet.HordeReport.AutomatedTestSessionData;
 
 namespace AutomatedPerfTest
 {
@@ -488,7 +488,7 @@ namespace AutomatedPerfTest
 				string PRSReportName = "PRS Report Upload";
 
 				AddTestReportArtifacts(Result, Report, ArtifactReportName, UtraceLocation, FileFilter);
-				AddLinkToReport(Result, Report, PRSReportName, PRSReportLink);
+				AddLinkToReport(Report, PRSReportName, PRSReportLink);
 			}
 
 			return Report;
@@ -931,26 +931,45 @@ namespace AutomatedPerfTest
 			}
 		}
 
-		private void AddTestReportArtifacts(TestResult Result, ITestReport Report, string TestName, string OutputPath, string FileFilter)
+		private void AddTestReportArtifacts(TestResult Result, ITestReport Report, string ArtifactName, string OutputPath, string FileFilter)
 		{
-			foreach (string AbsoluteFilePath in FindFiles(FileFilter, true, OutputPath))
+			// Attach Report Info
+			if (Report is AutomatedTestSessionData SessionReport)
+			{
+				TestEventStream Stream = SessionReport.GetEventStream();
+				string[] FoundFiles = FindFiles(FileFilter, true, OutputPath);
+				List<string> HTMLFiles = FoundFiles.Where(F => F.EndsWith(".html", StringComparison.OrdinalIgnoreCase)).ToList();
+				List<string> FilteredFiles = FoundFiles.Except(HTMLFiles).ToList();
+				if (HTMLFiles.Any())
+				{
+					foreach (string AbsoluteFilePath in HTMLFiles)
+					{
+						string EntryName = Path.GetFileName(AbsoluteFilePath);
+						string RelativePath = Path.GetRelativePath(OutputPath, AbsoluteFilePath);
+						EmbeddedHTMLFile EmbeddedFile = new EmbeddedHTMLFile(EntryName, RelativePath, OutputPath);
+						Stream.AddArtifact(DateTime.UtcNow, Microsoft.Extensions.Logging.LogLevel.Information, $"{{{EmbeddedFile.Token}}}", EmbeddedFile);
+					}
+				}
+				if (FilteredFiles.Any())
+				{
+					Dictionary<string, string> Files = new Dictionary<string, string>();
+					foreach (string AbsoluteFilePath in FilteredFiles)
+					{
+						string EntryName = Path.GetFileName(AbsoluteFilePath);
+						string RelativePath = Path.GetRelativePath(OutputPath, AbsoluteFilePath);
+						Files.Add(EntryName, RelativePath);
+					}
+					ArtifactFiles ArtifactFiles = new ArtifactFiles(ArtifactName, Files, OutputPath);
+					Stream.AddArtifact(DateTime.UtcNow, Microsoft.Extensions.Logging.LogLevel.Information, $"{{{ArtifactFiles.Token}}}", ArtifactFiles);
+				}
+			}
+			else
 			{
-				// Create a report per file found that successfully creates a Report of Type SimpleTestReport 
-				if (CreateSimpleReportForHorde(Result) is not HordeReport.SimpleTestReport ChildReport)
+				foreach (string AbsoluteFilePath in FindFiles(FileFilter, true, OutputPath))
 				{
-					continue;
+					string BaseFileName = Path.GetFileName(AbsoluteFilePath);
+					Report.AttachArtifact(AbsoluteFilePath, BaseFileName);
 				}
-
-				string BaseFileName = Path.GetFileName(AbsoluteFilePath);
-
-				// Setup Report Info
-				ChildReport.AttachArtifact(AbsoluteFilePath, BaseFileName);
-				ChildReport.TestName = TestName;
-
-				// TODO: UE-305553
-				ChildReport.URLLink = $"/api/v2/artifacts/{Context.WorkerJobID}/file?path={BaseFileName}&inline=true";				
-
-				AttachHordeReport(Report, ChildReport);
 			}
 		}
 
diff --git a/Engine/Source/Programs/AutomationTool/Gauntlet/Framework/Horde/Gauntlet.Horde.cs b/Engine/Source/Programs/AutomationTool/Gauntlet/Framework/Horde/Gauntlet.Horde.cs
--- a/Engine/Source/Programs/AutomationTool/Gauntlet/Framework/Horde/Gauntlet.Horde.cs
+++ b/Engine/Source/Programs/AutomationTool/Gauntlet/Framework/Horde/Gauntlet.Horde.cs
@@ -44,6 +44,15 @@ namespace Gauntlet
 			public string fileName { get; set; }
 		}
 
+
+		public delegate string SummaryLineDelegate();
+		/// <summary>
+		/// Delegate called when the Horde step summary is generated.
+		/// Returned string is appended as a line to the summary report.
+		/// </summary>
+		/// <returns></returns>
+		public static SummaryLineDelegate OnStepSummary;
+
 		public static void AddJobBuildMetaData(IBuildSource buildSource)
 		{
 			try
@@ -65,7 +74,14 @@ namespace Gauntlet
 			}
 		}
 
-		public static void GenerateSummary(string ArtifactPath)
+		public static string GenerateArtifactLink(string ArtifactPath)
+		{
+			// file:// is attractive here, and file://\\computer\share\folder\file.txt does work ok in Chrome, but it does not Firefox
+			// so, lacking a uniform way to link to a network filesystem, just put the path here as static text and escaped for markdown
+			return $"* **Netshare Artifacts**: {System.Text.RegularExpressions.Regex.Replace(ArtifactPath, @"([|\\*\+])", @"\$1")}";
+		}
+
+		public static void GenerateSummary()
 		{
 			try
 			{
@@ -74,25 +90,29 @@ namespace Gauntlet
 					return;
 				}
 
-				string LogFolder = CommandUtils.CmdEnv.LogFolder;
-
-				string MarkdownFilename = "GauntletStepDetails.md";
-
 				StringBuilder Markdown = new StringBuilder();
 
-				// file:// is attractive here, and file://\\computer\share\folder\file.txt does work ok in Chrome, but it does not Firefox
-				// so, lacking a uniform way to link to a network filesystem, just put the path here as static text and escaped for markdown
-				Markdown.Append($"* **Netshare Artifacts**: {System.Text.RegularExpressions.Regex.Replace(ArtifactPath, @"([|\\*\+])", @"\$1")}");
+				foreach (SummaryLineDelegate callback in OnStepSummary.GetInvocationList())
+				{
+					Markdown.AppendLine(callback());
+				}
 				// detect if any testdata was generated
 				if (Globals.Params.ParseParam("UseTestDataV2")
 					&& Directory.GetFiles(Path.Combine(CommandUtils.CmdEnv.EngineSavedFolder, "TestData"), "*.json", SearchOption.TopDirectoryOnly).Any())
 				{
-					Markdown.Append($"\n* **Test Results**: [Open report](/test-automation?job={JobId}&step={StepId})");
+					Markdown.AppendLine($"* **Test Results**: [Open report](/test-automation?job={JobId}&step={StepId})");
 				}
-				File.WriteAllText(Path.Combine(LogFolder, MarkdownFilename), Markdown.ToString());
 
-				JobStepReport report = new JobStepReport() { name="Gauntlet Step Details", fileName=MarkdownFilename };
-				File.WriteAllText(Path.Combine(LogFolder, "GauntletStepDetails.report.json"), JsonSerializer.Serialize(report));
+				if (Markdown.Length > 0)
+				{
+					string LogFolder = CommandUtils.CmdEnv.LogFolder;
+					string MarkdownFilename = "GauntletStepDetails.md";
+
+					File.WriteAllText(Path.Combine(LogFolder, MarkdownFilename), Markdown.ToString());
+
+					JobStepReport report = new JobStepReport() { name = "Gauntlet Step Details", fileName = MarkdownFilename };
+					File.WriteAllText(Path.Combine(LogFolder, "GauntletStepDetails.report.json"), JsonSerializer.Serialize(report));
+				}
 			}
 			catch (Exception Ex)
 			{
diff --git a/Engine/Source/Programs/AutomationTool/Gauntlet/Framework/Utils/Gauntlet.HordeReport.cs b/Engine/Source/Programs/AutomationTool/Gauntlet/Framework/Utils/Gauntlet.HordeReport.cs
--- a/Engine/Source/Programs/AutomationTool/Gauntlet/Framework/Utils/Gauntlet.HordeReport.cs
+++ b/Engine/Source/Programs/AutomationTool/Gauntlet/Framework/Utils/Gauntlet.HordeReport.cs
@@ -83,7 +83,7 @@ namespace Gauntlet
 				}
 				else
 				{
-					if(ArtifactProcessedHashes.Contains(ArtifactHash))
+					if (ArtifactProcessedHashes.Contains(ArtifactHash))
 					{
 						// already processed
 						return true;
@@ -141,11 +141,11 @@ namespace Gauntlet
 
 			public void AttachDependencyReport(object InReport, string Key = null)
 			{
-				if(ExtraReports == null)
+				if (ExtraReports == null)
 				{
 					ExtraReports = new Dictionary<string, object>();
 				}
-				if(InReport is BaseHordeReport InHordeReport)
+				if (InReport is BaseHordeReport InHordeReport)
 				{
 					Key = InHordeReport.GetTestDataKey(Key);
 				}
@@ -169,7 +169,7 @@ namespace Gauntlet
 			/// <returns></returns>
 			public virtual string GetTestDataKey(string BaseKey = null)
 			{
-				if(BaseKey == null)
+				if (BaseKey == null)
 				{
 					return Type;
 				}
@@ -518,8 +518,8 @@ namespace Gauntlet
 					// copy artifacts
 					foreach (Artifact TestArtifact in OutputTestResultDetailed.Artifacts)
 					{
-						
-						string[] ArtifactPaths= { TestArtifact.Files.Difference, TestArtifact.Files.Approved, TestArtifact.Files.Unapproved };
+
+						string[] ArtifactPaths = { TestArtifact.Files.Difference, TestArtifact.Files.Approved, TestArtifact.Files.Unapproved };
 						foreach (string ArtifactPath in ArtifactPaths)
 						{
 							if (!string.IsNullOrEmpty(ArtifactPath) && AttachArtifact(Path.Combine(ReportPath, ArtifactPath), ArtifactPath)) { ArtifactsCount++; }
@@ -725,7 +725,7 @@ namespace Gauntlet
 			}
 			public class TestResultData
 			{
- 				public List<TestEvent> Events
+				public List<TestEvent> Events
 				{
 					get { return Result.Events; }
 					set { Result.Events = value; }
@@ -863,7 +863,7 @@ namespace Gauntlet
 			{
 				TestResult CurrentTest = GetCurrentTest();
 				TestEvent Event = CurrentTest.NewEvent(InDateTime, InType, InMessage, InTag, InContext);
-				if(InType == EventType.Error)
+				if (InType == EventType.Error)
 				{
 					IndexTestError(Event.Hash, CurrentTest.TestUID, InMessage, InTag);
 				}
@@ -998,7 +998,7 @@ namespace Gauntlet
 								string ArtifactFilePath = IntArtifactEntry.Files.GetValueOrDefault("difference");
 								if (!IsInfo && !string.IsNullOrEmpty(ArtifactFilePath))
 								{
-									if(!OutTestPassResults.AddArtifactToLastEvent(
+									if (!OutTestPassResults.AddArtifactToLastEvent(
 										"difference",
 										Path.Combine(InReportPath, ArtifactFilePath),
 										ArtifactFilePath
@@ -1010,7 +1010,7 @@ namespace Gauntlet
 								ArtifactFilePath = IntArtifactEntry.Files.GetValueOrDefault("approved");
 								if (!IsInfo && !string.IsNullOrEmpty(ArtifactFilePath))
 								{
-									if(!OutTestPassResults.AddArtifactToLastEvent(
+									if (!OutTestPassResults.AddArtifactToLastEvent(
 										"approved",
 										Path.Combine(InReportPath, ArtifactFilePath),
 										ArtifactFilePath
@@ -1023,7 +1023,7 @@ namespace Gauntlet
 								if (!string.IsNullOrEmpty(ArtifactFilePath))
 								{
 									string AbsoluteLocation = Path.Combine(InReportPath, ArtifactFilePath);
-									if(!OutTestPassResults.AddArtifactToLastEvent(
+									if (!OutTestPassResults.AddArtifactToLastEvent(
 										"unapproved",
 										AbsoluteLocation,
 										ArtifactFilePath
@@ -1049,7 +1049,7 @@ namespace Gauntlet
 											foreach (string JsonFile in JsonMetadataFiles)
 											{
 												string JsonArtifactName = RelativeLocation + "/" + Path.GetFileName(JsonFile);
-												if(!OutTestPassResults.AddArtifactToLastEvent("json metadata", JsonFile, JsonArtifactName))
+												if (!OutTestPassResults.AddArtifactToLastEvent("json metadata", JsonFile, JsonArtifactName))
 												{
 													FailedToAttached.Add(JsonArtifactName);
 												}
@@ -1057,9 +1057,9 @@ namespace Gauntlet
 										}
 									}
 								}
-								if(FailedToAttached.Count() > 0)
+								if (FailedToAttached.Count() > 0)
 								{
-									foreach(var Item in FailedToAttached)
+									foreach (var Item in FailedToAttached)
 									{
 										OutTestPassResults.AddEvent(EventType.Warning, string.Format("Failed to attached Artifact {0}.", Item));
 									}
@@ -1230,7 +1230,7 @@ namespace Gauntlet
 			public AutomatedTestSessionData(string InKey, string InName = null) : base()
 			{
 				Key = InKey;
-				Data = new TestSession(string.IsNullOrEmpty(InName)? InKey : InName, this);
+				Data = new TestSession(string.IsNullOrEmpty(InName) ? InKey : InName, this);
 				CurrentPhase = null;
 			}
 
@@ -1549,7 +1549,7 @@ namespace Gauntlet
 				TestPhase Phase = GetPhase(Key);
 				if (Phase == null)
 				{
-					throw new AutomationException("Could not find Phase '{0}'", Key); 
+					throw new AutomationException("Could not find Phase '{0}'", Key);
 				}
 				CurrentPhase = Phase;
 			}
@@ -1696,7 +1696,7 @@ namespace Gauntlet
 						MessageTemplate.ParsePropertyValues(Format, Args, Properties);
 						Message = MessageTemplate.Render(Format, Properties);
 					}
-					Logging.EventId EventId = Level == Logging.LogLevel.Critical? KnownLogEvents.Gauntlet_FatalEvent : KnownLogEvents.Gauntlet_TestEvent;
+					Logging.EventId EventId = Level == Logging.LogLevel.Critical ? KnownLogEvents.Gauntlet_FatalEvent : KnownLogEvents.Gauntlet_TestEvent;
 					LogEvent Event = new LogEvent(Time, Level, EventId, Message, Format, Properties, null);
 					AddEvent(Event);
 				}
@@ -1764,6 +1764,24 @@ namespace Gauntlet
 					AddEvent(Event.Time, Level, Message, Args: Params.ToArray());
 				}
 
+				/// <summary>
+				/// Add URL link
+				/// </summary>
+				/// <param name="Time"></param>
+				/// <param name="Level"></param>
+				/// <param name="Name"></param>
+				/// <param name="URL"></param>
+				public virtual void AddURLLink(DateTime Time, Logging.LogLevel Level, string Name, string URL)
+				{
+					string Token = TestLogValueTypes.url.ToLower();
+					var Properties = new Dictionary<string, object>() { { Token, TestLogValue.CreateURLLink(Name, URL) } };
+					string Format = $"{{{Token}}}";
+					string LogMessage = MessageTemplate.Render(Format, Properties);
+					LogEvent Event = new LogEvent(Time, Level, KnownLogEvents.Gauntlet_TestEvent, LogMessage, Format, Properties, null);
+
+					AddEvent(Event);
+				}
+
 				/// <summary>
 				/// Add Artifact event
 				/// </summary>
@@ -1885,6 +1903,10 @@ namespace Gauntlet
 				// Log entry
 				public static readonly Utf8String Type = new Utf8String("$type");
 				public static readonly Utf8String Text = new Utf8String("$text");
+				// Generic artifacts
+				public static readonly string artifacts = "Artifacts";
+				// Embedded HTML
+				public static readonly string embeddedHtml = "Embedded HTML";
 				// Image compare
 				public static readonly string imageCompare = "Image Compare";
 				public static readonly string unapproved = "unapproved";
@@ -1893,6 +1915,9 @@ namespace Gauntlet
 				public static readonly string approvedMetadata = "approved_metadata";
 				public static readonly string difference = "difference";
 				public static readonly string differenceReport = "difference_report";
+				// URL link
+				public static readonly string url = "URL";
+
 			}
 
 			/// <summary>
@@ -1915,6 +1940,14 @@ namespace Gauntlet
 				{
 					return Text;
 				}
+
+				/// <summary>
+				/// Utility function to create log value for URL link
+				/// </summary>
+				public static TestLogValue CreateURLLink(string Text, string URL)
+				{
+					return new TestLogValue(TestLogValueTypes.url, Text, new Dictionary<Utf8String, object>() { { new Utf8String("href"), URL } });
+				}
 			}
 
 			/// <summary>
@@ -1969,6 +2002,9 @@ namespace Gauntlet
 					RootPath = InRootPath;
 				}
 
+				public ArtifactFiles(string InName, Dictionary<string, string> InFiles, string InRootPath)
+					: this(TestLogValueTypes.artifacts, InName, InFiles, InRootPath) { }
+
 				private bool AttachFileTo(ref string FilePath, AutomatedTestSessionData Report)
 				{
 					if (string.IsNullOrEmpty(FilePath))
@@ -2044,6 +2080,15 @@ namespace Gauntlet
 				}
 			}
 
+			/// <summary>
+			/// Utility object to attached Embedded HTML file and generate related event properties
+			/// </summary>
+			public class EmbeddedHTMLFile : ArtifactFiles, ILogArtifact
+			{
+				public EmbeddedHTMLFile(string InName, string InFile, string InRootPath)
+					: base(TestLogValueTypes.embeddedHtml, InName, new Dictionary<string, string>() { { "path", InFile } }, InRootPath) { }
+			}
+
 			private static Logging.LogLevel EventTypeToLogLevel(EventType Type)
 			{
 				Logging.LogLevel Level = Logging.LogLevel.None;
diff --git a/Engine/Source/Programs/AutomationTool/Gauntlet/Unreal/Base/Gauntlet.UnrealTestNode.cs b/Engine/Source/Programs/AutomationTool/Gauntlet/Unreal/Base/Gauntlet.UnrealTestNode.cs
--- a/Engine/Source/Programs/AutomationTool/Gauntlet/Unreal/Base/Gauntlet.UnrealTestNode.cs
+++ b/Engine/Source/Programs/AutomationTool/Gauntlet/Unreal/Base/Gauntlet.UnrealTestNode.cs
@@ -13,6 +13,7 @@ using AutomationUtils.Matchers;
 using static Gauntlet.HordeReport;
 using static Gauntlet.UnrealSessionInstance;
 using Gauntlet.Utils;
+using static Gauntlet.HordeReport.AutomatedTestSessionData;
 
 namespace Gauntlet
 {
@@ -59,7 +60,6 @@ namespace Gauntlet
 		/// </summary>
 		public override float MaxDuration { get; protected set; }
 
-
 		/// Behavior flags for this test
 		/// </summary>
 		public BehaviorFlags Flags { get; protected set; }
@@ -832,49 +832,27 @@ namespace Gauntlet
 		}
 
 		/// <summary>
-		/// Used to Add a URL to a Horde Child Report.
+		/// Used to Add a URL to a Horde Report.
 		/// </summary>
-		/// <param name="Result">The current TestResult.</param>
 		/// <param name="Report">The curret ITestReport.</param>
 		/// <param name="ReportName">The name you want to show up on the Horde Child Report.</param>
 		/// <param name="URLLink">The URL Link you want to populate on the Horde Child Report.</param>
-		public void AddLinkToReport(TestResult Result, ITestReport Report, string ReportName, string URLLink)
+		public void AddLinkToReport(ITestReport Report, string ReportName, string URLLink)
 		{
-			if (CreateSimpleReportForHorde(Result) is not HordeReport.SimpleTestReport ChildReport)
+			if (Report == null)
 			{
-				return;
+				Log.Warning($"No Test report to associated URL Link {URLLink} with.");
 			}
-
-			if(!string.IsNullOrEmpty(ReportName) && !string.IsNullOrEmpty(URLLink))
+			else
 			{
-				Log.Info($"Adding a Horde Child Report URL Link for Report Name: {ReportName}, URLLink: {URLLink}.");
-				ChildReport.TestName = ReportName;
-				ChildReport.URLLink = URLLink;
-				AttachHordeReport(Report, ChildReport);
-
-				return;
+				if (Report is AutomatedTestSessionData SessionReport)
+				{
+					TestEventStream Stream = SessionReport.GetEventStream();
+					Stream.AddURLLink(DateTime.UtcNow, Microsoft.Extensions.Logging.LogLevel.Information, ReportName, URLLink);
+				}
 			}
 
-			Log.Warning("A Horde Child Report URL Link for Report Name: {ReportName}, URLLink: {URLLink} was requested but could not be generated.");
-
-		}
-
-		/// <summary>
-		/// Used to Attach a Horde Report to a report whether it's a base report, child, or another type.
-		/// </summary>
-		/// <param name="Report">The base report to attach the child report to.</param>
-		/// <param name="ChildReport">The child report to attach to the base report.</param>
-		public void AttachHordeReport(ITestReport Report, HordeReport.SimpleTestReport ChildReport)
-		{
-			switch (Report)
-			{
-				case HordeReport.BaseHordeReport BaseReport:
-					BaseReport.AttachDependencyReport(ChildReport, ChildReport.TestName);
-					break;
-				default:
-					Report = ChildReport;
-					break;
-			}
+			Horde.OnStepSummary += () => $"* **{ReportName}**: [Open link]({URLLink})";
 		}
 
 		public virtual void PopulateCommandlineInfo()
@@ -1394,6 +1372,7 @@ namespace Gauntlet
 		/// <param name="Result"></param>
 		protected virtual ITestReport CreateSimpleReportForHorde(TestResult Result)
 		{
+			Horde.OnStepSummary += () => Horde.GenerateArtifactLink(ArtifactPath);
 			BaseHordeReport HordeTestReport = null;
 			if (Globals.Params.ParseParam("UseTestDataV2"))
 			{
diff --git a/Engine/Source/Programs/AutomationTool/Gauntlet/Unreal/RunUnreal.cs b/Engine/Source/Programs/AutomationTool/Gauntlet/Unreal/RunUnreal.cs
--- a/Engine/Source/Programs/AutomationTool/Gauntlet/Unreal/RunUnreal.cs
+++ b/Engine/Source/Programs/AutomationTool/Gauntlet/Unreal/RunUnreal.cs
@@ -332,28 +332,8 @@ namespace Gauntlet
 			// Tear down accounts
 			TeardownAccounts();
 
-			// If only running a single unreal test node, we'll direct the summary artifact link to the specific test
-			string ArtifactPath = Globals.LogDir;
-			if(AllTestNodes.Count == 1)
-			{
-				try
-				{
-					// Can't read UnrealTestNode artifact path without some ugly reflection, attempt dynamic property read instead
-					dynamic UnrealTestNode = AllTestNodes.First();
-					string UnrealArtifactPath = UnrealTestNode.ArtifactPath;
-					if (!string.IsNullOrEmpty(UnrealArtifactPath))
-					{
-						ArtifactPath = UnrealArtifactPath;
-					}
-				}
-				catch
-				{
-					// Type isn't UnrealTestNode, use default behavior
-				}
-			}
-
-			// Generate Horde summary for CIS test (maybe want to use a delegate here)
-			Horde.GenerateSummary(ArtifactPath);
+			// Generate Horde summary for CIS test
+			Horde.GenerateSummary();
 
 			return AllTestsPassed ? ExitCode.Success : ExitCode.Error_TestFailure;
 		}
diff --git a/Engine/Source/Programs/Horde/HordeDashboard/plugins/test-automation/testAutomationEvents.tsx b/Engine/Source/Programs/Horde/HordeDashboard/plugins/test-automation/testAutomationEvents.tsx
--- a/Engine/Source/Programs/Horde/HordeDashboard/plugins/test-automation/testAutomationEvents.tsx
+++ b/Engine/Source/Programs/Horde/HordeDashboard/plugins/test-automation/testAutomationEvents.tsx
@@ -1,5 +1,5 @@
-import { Checkbox, DefaultButton, DirectionalHint, FontIcon, HoverCard, HoverCardType, Image, IPlainCardProps, mergeStyles, mergeStyleSets, Spinner, SpinnerSize, Stack, Text } from "@fluentui/react";
-import { Artifacts, EventEntry, EventLevel, ImageCompare, LogPropertyType } from "./testEventsModel";
+import { Checkbox, DefaultButton, DirectionalHint, FontIcon, HoverCard, HoverCardType, Image, IPlainCardProps, Link, mergeStyles, mergeStyleSets, Spinner, SpinnerSize, Stack, Text } from "@fluentui/react";
+import { Artifacts, EmbeddedHTML, EventEntry, EventLevel, ImageCompare, LogPropertyType, URLLink } from "./testEventsModel";
 import { TestPhaseOutcome } from "./api";
 import { memo, useCallback, useEffect, useState } from "react";
 import { ArtifactFactory, DeviceRef, TestPhaseStatus } from "./testData";
@@ -79,6 +79,16 @@ export const getEventStyling = () => {
                 }
             }
         },
+        artifactLinks: {
+            borderRadius: 6,
+            padding: '2px 6px',
+            backgroundColor: dashboard.darktheme ? "rgba(255, 255, 255, 0.1)" : "rgba(0, 0, 0, 0.3)",
+            selectors: {
+                ':active,:hover': {
+                    backgroundColor: dashboard.darktheme ? "rgba(255, 255, 255, 0.3)" : "rgba(0, 0, 0, 0.2)"
+                }
+            }
+        },
     });
 
     return eventStyles;
@@ -102,6 +112,15 @@ const EventTypeChooser = (property: string | number | Record<string, string | nu
     switch (property["$type"]) {
         case LogPropertyType.ImageCompare:
             return <ImageComparer key={`ic-${property["$text"]}`} imageCompare={property as ImageCompare} phase={phase} />
+
+        case LogPropertyType.EmbeddedHTML:
+            return <EmbeddedHTMLWidget key={`e-html-${property["$text"]}`} embeddedHTML={property as EmbeddedHTML} phase={phase} />
+
+        case LogPropertyType.Artifacts:
+            return <ArtifactsWidget key={`a-${property["$text"]}`} artifacts={property as Artifacts} phase={phase} />
+
+        case LogPropertyType.URLLink:
+            return <URLLinkWidget key={`link-${property['$text']}`} url={property as URLLink} phase={phase} />
     }
 
     return;
@@ -192,6 +211,11 @@ export const PhaseEventsPane: React.FC<{phase: TestPhaseStatus, onEventsLoaded?:
                             if (imageCompare.difference_report) artifacts.push(imageCompare.difference_report);
                             break;
 
+                        case LogPropertyType.EmbeddedHTML:
+                            const embeddedHTML = property as EmbeddedHTML;
+                            if (embeddedHTML.path) artifacts.push(embeddedHTML.path);
+                            break;
+
                         case LogPropertyType.Artifacts:
                             const artifactsProp = property as Artifacts
                             Object.entries(artifactsProp).forEach(([key, value]) => {
@@ -270,6 +294,78 @@ export const PhaseEventsPane: React.FC<{phase: TestPhaseStatus, onEventsLoaded?:
             </Stack>
 }
 
+/// Artifacts widget
+type ArtifactLink = { name: string, href: string }
+const getArtifactLinks = async (artifacts: Artifacts, factory: ArtifactFactory): Promise<ArtifactLink[]> => {
+    const links: ArtifactLink[] = [];
+    for (const [key, value] of Object.entries(artifacts)) {
+        if (key.startsWith('$') || !value) continue;
+        const href = await factory.getLink(value);
+        links.push({name: key, href: href});
+    }
+
+    return links;
+}
+
+const ArtifactsWidget: React.FC<{artifacts: Artifacts, phase: TestPhaseStatus}> = memo(({artifacts, phase}) => {
+    const [artifactsInfo, setArtifactsInfo] = useState<ArtifactLink[] | undefined>(undefined);
+
+    const eventStyles = getEventStyling();
+
+    useEffect(() => {
+        const phaseArtifacts = phase.artifacts;
+        !!phaseArtifacts && getArtifactLinks(artifacts, phaseArtifacts).then(info => setArtifactsInfo(info));
+    }, [artifacts, phase])
+
+    return  <Stack styles={{ root: { paddingLeft: 16, width: '100%' } }}>
+                <Stack><Text variant="medium"><span style={{ fontWeight: "bold" }}>{artifacts['$type']}: </span>{artifacts['$text']}</Text></Stack>
+                {artifactsInfo &&
+                    <Stack horizontal tokens={{childrenGap: 6}}>
+                        {artifactsInfo.map(item => 
+                                <Stack className={eventStyles.artifactLinks}>
+                                    <Link href={item.href} target="_blank">{item.name}</Link>
+                                </Stack>
+                            )
+                        }
+                    </Stack>
+                }
+            </Stack>
+});
+
+/// Embedded HTML widget
+const getEmbeddedHTMLLink = async (artifact: EmbeddedHTML, factory: ArtifactFactory): Promise<string> => {
+    return await factory.getLink(artifact.path);
+}
+
+const EmbeddedHTMLWidget: React.FC<{embeddedHTML: EmbeddedHTML, phase: TestPhaseStatus}> = memo(({embeddedHTML, phase}) => {
+    const [artifactInfo, setArtifactInfo] = useState<string | undefined>(undefined);
+
+    useEffect(() => {
+        const phaseArtifacts = phase.artifacts;
+        !!phaseArtifacts && getEmbeddedHTMLLink(embeddedHTML, phaseArtifacts).then(info => setArtifactInfo(info));
+    }, [embeddedHTML, phase])
+
+        const borderColor = dashboard.darktheme ? "#4D4C4B" : "#6D6C6B";
+
+    return  <Stack styles={{ root: { paddingLeft: 16, width: '100%' } }}>
+                {artifactInfo &&
+                    <Stack style={{borderWidth: 1, borderColor: borderColor, borderStyle: 'solid', padding: 4, margin: 1}} grow>
+                        <iframe style={{resize: 'vertical'}} src={artifactInfo}/>
+                    </Stack>
+                }
+            </Stack>
+});
+
+/// URL link widget
+const URLLinkWidget: React.FC<{url: URLLink, phase: TestPhaseStatus}> = memo(({url, phase}) => {
+    const eventStyles = getEventStyling();
+    return <Stack styles={{ root: { paddingLeft: 16 } }} horizontal>
+                <Stack className={eventStyles.artifactLinks}>
+                    <Link href={url.href} target="_blank">{url.$text}</Link>
+                </Stack>
+            </Stack>
+})
+
 /// Image comparison widget
 const missingImage = "/images/missing-image.png";
 const MissingImageLabel = (): JSX.Element => { return <span style={{ fontWeight: 'bold' }}> [missing image]</span> }
diff --git a/Engine/Source/Programs/Horde/HordeDashboard/plugins/test-automation/testEventsModel.ts b/Engine/Source/Programs/Horde/HordeDashboard/plugins/test-automation/testEventsModel.ts
--- a/Engine/Source/Programs/Horde/HordeDashboard/plugins/test-automation/testEventsModel.ts
+++ b/Engine/Source/Programs/Horde/HordeDashboard/plugins/test-automation/testEventsModel.ts
@@ -31,7 +31,9 @@ export enum EventLevel {
 /// Known property types
 export enum LogPropertyType {
     ImageCompare = "Image Compare",
-    Artifacts = "Artifacts"
+    Artifacts = "Artifacts",
+    EmbeddedHTML = "Embedded HTML",
+    URLLink = "URL"
 }
 
 /// Image Compare detailed properties
@@ -64,3 +66,21 @@ export type ImageCompare = {
 export type Artifacts = {
     [Key in string]: string | undefined;
 }
+
+/// Embedded html detailed properties
+export type EmbeddedHTML = {
+    "$type": LogPropertyType.EmbeddedHTML;
+    "$text": string;
+
+    //** The path to the html artifact */
+    path: string;
+}
+
+/// URL Link detailed properties
+export type URLLink = {
+    "$type": LogPropertyType.URLLink;
+    "$text": string;
+
+    //** the url link  */
+    href: string;
+}
