Fix warning in UE.Automation when ArtifactPath does not exist

#rnx

diff --git a/Engine/Source/Programs/AutomationTool/Gauntlet/Unreal/Automation/UE.Automation.cs b/Engine/Source/Programs/AutomationTool/Gauntlet/Unreal/Automation/UE.Automation.cs
--- a/Engine/Source/Programs/AutomationTool/Gauntlet/Unreal/Automation/UE.Automation.cs
+++ b/Engine/Source/Programs/AutomationTool/Gauntlet/Unreal/Automation/UE.Automation.cs
@@ -984,15 +984,18 @@ namespace UE
 			foreach (UnrealSessionInstance.RoleInstance RoleInstance in TestInstance.RunningRoles)
 			{
 				DirectoryInfo SrcDirectory = new DirectoryInfo(RoleInstance.AppInstance.ArtifactPath);
-				foreach (DirectoryInfo Dir in SrcDirectory.GetDirectories())
+				if (SrcDirectory.Exists)
 				{
-					if (Dir.Name.Equals("Screenshots", StringComparison.OrdinalIgnoreCase))
+					foreach (DirectoryInfo Dir in SrcDirectory.GetDirectories())
 					{
-						try { Dir.Delete(true); } catch { Log.Verbose("Could not delete directory: {Dir}", Dir.FullName); }
-					}
-					else if(Dir.Name.Equals("Automation", StringComparison.OrdinalIgnoreCase))
-					{
-						Dir.GetDirectories("Tmp", SearchOption.TopDirectoryOnly).ToList().ForEach(T => { try { T.Delete(true); } catch { Log.Verbose("Could not delete directory: {Dir}", T.FullName); } });
+						if (Dir.Name.Equals("Screenshots", StringComparison.OrdinalIgnoreCase))
+						{
+							try { Dir.Delete(true); } catch { Log.Verbose("Could not delete directory: {Dir}", Dir.FullName); }
+						}
+						else if (Dir.Name.Equals("Automation", StringComparison.OrdinalIgnoreCase))
+						{
+							Dir.GetDirectories("Tmp", SearchOption.TopDirectoryOnly).ToList().ForEach(T => { try { T.Delete(true); } catch { Log.Verbose("Could not delete directory: {Dir}", T.FullName); } });
+						}
 					}
 				}
 			}
