Fix missing reference validation in Horde summary report generation code

#rnx

diff --git a/Engine/Source/Programs/AutomationTool/Gauntlet/Framework/Horde/Gauntlet.Horde.cs b/Engine/Source/Programs/AutomationTool/Gauntlet/Framework/Horde/Gauntlet.Horde.cs
--- a/Engine/Source/Programs/AutomationTool/Gauntlet/Framework/Horde/Gauntlet.Horde.cs
+++ b/Engine/Source/Programs/AutomationTool/Gauntlet/Framework/Horde/Gauntlet.Horde.cs
@@ -92,10 +92,14 @@ namespace Gauntlet
 
 				StringBuilder Markdown = new StringBuilder();
 
-				foreach (SummaryLineDelegate callback in OnStepSummary.GetInvocationList())
+				if (OnStepSummary != null)
 				{
-					Markdown.AppendLine(callback());
+					foreach (SummaryLineDelegate callback in OnStepSummary.GetInvocationList())
+					{
+						Markdown.AppendLine(callback());
+					}
 				}
+
 				// detect if any testdata was generated
 				if (Globals.Params.ParseParam("UseTestDataV2")
 					&& Directory.GetFiles(Path.Combine(CommandUtils.CmdEnv.EngineSavedFolder, "TestData"), "*.json", SearchOption.TopDirectoryOnly).Any())
@@ -116,7 +120,7 @@ namespace Gauntlet
 			}
 			catch (Exception Ex)
 			{
-				Log.Info("Exception while generating Horde summary\n{0}\n", Ex.Message);
+				Log.Info("Exception while generating Horde summary\n{0}\n", Ex);
 			}
 		}
 
