Use ReportError instead of just logging an error to have the report peak up on it

#rnx

diff --git a/Engine/Source/Programs/AutomationTool/Gauntlet/Unreal/Automation/UE.PLMTest.cs b/Engine/Source/Programs/AutomationTool/Gauntlet/Unreal/Automation/UE.PLMTest.cs
--- a/Engine/Source/Programs/AutomationTool/Gauntlet/Unreal/Automation/UE.PLMTest.cs
+++ b/Engine/Source/Programs/AutomationTool/Gauntlet/Unreal/Automation/UE.PLMTest.cs
@@ -313,7 +313,7 @@ namespace UE
 				bool bSuccess = CurrentAction.Run( RunningInstance, this );
 				if (!bSuccess)
 				{
-					Log.Error($"{CurrentAction} failed. Ending test");
+					ReportError($"{CurrentAction} failed. Ending test");
 					MarkTestComplete();
 					SetUnrealTestResult(TestResult.Failed);
 				}
