[TargetDeviceIOS] Improve error handling during Json parsing to avoid runtime exceptions when getting connected devices.

diff --git a/Engine/Source/Programs/AutomationTool/Gauntlet/Platform/IOS/Gauntlet.TargetDeviceIOS.cs b/Engine/Source/Programs/AutomationTool/Gauntlet/Platform/IOS/Gauntlet.TargetDeviceIOS.cs
--- a/Engine/Source/Programs/AutomationTool/Gauntlet/Platform/IOS/Gauntlet.TargetDeviceIOS.cs
+++ b/Engine/Source/Programs/AutomationTool/Gauntlet/Platform/IOS/Gauntlet.TargetDeviceIOS.cs
@@ -160,7 +160,7 @@ namespace Gauntlet
 
 			IsDefaultDevice = string.IsNullOrEmpty(InName) || InName.Equals("default", StringComparison.OrdinalIgnoreCase);
 			Name = IsDefaultDevice ? "default" : InName;
-
+			
 			if (IsDefaultDevice)
 			{
 				IEnumerable<ConnectedDevice> AvailableDevices = ConnectedDevices.Values.Where(Device => !Device.IsAssigned);
@@ -1052,14 +1052,25 @@ namespace Gauntlet
 
 				foreach (JsonObject JsonDevice in JsonDevices)
 				{
-					string UUID = JsonDevice.GetObjectField("hardwareProperties").GetStringField("udid");
-					string OSVersionString = JsonDevice.GetObjectField("deviceProperties").GetStringField("osVersionNumber");
-					int IOSMajorVersion = int.Parse(OSVersionString.Substring(0, OSVersionString.IndexOf('.')));
+					ConnectedDevice IOSDevice = new ConnectedDevice(string.Empty, -1, false);
+					
+					if (JsonDevice.TryGetObjectField("hardwareProperties", out JsonObject HardwareProperties))
+					{
+						HardwareProperties.TryGetStringField("udid", out IOSDevice.UUID);
+					}
 
-					bool bIsConnected = ValidateDeviceConnection(JsonDevice);
+					if (JsonDevice.TryGetObjectField("deviceProperties", out JsonObject DeviceProperties) &&
+					    DeviceProperties.TryGetStringField("osVersionNumber", out string OSVersionString) &&
+					    !string.IsNullOrWhiteSpace(OSVersionString))
+					{
+						IOSDevice.IOSMajorVersion = int.Parse(OSVersionString.Substring(0, OSVersionString.IndexOf('.')));
+					}
 
-					ConnectedDevice Device = new ConnectedDevice(UUID, IOSMajorVersion, bIsConnected);
-					Devices.Add(UUID, Device);
+					if (!string.IsNullOrWhiteSpace(IOSDevice.UUID) && IOSDevice.IOSMajorVersion != -1)
+					{
+						IOSDevice.bConnected = ValidateDeviceConnection(JsonDevice);	
+						Devices.Add(IOSDevice.UUID, IOSDevice);
+					}
 				}
 			}
 			else
@@ -1123,7 +1134,7 @@ namespace Gauntlet
 
 			return bIsConnected;
 		}
-
+		
 		// DEPRECATED
 		public IAppInstall InstallApplication(UnrealAppConfig AppConfig)
 		{
