Add support for Gauntlet scripts inside UAT plugins

#rnx

diff --git a/Engine/Source/Programs/AutomationTool/Gauntlet/Gauntlet.Automation.csproj b/Engine/Source/Programs/AutomationTool/Gauntlet/Gauntlet.Automation.csproj
--- a/Engine/Source/Programs/AutomationTool/Gauntlet/Gauntlet.Automation.csproj
+++ b/Engine/Source/Programs/AutomationTool/Gauntlet/Gauntlet.Automation.csproj
@@ -61,6 +61,21 @@
     </Compile>
   </ItemGroup>
 
+    <!-- The following will glob all UAT plugin's Gauntlet source
+  files and include them in the build. They will also appear in VisualStudio -->
+  <ItemGroup>
+    <Compile Include="../../../../Plugins/**/Build/Scripts/Gauntlet/**/*.cs">
+      <!-- Shorten the path so that everything is under a Plugins folder but does not include the Build/Scripts/Gauntlet part -->
+      <Link>$([System.Text.RegularExpressions.Regex]::Replace(%(FullPath), '^.*?(?=Plugins[\\/])|Build[\\/]Scripts[\\/]Gauntlet[\\/]',''))</Link>      
+    </Compile> 
+    <Compile Include="../../../../Platforms/*/Plugins/**/Build/Scripts/Gauntlet/**/*.cs">
+      <!-- RecursiveDir metadata expands to [PlatformName]/Source/Programs/[etc]. The Replace() will
+      replace everything from the first slash to the end leaving just the platform name. Should the
+      Replace() fail then RecursiveDir is still an agreeable value -->
+      <Link>Platform/$([System.Text.RegularExpressions.Regex]::Replace(%(Compile.RecursiveDir), [\\/].+$, ``))/%(Compile.FileName).cs</Link>
+    </Compile>
+  </ItemGroup>
+
   <Import Project="../../../../Platforms/*/Source/Programs/AutomationTool/Gauntlet/*.Gauntlet.targets" Condition="'$(OS)' == 'Windows_NT'" />
 
   <ItemGroup>
