Add Clang support to WindowsPGOPlatform in Gauntlet

diff --git a/Engine/Source/Programs/AutomationTool/Gauntlet/Platform/Windows/Gauntlet.WindowsPGO.cs b/Engine/Source/Programs/AutomationTool/Gauntlet/Platform/Windows/Gauntlet.WindowsPGO.cs
--- a/Engine/Source/Programs/AutomationTool/Gauntlet/Platform/Windows/Gauntlet.WindowsPGO.cs
+++ b/Engine/Source/Programs/AutomationTool/Gauntlet/Platform/Windows/Gauntlet.WindowsPGO.cs
@@ -1,6 +1,7 @@
 // Copyright Epic Games, Inc. All Rights Reserved.
 
 using AutomationTool;
+using System;
 using System.IO;
 using System.Text.RegularExpressions;
 using UnrealBuildTool;
@@ -14,6 +15,8 @@ namespace Gauntlet
 	{
 		string LocalOutputDirectory;
 		string PgcFilenamePrefix;
+		string ExpectedProfrawFileName;
+
 		UnrealTargetPlatform Platform;
 
 		protected WinBasePGOPlatform(UnrealTargetPlatform InPlatform)
@@ -26,10 +29,10 @@ namespace Gauntlet
 			return Platform;
 		}
 
-		public void GatherResults(string ArtifactPath)
+		private void GatherResultsMSVC(string PathToPGOArtifacts)
 		{
 			// Find PGC files
-			string PGODir = Path.Combine(ArtifactPath, "PGO");
+			string PGODir = PathToPGOArtifacts;
 			string[] PGCFiles = Directory.GetFiles(PGODir, "*.pgc");
 			if (PGCFiles.Length == 0)
 			{
@@ -67,14 +70,42 @@ namespace Gauntlet
 			}
 		}
 
+		private void GatherResultsClang()
+		{
+			// We've already verified the file was written to the desired path.
+			Log.Info("Found expected .profraw file: {0}", ExpectedProfrawFileName);
+		}
+
+		public void GatherResults(string ArtifactPath)
+		{
+			string PGODir = Path.Combine(ArtifactPath, "PGO");
+
+			if (File.Exists(ExpectedProfrawFileName)) 
+			{
+				Log.Info("Clang format profiling files detected");
+				GatherResultsClang();
+			}
+			else
+			{
+				Log.Info("MSVC format profiling files detected");
+				GatherResultsMSVC(PGODir);
+			}
+		}
+
 		public void ApplyConfiguration(PGOConfig Config)
 		{
 			PgcFilenamePrefix = Config.PgcFilenamePrefix;
 			LocalOutputDirectory = Path.GetFullPath(Config.ProfileOutputDirectory);
+			ExpectedProfrawFileName = Environment.GetEnvironmentVariable("LLVM_PROFILE_FILE");
+			if (ExpectedProfrawFileName == null && !string.IsNullOrWhiteSpace(PgcFilenamePrefix))
+			{
+				ExpectedProfrawFileName = Path.Combine(LocalOutputDirectory, PgcFilenamePrefix + ".profraw");
+				Environment.SetEnvironmentVariable("LLVM_PROFILE_FILE", ExpectedProfrawFileName);
+			}
 
 			UnrealTestRole ClientRole = Config.RequireRole(UnrealTargetRole.Client);
 			ClientRole.CommandLine += $" -PGOSweepToSaveDir"; // this means the PGC files will be created in the Gauntlet artifact folder
-			ClientRole.ConfigureDevice = ConfigureDevice;			
+			ClientRole.ConfigureDevice = ConfigureDevice;
 		}
 
 		protected abstract string GetLocalCachePath(ITargetDevice Device);
