Avoid failing generating Horde summary report when TestData folder did not exist

#rnx

diff --git a/Engine/Source/Programs/AutomationTool/Gauntlet/Framework/Horde/Gauntlet.Horde.cs b/Engine/Source/Programs/AutomationTool/Gauntlet/Framework/Horde/Gauntlet.Horde.cs
--- a/Engine/Source/Programs/AutomationTool/Gauntlet/Framework/Horde/Gauntlet.Horde.cs
+++ b/Engine/Source/Programs/AutomationTool/Gauntlet/Framework/Horde/Gauntlet.Horde.cs
@@ -101,10 +101,17 @@ namespace Gauntlet
 				}
 
 				// detect if any testdata was generated
-				if (Globals.Params.ParseParam("UseTestDataV2")
-					&& Directory.GetFiles(Path.Combine(CommandUtils.CmdEnv.EngineSavedFolder, "TestData"), "*.json", SearchOption.TopDirectoryOnly).Any())
+				if (Globals.Params.ParseParam("UseTestDataV2"))
 				{
-					Markdown.AppendLine($"* **Test Results**: [Open report](/test-automation?job={JobId}&step={StepId})");
+					DirectoryInfo TestDataOutputFolder = new DirectoryInfo(Path.Combine(CommandUtils.CmdEnv.EngineSavedFolder, "TestData"));
+					if (TestDataOutputFolder.Exists && TestDataOutputFolder.GetFiles("*.json", SearchOption.TopDirectoryOnly).Any())
+					{
+						Markdown.AppendLine($"* **Test Results**: [Open report](/test-automation?job={JobId}&step={StepId})");
+					}
+					else
+					{
+						Log.Verbose("Did not add Horde report link as any test data files in '{0}' was found.\n", TestDataOutputFolder.FullName);
+					}
 				}
 
 				if (Markdown.Length > 0)
