Allow populating test phases through delegates + Add RPC test cases to phase report

#rnx

diff --git a/Engine/Source/Programs/AutomationTool/Gauntlet/Framework/Utils/Gauntlet.HordeReport.cs b/Engine/Source/Programs/AutomationTool/Gauntlet/Framework/Utils/Gauntlet.HordeReport.cs
--- a/Engine/Source/Programs/AutomationTool/Gauntlet/Framework/Utils/Gauntlet.HordeReport.cs
+++ b/Engine/Source/Programs/AutomationTool/Gauntlet/Framework/Utils/Gauntlet.HordeReport.cs
@@ -1326,6 +1326,17 @@ namespace Gauntlet
 					timeElapseSec = TimeElapse;
 				}
 
+				/// <summary>
+				/// Set the outcome of the phase, it will udpate the test session phase counts
+				/// </summary>
+				/// <param name="OutCome"></param>
+				public void SetOutcome(TestPhaseOutcome OutCome)
+				{
+					UpdatePhasesCount(PhaseCount.Out);
+					outcome = OutCome;
+					UpdatePhasesCount();
+				}
+
 				/// <summary>
 				/// Set the outcome of the phase, it will update the test session phase counts
 				/// </summary>
@@ -1500,6 +1511,15 @@ namespace Gauntlet
 				return Data.devices.GetValueOrDefault(Key);
 			}
 
+			/// <summary>
+			/// Return all the devices associated with the report
+			/// </summary>
+			/// <returns></returns>
+			public virtual List<TestDevice> GetDevices()
+			{
+				return Data.devices.Values.ToList();
+			}
+
 			public override bool AttachDeviceLog(string InstanceName, string LogPath, string LogName)
 			{
 				if (AttachArtifact(LogPath, LogName))
diff --git a/Engine/Source/Programs/AutomationTool/Gauntlet/Unreal/Base/Gauntlet.UnrealTestNode.cs b/Engine/Source/Programs/AutomationTool/Gauntlet/Unreal/Base/Gauntlet.UnrealTestNode.cs
--- a/Engine/Source/Programs/AutomationTool/Gauntlet/Unreal/Base/Gauntlet.UnrealTestNode.cs
+++ b/Engine/Source/Programs/AutomationTool/Gauntlet/Unreal/Base/Gauntlet.UnrealTestNode.cs
@@ -68,6 +68,16 @@ namespace Gauntlet
 		/// </summary>
 		public override TestPriority Priority { get { return GetPriority(); } }
 
+		/// <summary>
+		/// Delegate type to populate test phase in report
+		/// </summary>
+		/// <param name="TestReport"></param>
+		protected delegate void PopulateTestPhasesReport(AutomatedTestSessionData TestReport);
+		/// <summary>
+		/// Call when populating Horde test report
+		/// </summary>
+		protected PopulateTestPhasesReport OnPopulateTestPhases;
+
 		/// <summary>
 		/// Returns a list of all log channels the heartbeat tick should look for.
 		/// </summary>
@@ -1383,41 +1393,39 @@ namespace Gauntlet
 				// Simple report has one phase
 				var InitPhase = TestReport.AddPhase("Initialization", "Initialization");
 				InitPhase.SetTiming(SessionTime, 0); // Unfortunately we are not tracking initialization duration at this point
-				var MainPhase = TestReport.AddPhase("Main");
-				MainPhase.SetTiming(SessionTime, TimeElapse);
-				// Propagate events
 				if (RoleResults != null)
 				{
 					foreach (UnrealRoleResult Item in RoleResults)
 					{
 						string RoleType = Item.Artifacts.SessionRole.RoleType.ToString();
-						var RelevantPhase = (Item.ProcessResult == UnrealProcessResult.InitializationFailure || Item.ProcessResult == UnrealProcessResult.LoginFailed)?
-								InitPhase: MainPhase;
-						var RelevantStream = RelevantPhase.GetStream();
-
-						foreach (UnrealTestEvent Event in Item.Events)
+						if (Item.ProcessResult == UnrealProcessResult.InitializationFailure || Item.ProcessResult == UnrealProcessResult.LoginFailed)
 						{
-							RelevantStream.AddEvent(Event);
-						}
-
-						// Add Process result
-						if (Item.ProcessResult != UnrealProcessResult.ExitOk)
-						{
-							RelevantStream.AddError($"{RoleType}: {Item.Summary}");
+							var RelevantStream = InitPhase.GetStream();
+							foreach (UnrealTestEvent Event in Item.Events)
+							{
+								RelevantStream.AddEvent(Event);
+							}
+							// Add Process result
+							if (Item.ProcessResult != UnrealProcessResult.ExitOk)
+							{
+								RelevantStream.AddError($"{RoleType}: {Item.Summary}");
+							}
 						}
 						// Since we are going through the Role, lets add the related device
 						var ReportDevice = AddRoleDeviceToSessionReport(Item, TestReport);
 						InitPhase.deviceKeys.Add(ReportDevice.key);
-						MainPhase.deviceKeys.Add(ReportDevice.key);
 					}
 				}
-				var MainStream = MainPhase.GetStream();
-				foreach (UnrealTestEvent Event in TestNodeEvents)
+				if (OnPopulateTestPhases != null)
 				{
-					MainStream.AddEvent(Event);
+					// Some delegates have been setup
+					OnPopulateTestPhases(TestReport);
+				}
+				else
+				{
+					// Default phase if no delegate were registered
+					PopulateClosurePhase(TestReport, "Main", SessionTime, TimeElapse);
 				}
-				// Outcome
-				MainPhase.SetOutcome(GetTestResult());
 				HordeTestReport = TestReport;
 				// Metadata
 				SetReportMetadata(HordeTestReport, new UnrealTestRole[] { GetCachedConfiguration().GetMainRequiredRole() });
@@ -1463,6 +1471,46 @@ namespace Gauntlet
 			return HordeTestReport;
 		}
 
+		/// <summary>
+		/// Default populate closure phase
+		/// </summary>
+		/// <param name="TestReport"></param>
+		/// <param name="PhaseName"></param>
+		protected virtual void PopulateClosurePhase(AutomatedTestSessionData TestReport, string PhaseName = "Closure", DateTime? SessionTime = null, float TimeElapseSeconds = 0)
+		{
+			var ClosurePhase = TestReport.AddPhase(PhaseName);
+			ClosurePhase.SetTiming(SessionTime ?? DateTime.UtcNow, TimeElapseSeconds);
+			// Propagate events
+			if (RoleResults != null)
+			{
+				foreach (UnrealRoleResult Item in RoleResults)
+				{
+					string RoleType = Item.Artifacts.SessionRole.RoleType.ToString();
+					if (Item.ProcessResult != UnrealProcessResult.InitializationFailure && Item.ProcessResult != UnrealProcessResult.LoginFailed)
+					{
+						var RelevantStream = ClosurePhase.GetStream();
+						foreach (UnrealTestEvent Event in Item.Events)
+						{
+							RelevantStream.AddEvent(Event);
+						}
+						// Add Process result
+						if (Item.ProcessResult != UnrealProcessResult.ExitOk)
+						{
+							RelevantStream.AddError($"{RoleType}: {Item.Summary}");
+						}
+					}
+				}
+				ClosurePhase.deviceKeys = TestReport.GetDevices().Select(D => D.key).ToHashSet();
+			}
+			var MainStream = ClosurePhase.GetStream();
+			foreach (UnrealTestEvent Event in TestNodeEvents)
+			{
+				MainStream.AddEvent(Event);
+			}
+			// Outcome
+			ClosurePhase.SetOutcome(GetTestResult());
+		}
+
 		private AutomatedTestSessionData.TestDevice AddRoleDeviceToSessionReport(UnrealRoleResult RoleResult, AutomatedTestSessionData TestReport)
 		{
 			ITargetDevice TargetDevice = RoleResult.Artifacts.AppInstance.Device;
