Fix GenerateProcessedCommandLine captures when escape double quotes are used in conjunction to normal double quotes

#rnx

diff --git a/Engine/Source/Programs/AutomationTool/Gauntlet/Unreal/Base/Gauntlet.UnrealTestConfiguration.cs b/Engine/Source/Programs/AutomationTool/Gauntlet/Unreal/Base/Gauntlet.UnrealTestConfiguration.cs
--- a/Engine/Source/Programs/AutomationTool/Gauntlet/Unreal/Base/Gauntlet.UnrealTestConfiguration.cs
+++ b/Engine/Source/Programs/AutomationTool/Gauntlet/Unreal/Base/Gauntlet.UnrealTestConfiguration.cs
@@ -17,6 +17,8 @@ namespace Gauntlet
 		private Dictionary<string, object> Params;
 		private HashSet<string> NonOptionParams;
 
+		public static readonly Regex CommandlineOptionRegex = new Regex(@"-(?<option>-?[\w\d.:!\[\]\/\\?-]+)(=(?<value>(""(?:(?:\\""|[^""])+)"")|(?:\S+)))?", RegexOptions.Compiled);
+
 		// Give external people read-only access
 		public IReadOnlyDictionary<string, object> Arguments {  get { return Params;  } }
 
@@ -54,7 +56,7 @@ namespace Gauntlet
 		public void AddRawCommandline(string InRawCommandline, bool bOverrideExistingValues = true)
 		{
 			// turn Name(p1,etc) into a collection of Name|(p1,etc) groups
-			MatchCollection Matches = Regex.Matches(InRawCommandline, "-(?<option>\\-?[\\w\\d.:!\\[\\]\\/\\\\\\-]+)(=(?<value>(\"([^\"]*)\")|(\\S+)))?");
+			MatchCollection Matches = CommandlineOptionRegex.Matches(InRawCommandline);
 
 			foreach (Match M in Matches)
 			{
diff --git a/Engine/Source/Programs/AutomationTool/Gauntlet/Unreal/BuildSource/Gauntlet.UnrealBuildSource.cs b/Engine/Source/Programs/AutomationTool/Gauntlet/Unreal/BuildSource/Gauntlet.UnrealBuildSource.cs
--- a/Engine/Source/Programs/AutomationTool/Gauntlet/Unreal/BuildSource/Gauntlet.UnrealBuildSource.cs
+++ b/Engine/Source/Programs/AutomationTool/Gauntlet/Unreal/BuildSource/Gauntlet.UnrealBuildSource.cs
@@ -722,11 +722,10 @@ namespace Gauntlet
 		/// <returns></returns>
 		private string GenerateProcessedCommandLine(string InCommandLine)
 		{
-
 			// Break down Commandline into individual tokens
 			Dictionary<string, string> CommandlineTokens = new Dictionary<string, string>(StringComparer.OrdinalIgnoreCase);
 			// turn Name(p1,etc) into a collection of Name|(p1,etc) groups
-			MatchCollection Matches = Regex.Matches(InCommandLine, @"(?<option>-?[\w\d.:!\[\]\/\\?-]+)(=(?<value>(""([^""]*)"")|(\S+)))?");
+			MatchCollection Matches = GauntletCommandLine.CommandlineOptionRegex.Matches(InCommandLine);
 
 			foreach (Match M in Matches)
 			{
@@ -745,9 +744,9 @@ namespace Gauntlet
 
 					if (string.IsNullOrWhiteSpace(Params))
 					{
-						Log.Info(string.Format("Duplicate flag {0} found and ignored. Please fix this as it will increase in severity in 01/2019. ", Name));
+						Log.Info(string.Format("Duplicate flag {0} found and ignored. Please fix this. ", Name));
 					}
-					else if (Name.ToLower() == "-execcmds")
+					else if (Name.ToLower() == "execcmds")
 					{
 						// Special cased as execcmds is something that is totally able to be appended to. Requote everything when we're done.
 						CommandlineTokens[Name] = string.Format("\"{0}, {1}\"", CommandlineTokens[Name].Replace("\"", ""), Params.Replace("\"", ""));
@@ -756,7 +755,7 @@ namespace Gauntlet
 					{
 						if (CommandlineTokens[Name] == Params)
 						{
-							Log.Info(string.Format("Duplicate flag {0}={1} found and ignored. Please fix this as this log line will increase in severity in 01/2019. ", Name, Params));
+							Log.Info(string.Format("Duplicate flag {0}={1} found and ignored. Please fix this. ", Name, Params));
 						}
 						else
 						{
@@ -767,14 +766,14 @@ namespace Gauntlet
 				}
 				else
 				{
-					CommandlineTokens.Add(Name, (Params.Contains(' ') && !Params.Contains('\"')) ? string.Format("\"{0}\"", Params) : Params);
+					CommandlineTokens.Add(Name, Params);
 				}
 			}
 
 			string CommandlineToReturn = "";
 			foreach (string DictKey in CommandlineTokens.Keys)
 			{
-				CommandlineToReturn += string.Format("{0}{1} ",
+				CommandlineToReturn += string.Format("-{0}{1} ",
 					DictKey,
 					string.IsNullOrWhiteSpace(CommandlineTokens[DictKey]) ? "" : string.Format("={0}", CommandlineTokens[DictKey])
 					);
