Backout the pak file rename overrides which incorrectly renamed paks invalid test paths.
Add option -DLCSkipNamedPak to allow DLC paking to not add the DLCName to the pak file.

Remove DLCNameOverride option because it is too dangerous.

diff --git a/Engine/Source/Programs/AutomationTool/AutomationUtils/ProjectParams.cs b/Engine/Source/Programs/AutomationTool/AutomationUtils/ProjectParams.cs
--- a/Engine/Source/Programs/AutomationTool/AutomationUtils/ProjectParams.cs
+++ b/Engine/Source/Programs/AutomationTool/AutomationUtils/ProjectParams.cs
@@ -313,6 +313,7 @@ namespace AutomationTool
 			this.DiscVersion = InParams.DiscVersion;
             this.DLCIncludeEngineContent = InParams.DLCIncludeEngineContent;
 			this.DLCActLikePatch = InParams.DLCActLikePatch;
+			this.DLCSkipNamedPak = InParams.DLCSkipNamedPak;
 			this.DLCPakPluginFile = InParams.DLCPakPluginFile;
 			this.DLCOverrideCookedSubDir = InParams.DLCOverrideCookedSubDir;
 			this.DLCOverrideStagedSubDir = InParams.DLCOverrideStagedSubDir;
@@ -550,6 +551,7 @@ namespace AutomationTool
             bool? DLCIncludeEngineContent = null,
 			bool? DLCPakPluginFile = null,
 			bool? DLCActLikePatch = null,
+			bool? DLCSkipNamedPak = null,
 			bool? CrashReporter = null,
 			bool? DedicatedServer = null,
 			bool? Client = null,
@@ -774,6 +776,7 @@ namespace AutomationTool
             this.DLCIncludeEngineContent = GetParamValueIfNotSpecified(Command, DLCIncludeEngineContent, this.DLCIncludeEngineContent, "DLCIncludeEngineContent");
 			this.DLCPakPluginFile = GetParamValueIfNotSpecified(Command, DLCPakPluginFile, this.DLCPakPluginFile, "DLCPakPluginFile");
 			this.DLCActLikePatch = GetParamValueIfNotSpecified(Command, DLCActLikePatch, this.DLCActLikePatch, "DLCActLikePatch");
+			this.DLCSkipNamedPak = GetParamValueIfNotSpecified(Command, DLCSkipNamedPak, this.DLCSkipNamedPak, "DLCSkipNamedPak");
 			this.DLCOverrideCookedSubDir = ParseParamValueIfNotSpecified(Command, DLCOverrideCookedSubDir, "DLCOverrideCookedSubDir", null);
 			this.DLCOverrideStagedSubDir = ParseParamValueIfNotSpecified(Command, DLCOverrideStagedSubDir, "DLCOverrideStagedSubDir", null);
 
@@ -1986,6 +1989,11 @@ namespace AutomationTool
 		/// </summary>
 		public string DLCOverrideEncryptionKeyGuid;
 
+		/// <summary>
+		/// When true the DLC Name will not be added to the generated pak files.
+		/// </summary>
+		public bool DLCSkipNamedPak;
+
 		/// <summary>
 		/// After cook completes diff the cooked content against another cooked content directory.
 		///  report all errors to the log
diff --git a/Engine/Source/Programs/AutomationTool/Scripts/CopyBuildToStagingDirectory.Automation.cs b/Engine/Source/Programs/AutomationTool/Scripts/CopyBuildToStagingDirectory.Automation.cs
--- a/Engine/Source/Programs/AutomationTool/Scripts/CopyBuildToStagingDirectory.Automation.cs
+++ b/Engine/Source/Programs/AutomationTool/Scripts/CopyBuildToStagingDirectory.Automation.cs
@@ -4068,30 +4068,6 @@ namespace AutomationScripts
 
 			CommonAdditionalArgs += " " + AdditionalCompressionOptionsOnCommandLine;
 
-			Dictionary<string, string> DLCToNameOverrides = new Dictionary<string, string>();
-			{
-				List<string> DLCToNameOverridesTmp;
-				if (PlatformGameConfig.GetArray("/Script/UnrealEd.ProjectPackagingSettings", "DLCPakNameOverride", out DLCToNameOverridesTmp))
-				{
-					// Each string has the format '+DLCPakNameOverride=(PakName="TheNameToReplaceFromThePak", OverrideName="TheNameToUseInstead")'
-					// for @ regex " is written as ""
-					Regex DLCPakNameOverrideRegex = new Regex(@"\(PakName\s*=\s*(?<PakName>""*\w*""*),\s*OverrideName\s*=\s*(?<OverrideName>""*\w*""*)\)");
-
-					foreach (var PathStr in DLCToNameOverridesTmp)
-					{
-						var DLCPakOverrideMatch = DLCPakNameOverrideRegex.Match(PathStr);
-						if (DLCPakOverrideMatch.Success)
-						{
-							string DLCNameToOverride = DLCPakOverrideMatch.Groups["PakName"].Value.StripQuoteArgument();
-							string OverrideName = DLCPakOverrideMatch.Groups["OverrideName"].Value.StripQuoteArgument();
-							DLCToNameOverrides.Add(DLCNameToOverride, OverrideName);
-
-							Logger.LogInformation($"Adding DLCPak Name Override: {DLCNameToOverride} -> {OverrideName}");
-						}
-					}
-				}
-			}
-
 			bool bCopiedExistingGlobalUtoc = false;
 			List<Tuple<FileReference, StagedFileReference, string>> Outputs = new List<Tuple<FileReference, StagedFileReference, string>>();
 			foreach (CreatePakParams PakParams in PakParamsList)
@@ -4107,12 +4083,15 @@ namespace AutomationScripts
 				StagedFileReference OutputRelativeLocation;
 				if (Params.HasDLCName)
 				{
-					string DLCPakName = Params.DLCFile.GetFileNameWithoutExtension() + OutputFilename;
-					foreach (var PakToReplacement in DLCToNameOverrides)
+					string DLCPakName = null;
+					if (Params.DLCSkipNamedPak)
 					{
-						DLCPakName = DLCPakName.Replace(PakToReplacement.Key, PakToReplacement.Value);
+						DLCPakName = OutputFilename;
+					}
+					else
+					{
+						DLCPakName = Params.DLCFile.GetFileNameWithoutExtension() + OutputFilename;
 					}
-
 					if (Params.DLCOverrideStagedSubDir != null)
 					{
 						OutputRelativeLocation = StagedFileReference.Combine(SC.RelativeProjectRootForStage, Params.DLCOverrideStagedSubDir, "Content", "Paks", DLCPakName + ".pak");
