Cloud Artifacts - Removed ability to specify a range of commits to allow for download. Does not work outside of P4 and even within p4 its very hard to know what to set this to. Instead we just have exact commit matching (which was the default before and still is) which you can now opt out from if you do not care about which commit it is or its derived from the build name.

diff --git a/Engine/Source/Programs/AutomationTool/BuildGraph/Tasks/RetrieveCloudArtifactTask.cs b/Engine/Source/Programs/AutomationTool/BuildGraph/Tasks/RetrieveCloudArtifactTask.cs
--- a/Engine/Source/Programs/AutomationTool/BuildGraph/Tasks/RetrieveCloudArtifactTask.cs
+++ b/Engine/Source/Programs/AutomationTool/BuildGraph/Tasks/RetrieveCloudArtifactTask.cs
@@ -42,10 +42,10 @@ namespace AutomationTool.Tasks
 		public string? Commit { get; set; }
 
 		/// <summary>
-		/// Maximum commit for the artifact
+		/// Requires that the current synced commit is the same as the artifacts commit
 		/// </summary>
 		[TaskParameter(Optional = true)]
-		public string? MaxCommit { get; set; }
+		public bool? RequireMatchingCommit { get; set; }
 
 		/// <summary>
 		/// Name of the artifact
@@ -154,11 +154,6 @@ namespace AutomationTool.Tasks
 		/// <param name="tagNameToFileSet">Mapping from tag names to the set of files they include.</param>
 		public override async Task ExecuteAsync(JobContext job, HashSet<FileReference> buildProducts, Dictionary<string, HashSet<FileReference>> tagNameToFileSet)
 		{
-			if (!String.IsNullOrEmpty(_parameters.Commit) && !String.IsNullOrEmpty(_parameters.MaxCommit))
-			{
-				throw new AutomationException("Cannot specify both Commit and MaxCommit parameters for retrieving an artifact.");
-			}
-
 			string project = _parameters.Project.GetFileNameWithoutExtension();
 			FileReference? projectPath = _parameters.Project;
 			if (!File.Exists(_parameters.Project.FullName))
@@ -274,29 +269,28 @@ namespace AutomationTool.Tasks
 			}
 
 			// Get the current commit id
-			CommitId? minCommitId = null;
-			CommitId? maxCommitId = null;
-			if (!String.IsNullOrEmpty(_parameters.MaxCommit))
-			{
-				maxCommitId = new CommitId(_parameters.MaxCommit);
-			}
-			else if (!String.IsNullOrEmpty(_parameters.Commit))
-			{
-				minCommitId = maxCommitId = new CommitId(_parameters.Commit);
-			}
-			else
+			CommitId? requiredCommitId = null;
+			if (_parameters.RequireMatchingCommit.GetValueOrDefault(true))
 			{
-				try
+				// the commit needs to match, check the commit option or fallback to P4 in case it has not been specified
+				if (!String.IsNullOrEmpty(_parameters.Commit))
 				{
-					int change = CommandUtils.P4Env.Changelist;
-					if (change > 0)
-					{
-						minCommitId = maxCommitId = CommitId.FromPerforceChange(CommandUtils.P4Env.Changelist);
-					}
+					requiredCommitId = new CommitId(_parameters.Commit);
 				}
-				catch (AutomationException)
+				else
 				{
-					// not an error to run without p4
+					try
+					{
+						int change = CommandUtils.P4Env.Changelist;
+						if (change > 0)
+						{
+							requiredCommitId = CommitId.FromPerforceChange(CommandUtils.P4Env.Changelist);
+						}
+					}
+					catch (AutomationException)
+					{
+						// not an error to run without p4
+					}
 				}
 			}
 
@@ -317,27 +311,11 @@ namespace AutomationTool.Tasks
 			queryWriter.WriteString("$eq", CreateCloudArtifactTask.SanitizeBucketValue(streamId.ToString()));
 			queryWriter.EndObject();
 
-			if (maxCommitId != null)
+			if (requiredCommitId != null)
 			{
-				if (minCommitId == maxCommitId)
-				{
-					queryWriter.BeginObject("commit");
-					queryWriter.WriteString("$eq", maxCommitId.ToString());
-					queryWriter.EndObject();
-				}
-				else if (minCommitId != null)
-				{
-					queryWriter.BeginObject("commit");
-					queryWriter.WriteString("$lte", maxCommitId.ToString());
-					queryWriter.WriteString("$gte", minCommitId.ToString());
-					queryWriter.EndObject();
-				}
-				else
-				{
-					queryWriter.BeginObject("commit");
-					queryWriter.WriteString("$lte", maxCommitId.ToString());
-					queryWriter.EndObject();
-				}
+				queryWriter.BeginObject("commit");
+				queryWriter.WriteString("$eq", requiredCommitId.ToString());
+				queryWriter.EndObject();
 			}
 
 			if (name != null)
