Fix Linux Android SDK path references

#android

diff --git a/Engine/Source/Programs/AutomationTool/Android/AndroidPlatform.Automation.cs b/Engine/Source/Programs/AutomationTool/Android/AndroidPlatform.Automation.cs
--- a/Engine/Source/Programs/AutomationTool/Android/AndroidPlatform.Automation.cs
+++ b/Engine/Source/Programs/AutomationTool/Android/AndroidPlatform.Automation.cs
@@ -343,7 +343,7 @@ public class AndroidPlatform : Platform
 		if (OperatingSystem.IsLinux())
 		{
 			string UserHome = Environment.GetEnvironmentVariable("HOME");
-			string AndroidSdkPath = Path.Combine(UserHome, "Android", "sdk");
+			string AndroidSdkPath = Path.Combine(UserHome, "Android", "Sdk");
 
 			return AndroidSdkPath;
 		}
diff --git a/Engine/Source/Programs/UnrealBuildTool/Platform/Android/AndroidPlatformSDK.cs b/Engine/Source/Programs/UnrealBuildTool/Platform/Android/AndroidPlatformSDK.cs
--- a/Engine/Source/Programs/UnrealBuildTool/Platform/Android/AndroidPlatformSDK.cs
+++ b/Engine/Source/Programs/UnrealBuildTool/Platform/Android/AndroidPlatformSDK.cs
@@ -205,7 +205,7 @@ namespace UnrealBuildTool
 		{
 			if (OperatingSystem.IsLinux())
 			{
-				return Path.Combine(Environment.GetEnvironmentVariable("HOME")!, "Android/sdk");
+				return Path.Combine(Environment.GetEnvironmentVariable("HOME")!, "Android/Sdk");
 			}
 			else if (OperatingSystem.IsMacOS())
 			{
