Use passed test list as suite name

#rnx

diff --git a/Engine/Source/Programs/AutomationTool/Gauntlet/Unreal/Automation/UE.Automation.cs b/Engine/Source/Programs/AutomationTool/Gauntlet/Unreal/Automation/UE.Automation.cs
--- a/Engine/Source/Programs/AutomationTool/Gauntlet/Unreal/Automation/UE.Automation.cs
+++ b/Engine/Source/Programs/AutomationTool/Gauntlet/Unreal/Automation/UE.Automation.cs
@@ -505,13 +505,18 @@ namespace UE
 			get
 			{
 				string BaseSuite = base.Suite;
-				var Config = GetCachedConfiguration();
-				if (Config is AutomationTestConfig)
+				if (GetCachedConfiguration() is AutomationTestConfig AutomationConfig)
 				{
-					var AutomationConfig = Config as AutomationTestConfig;
-					if (!string.IsNullOrEmpty(AutomationConfig.RunTest) && AutomationConfig.RunTest.ToLower().StartsWith("group:"))
+					if (!string.IsNullOrEmpty(AutomationConfig.RunTest))
 					{
-						BaseSuite  = AutomationConfig.RunTest.Substring(6);
+						// try to get a usefull identification token
+						int StartPos = AutomationConfig.RunTest.ToLower().StartsWith("group:") ? 6 : AutomationConfig.RunTest.LastIndexOf('.') + 1;
+						BaseSuite = AutomationConfig.RunTest.Substring(StartPos);
+						if (string.IsNullOrEmpty(BaseSuite))
+						{
+							// get the token before if the first one is empty
+							BaseSuite = AutomationConfig.RunTest.Substring(0, StartPos - 1).Substring(AutomationConfig.RunTest.LastIndexOf('.', StartPos - 2) + 1);
+						}
 					}
 				}
 
