[Gauntlet] Adds configuration options to support reserving devices with tags from the Horde device reservation service.
Users who utilize the HordeDeviceReservationService can specify the tags in their test context's constraint, or by passing any of the following commandline arguments when running RunUnreal

-RequiredDeviceTags= - Only consider devices that have this tag for reservation
-PreferredDeviceTags=  - Prioritize devices that have this tag for reservation, but not a hard requirement
-UndesiredDeviceTags= -  De-prioritize devices that have this tag for reservation, but still use them as a last resort
-BlockedDeviceTags= - Do not consider devices that have this tag for reservation

Examples:
-RequiredDeviceTags="Required Tag 1, Required Tag 2"
-BlockedDeviceTags="Blocked Tag"

#tests self test agsinst dev server

diff --git a/Engine/Source/Programs/AutomationTool/AutomationUtils/DeviceReservation.cs b/Engine/Source/Programs/AutomationTool/AutomationUtils/DeviceReservation.cs
--- a/Engine/Source/Programs/AutomationTool/AutomationUtils/DeviceReservation.cs
+++ b/Engine/Source/Programs/AutomationTool/AutomationUtils/DeviceReservation.cs
@@ -261,6 +261,142 @@ namespace AutomationTool.DeviceReservation
 			public string JobId { get; set; }
 			public string StepId { get; set; }
 			public string DeviceName { get; set; }
+			public Tags Tags { get; set; }
+		}
+		
+		public class Tags
+		{
+			public enum Type
+			{
+				Required,
+				Preferred,
+				Undesired,
+				Blocked
+			}
+
+			/// <summary>
+			/// Adds a tag to the request
+			/// </summary>
+			/// <param name="Tag">The name of the requested tag</param>
+			/// <param name="TagType">How the tag request behaves</param>
+			/// <exception cref="NotImplementedException"></exception>
+			public void AddTag(string Tag, Type TagType)
+			{
+				switch (TagType)
+				{
+					case Type.Required:
+						Add(Tag, Required);
+						break;
+					case Type.Preferred:
+						Add(Tag, Preferred);
+						break;
+					case Type.Undesired:
+						Add(Tag, Undesired);
+						break;
+					case Type.Blocked:
+						Add(Tag, Blocked);
+						break;
+					default:
+						throw new NotImplementedException($"Missing AddTag case for {TagType.ToString()}");
+				}
+			}
+
+			/// <summary>
+			/// Adds multiple tags to the request
+			/// </summary>
+			/// <param name="Tags">The names of the requested tags</param>
+			/// <param name="TagType">How the tag request behaves</param>
+			/// <exception cref="NotImplementedException"></exception>
+			public void AddTags(IEnumerable<string> Tags, Type TagType)
+			{
+				foreach(string Tag in Tags)
+				{
+					AddTag(Tag, TagType);
+				}
+			}
+
+			/// <summary>
+			/// Required tags
+			/// </summary>
+			public IReadOnlyList<string> RequiredTags => Required;
+
+			/// <summary>
+			/// Preferred Tags
+			/// </summary>
+			public IReadOnlyList<string> PreferredTags => Preferred;
+
+			/// <summary>
+			/// Undesired Tags
+			/// </summary>
+			public IReadOnlyList<string> UndesiredTags => Undesired;
+
+			/// <summary>
+			/// Blocked Tags
+			/// </summary>
+			public IReadOnlyList<string> BlockedTags => Blocked;
+
+			/// <summary>
+			/// Devices that do not have this tag will not be reserved
+			/// </summary>
+			private List<string> Required { get; set; } = new List<string>();
+			
+			/// <summary>
+			/// Devices that have this tag will be selected over devices that do not, provided at least one is available
+			/// </summary>
+			private List<string> Preferred{ get; set; } = new List<string>();
+
+			/// <summary>
+			/// Devices that have this tag will be de-prioritized over devices that do not
+			/// </summary>
+			private List<string> Undesired { get; set; } = new List<string>();
+
+			/// <summary>
+			/// Devices that have this tag will not be reserved
+			/// </summary>
+			private List<string> Blocked { get; set; } = new List<string>();
+
+			public override string ToString()
+			{
+				return $"Required: {(RequiredTags == null ? string.Empty : string.Join(", ", RequiredTags))}\n" +
+				       $"Preferred: {(PreferredTags == null ? string.Empty : string.Join(", ", PreferredTags))}\n" +
+				       $"Undesired: {(UndesiredTags == null ? string.Empty : string.Join(", ", UndesiredTags))}\n" +
+				       $"Blocked: {(BlockedTags == null ? string.Empty : string.Join(", ", BlockedTags))}";
+			}
+			
+			public override bool Equals(object OtherObj)
+			{
+				if (ReferenceEquals(OtherObj, null))
+				{
+					return false;
+				}
+				
+				if (ReferenceEquals(this, OtherObj))
+				{
+					return true;
+				}
+				
+				return OtherObj is Tags Other && Other.ToString().Equals(ToString(), StringComparison.OrdinalIgnoreCase);
+			}
+
+			public override int GetHashCode()
+			{
+				return ToString().GetHashCode();
+			}
+
+			private void Add(string Tag, List<string> TagList)
+			{
+				if (!string.IsNullOrEmpty(Tag))
+				{
+					TagList.Add(FormatTag(Tag));
+				}
+			}
+
+			private string FormatTag(string Tag)
+			{
+				// Trim whitespaces and capitilaze the first letter of each word. E.g. foo bar => Foo Bar
+				string[] TagWords = Tag.Trim().Split(' ').Select(Word => string.Concat(Word[0].ToString().ToUpper(), Word.AsSpan(1))).ToArray();
+				return string.Join(' ', TagWords);
+			}
 		}
 
 		public sealed class DeviceErrorRequest
@@ -268,7 +404,7 @@ namespace AutomationTool.DeviceReservation
 			public string Message { get; set; }
 		}
 
-		public static Reservation Create(Uri BaseUri, string[] DeviceTypes, TimeSpan Duration, int RetryMax = 5, string PoolID = "", string DeviceName = "")
+		public static Reservation Create(Uri BaseUri, string[] DeviceTypes, TimeSpan Duration, int RetryMax = 5, string PoolID = "", string DeviceName = "", Tags Tags = null)
 		{
 			bool bFirst = true;
 			TimeSpan RetryTime = TimeSpan.FromMinutes(1);
@@ -298,7 +434,8 @@ namespace AutomationTool.DeviceReservation
 						PoolId = PoolID,
 						JobId = Environment.GetEnvironmentVariable("UE_HORDE_JOBID"),
 						StepId = Environment.GetEnvironmentVariable("UE_HORDE_STEPID"),
-						DeviceName = DeviceName
+						DeviceName = DeviceName,
+						Tags = Tags
 					});
 				}
 				catch (WebException WebEx)
@@ -414,6 +551,7 @@ namespace AutomationTool.DeviceReservation
 		public TimeSpan AvailableEndTime { get; set; }
 		public bool Enabled { get; set; }
 		public string DeviceData { get; set; }
+		public List<string> Tags { get; set; }
 
 		public Device Clone() { return (Device)MemberwiseClone(); }
 
diff --git a/Engine/Source/Programs/AutomationTool/Gauntlet/Framework/Devices/Gauntlet.DevicePool.cs b/Engine/Source/Programs/AutomationTool/Gauntlet/Framework/Devices/Gauntlet.DevicePool.cs
--- a/Engine/Source/Programs/AutomationTool/Gauntlet/Framework/Devices/Gauntlet.DevicePool.cs
+++ b/Engine/Source/Programs/AutomationTool/Gauntlet/Framework/Devices/Gauntlet.DevicePool.cs
@@ -14,6 +14,7 @@ using EpicGames.Core;
 using System.Data;
 using Gauntlet.Utils;
 using System.Reflection;
+using static AutomationTool.DeviceReservation.Reservation;
 
 namespace Gauntlet
 {
@@ -54,6 +55,8 @@ namespace Gauntlet
 
 		public bool RemoveOnShutdown { get; set; }
 
+		public List<string> Tags { get; set; } = new List<string>();
+
 		public override string ToString()
 		{
 			return string.Format("{0} @ {1}. Platform={2} Model={3}", Name, Address, Platform, string.IsNullOrEmpty(Model) ? "Unspecified" : Model);
@@ -69,13 +72,15 @@ namespace Gauntlet
 		public readonly EPerfSpec PerfSpec;
 		public readonly string Model;
 		public readonly string DeviceName;
+		public readonly Tags Tags;
 
-		public UnrealDeviceTargetConstraint(UnrealTargetPlatform? Platform, EPerfSpec PerfSpec = EPerfSpec.Unspecified, string Model = null, string DeviceName = null)
+		public UnrealDeviceTargetConstraint(UnrealTargetPlatform? Platform, EPerfSpec PerfSpec = EPerfSpec.Unspecified, string Model = null, string DeviceName = null, Tags Tags = null)
 		{
 			this.Platform = Platform;
 			this.PerfSpec = PerfSpec;
 			this.Model = Model == null ? string.Empty : Model;
 			this.DeviceName = DeviceName == null ? string.Empty : DeviceName;
+			this.Tags = Tags;
 		}
 
 		/// <summary>
@@ -128,7 +133,8 @@ namespace Gauntlet
 			return Other.Platform == Platform
 				&& Other.Model.Equals(Model, StringComparison.InvariantCultureIgnoreCase)
 				&& Other.DeviceName.Equals(DeviceName, StringComparison.InvariantCultureIgnoreCase)
-				&& Other.PerfSpec == PerfSpec;
+				&& Other.PerfSpec == PerfSpec
+				&& (Other.Tags == null && Tags == null) || (Other.Tags != null && Other.Tags.Equals(Tags));
 		}
 
 		public override bool Equals(object Obj)
diff --git a/Engine/Source/Programs/AutomationTool/Gauntlet/SelfTest/Devices/Gauntlet.SelfTest.TestTargetDevice.cs b/Engine/Source/Programs/AutomationTool/Gauntlet/SelfTest/Devices/Gauntlet.SelfTest.TestTargetDevice.cs
--- a/Engine/Source/Programs/AutomationTool/Gauntlet/SelfTest/Devices/Gauntlet.SelfTest.TestTargetDevice.cs
+++ b/Engine/Source/Programs/AutomationTool/Gauntlet/SelfTest/Devices/Gauntlet.SelfTest.TestTargetDevice.cs
@@ -1,13 +1,10 @@
 ﻿// Copyright Epic Games, Inc. All Rights Reserved.
 
-using Gauntlet;
+using AutomationTool.DeviceReservation;
+using AutomationTool;
 using System;
 using System.Collections.Generic;
-using System.Diagnostics;
-using System.IO;
 using System.Linq;
-using System.Text;
-using System.Threading.Tasks;
 using UnrealBuildTool;
 
 namespace Gauntlet.SelfTest
@@ -28,4 +25,263 @@ namespace Gauntlet.SelfTest
 			return TestFailures.Count == 0;
 		}
 	}
+
+	[TestGroup("DeviceReservation")]
+	public class TestDeviceReservation : BaseTestNode
+	{
+		private enum TestConstraintType
+		{
+			Name,
+			Tag,
+			Platform
+		}
+
+		/// <summary>
+		/// Comma separated list of platforms to make a reservation for
+		/// </summary>
+		private List<string> Platforms;
+
+		/// <summary>
+		/// Base device reservation URI
+		/// </summary>
+		private Uri DeviceServiceURL;
+
+		/// <summary>
+		/// Device pool to make reservations from
+		/// </summary>
+		private string PoolID = "RESERVATION-TEST";
+
+		/// <summary>
+		/// Whether or not devices have been marked as a problem
+		/// </summary>
+		private bool HasMarkedProblems;
+
+		/// <summary>
+		/// Time test will be marked as complete.
+		/// Used to stall the test so users can view the problem device
+		/// </summary>
+		private DateTime TargetTime;
+
+		/// <summary>
+		/// Internal list of constraints. Used for making reservations
+		/// </summary>
+		private List<UnrealDeviceTargetConstraint> Constraints;
+
+		/// <summary>
+		/// Internal list of active reservations
+		/// </summary>
+		private List<DeviceReservationAutoRenew> Reservations;
+
+		/// <summary>
+		/// Internal list of constraints and their intended test purpose
+		/// </summary>
+		private Dictionary<TestConstraintType, List<UnrealDeviceTargetConstraint>> ConstraintMap;
+
+		/// <summary>
+		/// Internal list of constraints mapped to their associated reservation
+		/// </summary>
+		private Dictionary<UnrealDeviceTargetConstraint, DeviceReservationAutoRenew> ReservationMap;
+
+		public TestDeviceReservation()
+		{
+			Constraints = new List<UnrealDeviceTargetConstraint>();
+			Reservations = new List<DeviceReservationAutoRenew>();
+			ConstraintMap = new Dictionary<TestConstraintType, List<UnrealDeviceTargetConstraint>>()
+			{
+				{TestConstraintType.Name,  new List<UnrealDeviceTargetConstraint>() },
+				{TestConstraintType.Tag,  new List<UnrealDeviceTargetConstraint>() },
+				{TestConstraintType.Platform,  new List<UnrealDeviceTargetConstraint>() },
+			};
+			ReservationMap = new Dictionary<UnrealDeviceTargetConstraint, DeviceReservationAutoRenew>();
+
+			Platforms = Gauntlet.Globals.Params.ParseValues("Platforms", true);
+			if (Platforms.Count == 0)
+			{
+				throw new AutomationException("One or more values for -Platforms= needs to be specified");
+			}
+
+			if (!Uri.TryCreate(Gauntlet.Globals.Params.ParseValue("DeviceServiceURL", string.Empty), UriKind.Absolute, out DeviceServiceURL))
+			{
+				throw new AutomationException("-DeviceServiceURL= needs to be specified and valid");
+			}
+
+			PoolID = Gauntlet.Globals.Params.ParseValue("PoolID", PoolID);
+		}
+
+		public override bool StartTest(int Pass, int NumPasses)
+		{
+			// For each platform requested, define 3 constraints
+			foreach (string PlatformString in Platforms)
+			{
+				UnrealTargetPlatform Platform = UnrealTargetPlatform.Parse(PlatformString);
+
+				// First constraint, reserve a device that has a name matching "<Platform>-TestDevice-Name"
+				string Name = PlatformString + "-TestDevice-Name";
+				UnrealDeviceTargetConstraint NameConstraint = new UnrealDeviceTargetConstraint(Platform, DeviceName: Name);
+
+				// Second constraint, reserve a device with "Required Tag" metadata
+				Reservation.Tags Tags = new Reservation.Tags();
+				Tags.AddTag("Required Tag", Reservation.Tags.Type.Required);
+
+				UnrealDeviceTargetConstraint TagConstraint = new UnrealDeviceTargetConstraint(Platform, Tags: Tags);
+
+				// Third constraint, just a default platform constraint. We'll also test failing this device
+				UnrealDeviceTargetConstraint PlatformConstraint = new UnrealDeviceTargetConstraint(Platform);
+
+				ConstraintMap[TestConstraintType.Name].Add(NameConstraint);
+				ConstraintMap[TestConstraintType.Tag].Add(TagConstraint);
+				ConstraintMap[TestConstraintType.Platform].Add(PlatformConstraint);
+				Constraints.AddRange([NameConstraint, TagConstraint, PlatformConstraint]);
+			}
+
+			// Make all the reservations
+			foreach (UnrealDeviceTargetConstraint Constraint in Constraints)
+			{
+				string[] FormattedConstraint = [Constraint.FormatWithIdentifier()];
+
+				DeviceReservationAutoRenew NewReservation = new DeviceReservationAutoRenew(
+					DeviceServiceURL.ToString(), 
+					Reservation.Create(DeviceServiceURL, FormattedConstraint, TimeSpan.FromMinutes(5), 5, PoolID, Constraint.DeviceName, Constraint.Tags));
+
+				Reservations.Add(NewReservation);
+				ReservationMap.Add(Constraint, NewReservation);
+			}
+
+			// Ensure each reservation succeeded and returned the expected device type
+			Dictionary<UnrealDeviceTargetConstraint, string> FailedNames = AssertReservedDeviceStatus(TestConstraintType.Name);
+			Dictionary<UnrealDeviceTargetConstraint, string> FailedTags = AssertReservedDeviceStatus(TestConstraintType.Tag);
+			Dictionary<UnrealDeviceTargetConstraint, string> FailedPlatforms = AssertReservedDeviceStatus(TestConstraintType.Platform);
+
+			bool bAnyFailures = FailedNames.Any() || FailedTags.Any() || FailedPlatforms.Any();
+			if (bAnyFailures)
+			{
+				LogFailures(TestConstraintType.Name, FailedNames);
+				LogFailures(TestConstraintType.Tag, FailedTags);
+				LogFailures(TestConstraintType.Platform, FailedPlatforms);
+			}
+
+			TargetTime = DateTime.Now + TimeSpan.FromMinutes(1);
+
+			return !bAnyFailures && base.StartTest(Pass, NumPasses);
+		}
+
+		public override void TickTest()
+		{
+			if (!HasMarkedProblems)
+			{
+				HasMarkedProblems = true;
+
+				// Report an erorr for each "Platform" device
+				foreach (UnrealDeviceTargetConstraint PlatformConstraint in ConstraintMap[TestConstraintType.Platform])
+				{
+					Device Device = ReservationMap[PlatformConstraint].Devices[0];
+					Reservation.ReportDeviceError(DeviceServiceURL.ToString(), Device.Name, "Marking device as problem for self test");
+					ReservationMap[PlatformConstraint].Dispose();
+				}
+			}
+
+			if (DateTime.Now > TargetTime)
+			{
+				MarkComplete();
+			}
+		}
+
+		public override void StopTest(StopReason InReason)
+		{
+			foreach (DeviceReservationAutoRenew Entry in Reservations)
+			{
+				Entry.Dispose();
+			}
+
+			Reservations.Clear();
+			ReservationMap.Clear();
+
+			base.StopTest(InReason);
+		}
+
+		private Dictionary<UnrealDeviceTargetConstraint, string> AssertReservedDeviceStatus(TestConstraintType Type)
+		{
+			Dictionary<UnrealDeviceTargetConstraint, string> FailedReservations = new Dictionary<UnrealDeviceTargetConstraint, string>();
+
+			foreach (UnrealDeviceTargetConstraint Constraint in ConstraintMap[Type])
+			{
+				DeviceReservationAutoRenew Reservation = ReservationMap[Constraint];
+
+				if (Reservation.Devices.Count == 0)
+				{
+					FailedReservations.Add(Constraint, "No devices were reserved");
+					continue;
+				}
+
+				if (Reservation.Devices.Count > 1)
+				{
+					FailedReservations.Add(Constraint, "More than one device was reserved");
+					continue;
+				}
+
+				Device Device = Reservation.Devices[0];
+
+				switch (Type)
+				{
+					case TestConstraintType.Name:
+						string ExpectedDeviceName = Constraint.Platform.ToString() + "-TestDevice-Name";
+						string ActualDeviceName = Device.Name;
+						if (!ActualDeviceName.Equals(ExpectedDeviceName, StringComparison.OrdinalIgnoreCase))
+						{
+							FailedReservations.Add(Constraint, $"Device name mismatch. Expected : {ExpectedDeviceName} | Actual: {ActualDeviceName}");
+							continue;
+						}
+						break;
+
+					case TestConstraintType.Tag:
+						string ExpectedDeviceTags = "Required Tag";
+						if (Device.Tags == null || Device.Tags.Count == 0)
+						{
+							string ActualTags = Device.Tags == null ? "<Null>" : "<Empty>";
+							FailedReservations.Add(Constraint, $"Device missing tags. Expected : {ExpectedDeviceTags} | Actual: {ActualTags}");
+						}
+
+						if (Device.Tags.Count > 1)
+						{
+							FailedReservations.Add(Constraint, $"Device unexpected tags. Expected : {ExpectedDeviceTags} | Actual: {string.Join(", ", Device.Tags)}");
+						}
+
+						string ActualDeviceTags = Device.Tags[0];
+						if (!ActualDeviceTags.Equals(ExpectedDeviceTags, StringComparison.OrdinalIgnoreCase))
+						{
+							FailedReservations.Add(Constraint, $"Device tag mismatch. Expected : {ExpectedDeviceTags} | Actual {ActualDeviceTags}");
+							continue;
+						}
+						break;
+
+					case TestConstraintType.Platform:
+						UnrealTargetPlatform ExpectedPlatform = Constraint.Platform!.Value;
+						UnrealTargetPlatform ActualPlatform = UnrealTargetPlatform.Parse(Device.Type.Replace("-DevKit", "", StringComparison.OrdinalIgnoreCase));
+						if (ExpectedPlatform != ActualPlatform)
+						{
+							FailedReservations.Add(Constraint, $"Platform mismatch. Expected : {ExpectedPlatform} | Actual {ActualPlatform}");
+							continue;
+						}
+						break;
+				}
+			}
+
+			return FailedReservations;
+		}
+
+		private void LogFailures(TestConstraintType Type, Dictionary<UnrealDeviceTargetConstraint, string> Failures)
+		{
+			if (!Failures.Any())
+			{
+				return;
+			}
+
+			string Error = $"{Type} Failures ({Failures.Count})";
+			foreach (var Failure in Failures)
+			{
+				Error += $"\n\t{Failure.Key.ToString()} | {Failure.Value}";
+			}
+			Log.Error(Error);
+		}
+	}
 }
diff --git a/Engine/Source/Programs/AutomationTool/Gauntlet/Unreal/RunUnreal.cs b/Engine/Source/Programs/AutomationTool/Gauntlet/Unreal/RunUnreal.cs
--- a/Engine/Source/Programs/AutomationTool/Gauntlet/Unreal/RunUnreal.cs
+++ b/Engine/Source/Programs/AutomationTool/Gauntlet/Unreal/RunUnreal.cs
@@ -432,6 +432,12 @@ namespace Gauntlet
 				List<string> DeviceNameArgs = Globals.Params.ParseValues("DeviceName", false);
 				string DeviceName = DeviceNameArgs.Count > 0 ? DeviceNameArgs.Last() : string.Empty;
 
+				Reservation.Tags DeviceTags = new Reservation.Tags();
+				DeviceTags.AddTags(Globals.Params.ParseValues("RequiredDeviceTags", true), Reservation.Tags.Type.Required);
+				DeviceTags.AddTags(Globals.Params.ParseValues("PreferredDeviceTags", true), Reservation.Tags.Type.Preferred);
+				DeviceTags.AddTags(Globals.Params.ParseValues("UndesiredDeviceTags", true), Reservation.Tags.Type.Undesired);
+				DeviceTags.AddTags(Globals.Params.ParseValues("BlockedDeviceTags", true), Reservation.Tags.Type.Blocked);
+
 				TestContext.Constraint = new UnrealDeviceTargetConstraint(UnrealPlatform, PerfSpec, Model, DeviceName);
 
 				// parse worker job id
