Cloud Artifacts - Specify a manifest of files to upload to zen, thus respecting the Files fileset in the task

diff --git a/Engine/Source/Programs/AutomationTool/BuildGraph/Tasks/CreateCloudArtifactTask.cs b/Engine/Source/Programs/AutomationTool/BuildGraph/Tasks/CreateCloudArtifactTask.cs
--- a/Engine/Source/Programs/AutomationTool/BuildGraph/Tasks/CreateCloudArtifactTask.cs
+++ b/Engine/Source/Programs/AutomationTool/BuildGraph/Tasks/CreateCloudArtifactTask.cs
@@ -452,26 +452,33 @@ namespace AutomationTool.Tasks
 			DirectoryReference baseDir = ResolveDirectory(_parameters.BaseDir);
 			List<FileReference> files = BgTaskImpl.ResolveFilespec(baseDir, _parameters.Files, tagNameToFileSet).ToList();
 
-			bool validFiles = files.Count != 0;
-			if (files.Count != 0)
+			FileReference manifestFile = FileReference.Combine(baseDir, "zen-manifest.txt");
 			{
-				foreach (FileReference file in files)
+				await using StreamWriter streamWriter = new StreamWriter(manifestFile.FullName);
+				bool validFiles = files.Count != 0;
+				if (files.Count != 0)
 				{
-					if (!file.IsUnderDirectory(baseDir))
+					foreach (FileReference file in files)
 					{
-						Logger.LogError("Artifact file {File} is not under {BaseDir}", file, baseDir);
-						validFiles = false;
+						if (!file.IsUnderDirectory(baseDir))
+						{
+							Logger.LogError("Artifact file {File} is not under {BaseDir}", file, baseDir);
+							validFiles = false;
+						}
+
+						// always use / as path separation in the manifest no matter the host platform
+						await streamWriter.WriteLineAsync(file.MakeRelativeTo(baseDir).Replace('\\', '/'));
 					}
 				}
-			}
-			else
-			{
-				Logger.LogWarning("No files found under {BaseDir}", baseDir);
-			}
+				else
+				{
+					Logger.LogWarning("No files found under {BaseDir}", baseDir);
+				}
 
-			if (!validFiles)
-			{
-				throw new AutomationException($"Unable to create artifact {name} with given file list.");
+				if (!validFiles)
+				{
+					throw new AutomationException($"Unable to create artifact {name} with given file list.");
+				}
 			}
 
 			string bucketId = $"{project}.{type}.{streamId}.{platform}";
@@ -495,7 +502,7 @@ namespace AutomationTool.Tasks
 			// set part name to default by default - should be exposed as a option to allow for uploading separate parts
 			string partName = "default";
 
-			BuildUploader.UploadBuild(host, baseDir, ns, bucketId, buildId, partName, partId, metadata, accessToken, new BuildUploader.BuildInfo
+			BuildUploader.UploadBuild(host, baseDir, ns, bucketId, buildId, partName, partId, manifestFile, metadata, accessToken, new BuildUploader.BuildInfo
 			{
 				Name = name.ToString(),
 				ProjectName = project,
@@ -613,7 +620,7 @@ namespace AutomationTool.Tasks
 			public string BaselineBranch { get; set; }
 		}
 		public static void UploadBuild(string host, DirectoryReference baseDir, string namespaceId, string bucketId, CbObjectId buildId,
-			string partName, CbObjectId partId, Dictionary<string, object> metadata, string accessToken, BuildInfo info,
+			string partName, CbObjectId partId, FileReference manifestFile, Dictionary<string, object> metadata, string accessToken, BuildInfo info,
 			bool allowMultipart = false, bool assumeHttp2 = false, bool forceClean = false, bool verify = false, bool verbose = false,
 			bool boostWorkers = false, FileReference? traceFile = null)
 		{
@@ -660,7 +667,7 @@ namespace AutomationTool.Tasks
 				// trace file option needs to be passed before other arguments including the verbs
 				string traceOption = traceFile != null ? $"--tracefile={traceFile.FullName} " : "";
 				// pass the access token via environment variable to avoid showing it on the cli
-				string cmdline = $"{traceOption}builds upload --url {host} --namespace {namespaceId} --bucket {bucketId} --build-id {buildId} --build-part-name \"{partName}\" --build-part-id {partId} --metadata-path \"{metadataJsonPath}\" --access-token-env UE-CloudDataCacheAccessToken --create-build --plain-progress --allow-multipart={allowMultipart} {http2Options} {boostWorkersOption} --clean={forceClean} {verifyOption} {verboseOption} \"{baseDir}\"";
+				string cmdline = $"{traceOption}builds upload --url {host} --namespace {namespaceId} --bucket {bucketId} --build-id {buildId} --build-part-name \"{partName}\" --build-part-id {partId} --metadata-path \"{metadataJsonPath}\" --manifest-path \"{manifestFile}\" --access-token-env UE-CloudDataCacheAccessToken --create-build --plain-progress --allow-multipart={allowMultipart} {http2Options} {boostWorkersOption} --clean={forceClean} {verifyOption} {verboseOption} \"{baseDir}\"";
 				Log.Logger.LogInformation("Using metadata '{Metadata}'", metadataString);
 				CommandUtils.RunAndLog(CommandUtils.CmdEnv, zenExe.FullName,cmdline , Options: CommandUtils.ERunOptions.Default, EnvVars: new Dictionary<string, string> {{"UE-CloudDataCacheAccessToken", accessToken}});
 			}
