[QA Automation Tools] Implement TCP like guaranteed delivery for RPC calls

diff --git a/Engine/Source/Programs/AutomationTool/Gauntlet/Framework/RpcFramework/Gauntlet.RpcExecutor.cs b/Engine/Source/Programs/AutomationTool/Gauntlet/Framework/RpcFramework/Gauntlet.RpcExecutor.cs
--- a/Engine/Source/Programs/AutomationTool/Gauntlet/Framework/RpcFramework/Gauntlet.RpcExecutor.cs
+++ b/Engine/Source/Programs/AutomationTool/Gauntlet/Framework/RpcFramework/Gauntlet.RpcExecutor.cs
@@ -770,7 +770,7 @@ namespace Gauntlet
 			HttpRequestMessage RpcRequest = new HttpRequestMessage(MethodToUse, RpcUrl);
 			RpcRequest.Headers.Add("rpcname", RpcName);
 			RpcRequest.Headers.Connection.Add("keep-alive");
-			RpcRequest.Headers.Add("Keep-Alive", "600");
+			RpcRequest.Headers.Add("Keep-Alive", $"{Timeout}");
 			if (RpcEntry.Args != null && RpcEntry.Args.Count > 0)
 			{
 				foreach (RpcArgumentDesc ArgDesc in RpcEntry.Args)
