[Gauntlet] Minor fixes for device tag reservation support
#tests Gauntlet self test

diff --git a/Engine/Source/Programs/AutomationTool/AutomationUtils/DeviceReservation.cs b/Engine/Source/Programs/AutomationTool/AutomationUtils/DeviceReservation.cs
--- a/Engine/Source/Programs/AutomationTool/AutomationUtils/DeviceReservation.cs
+++ b/Engine/Source/Programs/AutomationTool/AutomationUtils/DeviceReservation.cs
@@ -263,7 +263,7 @@ namespace AutomationTool.DeviceReservation
 			public string DeviceName { get; set; }
 			public Tags Tags { get; set; }
 		}
-		
+
 		public class Tags
 		{
 			public enum Type
@@ -339,7 +339,7 @@ namespace AutomationTool.DeviceReservation
 			/// Devices that do not have this tag will not be reserved
 			/// </summary>
 			private List<string> Required { get; set; } = new List<string>();
-			
+
 			/// <summary>
 			/// Devices that have this tag will be selected over devices that do not, provided at least one is available
 			/// </summary>
@@ -362,19 +362,19 @@ namespace AutomationTool.DeviceReservation
 				       $"Undesired: {(UndesiredTags == null ? string.Empty : string.Join(", ", UndesiredTags))}\n" +
 				       $"Blocked: {(BlockedTags == null ? string.Empty : string.Join(", ", BlockedTags))}";
 			}
-			
+
 			public override bool Equals(object OtherObj)
 			{
 				if (ReferenceEquals(OtherObj, null))
 				{
 					return false;
 				}
-				
+
 				if (ReferenceEquals(this, OtherObj))
 				{
 					return true;
 				}
-				
+
 				return OtherObj is Tags Other && Other.ToString().Equals(ToString(), StringComparison.OrdinalIgnoreCase);
 			}
 
@@ -385,7 +385,7 @@ namespace AutomationTool.DeviceReservation
 
 			private void Add(string Tag, List<string> TagList)
 			{
-				if (!string.IsNullOrEmpty(Tag))
+				if (!string.IsNullOrWhiteSpace(Tag))
 				{
 					TagList.Add(FormatTag(Tag));
 				}
diff --git a/Engine/Source/Programs/AutomationTool/Gauntlet/Framework/Devices/Gauntlet.DevicePool.cs b/Engine/Source/Programs/AutomationTool/Gauntlet/Framework/Devices/Gauntlet.DevicePool.cs
--- a/Engine/Source/Programs/AutomationTool/Gauntlet/Framework/Devices/Gauntlet.DevicePool.cs
+++ b/Engine/Source/Programs/AutomationTool/Gauntlet/Framework/Devices/Gauntlet.DevicePool.cs
@@ -134,7 +134,7 @@ namespace Gauntlet
 				&& Other.Model.Equals(Model, StringComparison.InvariantCultureIgnoreCase)
 				&& Other.DeviceName.Equals(DeviceName, StringComparison.InvariantCultureIgnoreCase)
 				&& Other.PerfSpec == PerfSpec
-				&& (Other.Tags == null && Tags == null) || (Other.Tags != null && Other.Tags.Equals(Tags));
+				&& ((Other.Tags == null && Tags == null) || (Other.Tags != null && Other.Tags.Equals(Tags)));
 		}
 
 		public override bool Equals(object Obj)
@@ -187,6 +187,10 @@ namespace Gauntlet
 			{
 				Value = string.Format("{0}:{1}", Value, DeviceName);
 			}
+			if (Tags != null)
+			{
+				Value = string.Format("{0}:{1}", Value, Tags.ToString());
+			}
 			return Value;
 		}
 
diff --git a/Engine/Source/Programs/AutomationTool/Gauntlet/SelfTest/Devices/Gauntlet.SelfTest.TestTargetDevice.cs b/Engine/Source/Programs/AutomationTool/Gauntlet/SelfTest/Devices/Gauntlet.SelfTest.TestTargetDevice.cs
--- a/Engine/Source/Programs/AutomationTool/Gauntlet/SelfTest/Devices/Gauntlet.SelfTest.TestTargetDevice.cs
+++ b/Engine/Source/Programs/AutomationTool/Gauntlet/SelfTest/Devices/Gauntlet.SelfTest.TestTargetDevice.cs
@@ -18,7 +18,7 @@ namespace Gauntlet.SelfTest
 
 			// Device should reboot (or pretend it did)
 			CheckResult(TestDevice.Reboot(), "Failed to reboot device");
-		
+
 			// Device should connect
 			CheckResult(TestDevice.Connect() && TestDevice.IsConnected, "Failed to connect to device");
 
@@ -239,11 +239,13 @@ namespace Gauntlet.SelfTest
 						{
 							string ActualTags = Device.Tags == null ? "<Null>" : "<Empty>";
 							FailedReservations.Add(Constraint, $"Device missing tags. Expected : {ExpectedDeviceTags} | Actual: {ActualTags}");
+							continue;
 						}
 
 						if (Device.Tags.Count > 1)
 						{
 							FailedReservations.Add(Constraint, $"Device unexpected tags. Expected : {ExpectedDeviceTags} | Actual: {string.Join(", ", Device.Tags)}");
+							continue;
 						}
 
 						string ActualDeviceTags = Device.Tags[0];
diff --git a/Engine/Source/Programs/AutomationTool/Gauntlet/Unreal/RunUnreal.cs b/Engine/Source/Programs/AutomationTool/Gauntlet/Unreal/RunUnreal.cs
--- a/Engine/Source/Programs/AutomationTool/Gauntlet/Unreal/RunUnreal.cs
+++ b/Engine/Source/Programs/AutomationTool/Gauntlet/Unreal/RunUnreal.cs
@@ -438,7 +438,7 @@ namespace Gauntlet
 				DeviceTags.AddTags(Globals.Params.ParseValues("UndesiredDeviceTags", true), Reservation.Tags.Type.Undesired);
 				DeviceTags.AddTags(Globals.Params.ParseValues("BlockedDeviceTags", true), Reservation.Tags.Type.Blocked);
 
-				TestContext.Constraint = new UnrealDeviceTargetConstraint(UnrealPlatform, PerfSpec, Model, DeviceName);
+				TestContext.Constraint = new UnrealDeviceTargetConstraint(UnrealPlatform, PerfSpec, Model, DeviceName, DeviceTags);
 
 				// parse worker job id
 				List<string> WorkerJobIDArgs = CombinedParams.ParseValues("WorkerJobID", false);
