[Gauntlet] Consider tags when determining if an UnrealDeviceTargetConstraint is an identity/match

#tests Device reservation self test

diff --git a/Engine/Source/Programs/AutomationTool/AutomationUtils/DeviceReservation.cs b/Engine/Source/Programs/AutomationTool/AutomationUtils/DeviceReservation.cs
--- a/Engine/Source/Programs/AutomationTool/AutomationUtils/DeviceReservation.cs
+++ b/Engine/Source/Programs/AutomationTool/AutomationUtils/DeviceReservation.cs
@@ -383,6 +383,15 @@ namespace AutomationTool.DeviceReservation
 				return ToString().GetHashCode();
 			}
 
+			/// <summary>
+			/// Does this Tag request have any tags associated with it
+			/// </summary>
+			/// <returns></returns>
+			public bool Any()
+			{
+				return RequiredTags?.Count > 0 || PreferredTags?.Count > 0 || Undesired?.Count > 0 || BlockedTags?.Count > 0;
+			}
+
 			private void Add(string Tag, List<string> TagList)
 			{
 				if (!string.IsNullOrWhiteSpace(Tag))
diff --git a/Engine/Source/Programs/AutomationTool/Gauntlet/Framework/Devices/Gauntlet.DevicePool.cs b/Engine/Source/Programs/AutomationTool/Gauntlet/Framework/Devices/Gauntlet.DevicePool.cs
--- a/Engine/Source/Programs/AutomationTool/Gauntlet/Framework/Devices/Gauntlet.DevicePool.cs
+++ b/Engine/Source/Programs/AutomationTool/Gauntlet/Framework/Devices/Gauntlet.DevicePool.cs
@@ -88,7 +88,10 @@ namespace Gauntlet
 		/// </summary>
 		public bool IsIdentity()
 		{
-			return (PerfSpec == EPerfSpec.Unspecified) && (Model == string.Empty) && (DeviceName == string.Empty);
+			return (PerfSpec == EPerfSpec.Unspecified)
+				&& (Model == string.Empty)
+				&& (DeviceName == string.Empty)
+				&& (Tags == null || !Tags.Any());
 		}
 
 
@@ -114,8 +117,22 @@ namespace Gauntlet
 
 			bool ModelMatch = Model == string.Empty ? true : Model.Equals(DeviceDef.Model, StringComparison.InvariantCultureIgnoreCase);
 			bool PerfMatch = (PerfSpec == EPerfSpec.Unspecified) ? true : PerfSpec == DeviceDef.PerfSpec;
-			bool NameMatch = DeviceName == string.Empty? true : DeviceName.Equals(DeviceDef.Name, StringComparison.InvariantCultureIgnoreCase);
-			return ModelMatch && PerfMatch && NameMatch;
+			bool NameMatch = DeviceName == string.Empty ? true : DeviceName.Equals(DeviceDef.Name, StringComparison.InvariantCultureIgnoreCase);
+
+			// Tag matches if null or has no required/blocked tags
+			bool TagMatch = Tags == null
+				|| ((Tags.RequiredTags == null || Tags.RequiredTags.Count == 0)
+				&& (Tags.BlockedTags == null || Tags.BlockedTags.Count == 0));
+
+			if (!TagMatch)
+			{
+				// If there are tags, make sure it has the required tags but none of the blocked tags
+				bool bHasRequiredTags = Tags.RequiredTags == null || Tags.RequiredTags.All(Tag => DeviceDef.Tags.Contains(Tag, StringComparer.OrdinalIgnoreCase));
+				bool bHasNoBlockedTags = Tags.BlockedTags == null || !DeviceDef.Tags.Any(Tag => Tags.BlockedTags.Contains(Tag, StringComparer.OrdinalIgnoreCase));
+				TagMatch = bHasRequiredTags && bHasNoBlockedTags;
+			}
+
+			return ModelMatch && PerfMatch && NameMatch && TagMatch;
 		}
 
 		public bool Equals(UnrealDeviceTargetConstraint Other)
diff --git a/Engine/Source/Programs/AutomationTool/Gauntlet/Framework/Devices/Gauntlet.DeviceReservationService.cs b/Engine/Source/Programs/AutomationTool/Gauntlet/Framework/Devices/Gauntlet.DeviceReservationService.cs
--- a/Engine/Source/Programs/AutomationTool/Gauntlet/Framework/Devices/Gauntlet.DeviceReservationService.cs
+++ b/Engine/Source/Programs/AutomationTool/Gauntlet/Framework/Devices/Gauntlet.DeviceReservationService.cs
@@ -282,7 +282,8 @@ namespace Gauntlet
 					Name = Device.Name,
 					Platform = UnrealTargetPlatform.Parse(UnrealTargetPlatform.GetValidPlatformNames().FirstOrDefault(Entry => Entry == Device.Type.Replace("-DevKit", "", StringComparison.OrdinalIgnoreCase))),
 					DeviceData = Device.DeviceData,
-					Model = Device.Model
+					Model = Device.Model,
+					Tags = Device.Tags
 				};
 
 				EPerfSpec PerfSpec = EPerfSpec.Unspecified;
diff --git a/Engine/Source/Programs/AutomationTool/Gauntlet/SelfTest/Devices/Gauntlet.SelfTest.TestTargetDevice.cs b/Engine/Source/Programs/AutomationTool/Gauntlet/SelfTest/Devices/Gauntlet.SelfTest.TestTargetDevice.cs
--- a/Engine/Source/Programs/AutomationTool/Gauntlet/SelfTest/Devices/Gauntlet.SelfTest.TestTargetDevice.cs
+++ b/Engine/Source/Programs/AutomationTool/Gauntlet/SelfTest/Devices/Gauntlet.SelfTest.TestTargetDevice.cs
@@ -220,12 +220,26 @@ namespace Gauntlet.SelfTest
 				}
 
 				Device Device = Reservation.Devices[0];
+				DeviceDefinition DeviceDefinition = new DeviceDefinition()
+				{
+					Address = Device.IPOrHostName,
+					Name = Device.Name,
+					Platform = UnrealTargetPlatform.Parse(UnrealTargetPlatform.GetValidPlatformNames().FirstOrDefault(Entry => Entry == Device.Type.Replace("-DevKit", "", StringComparison.OrdinalIgnoreCase))),
+					DeviceData = Device.DeviceData,
+					Model = Device.Model,
+					Tags = Device.Tags
+				};
+
+				if (!Constraint.Check(DeviceDefinition))
+				{
+					Log.Error("{Constraint} failed check", Constraint.ToString());
+				}
 
 				switch (Type)
 				{
 					case TestConstraintType.Name:
 						string ExpectedDeviceName = Constraint.Platform.ToString() + "-TestDevice-Name";
-						string ActualDeviceName = Device.Name;
+						string ActualDeviceName = DeviceDefinition.Name;
 						if (!ActualDeviceName.Equals(ExpectedDeviceName, StringComparison.OrdinalIgnoreCase))
 						{
 							FailedReservations.Add(Constraint, $"Device name mismatch. Expected : {ExpectedDeviceName} | Actual: {ActualDeviceName}");
@@ -235,20 +249,20 @@ namespace Gauntlet.SelfTest
 
 					case TestConstraintType.Tag:
 						string ExpectedDeviceTags = "Required Tag";
-						if (Device.Tags == null || Device.Tags.Count == 0)
+						if (DeviceDefinition.Tags == null || DeviceDefinition.Tags.Count == 0)
 						{
-							string ActualTags = Device.Tags == null ? "<Null>" : "<Empty>";
+							string ActualTags = DeviceDefinition.Tags == null ? "<Null>" : "<Empty>";
 							FailedReservations.Add(Constraint, $"Device missing tags. Expected : {ExpectedDeviceTags} | Actual: {ActualTags}");
 							continue;
 						}
 
-						if (Device.Tags.Count > 1)
+						if (DeviceDefinition.Tags.Count > 1)
 						{
-							FailedReservations.Add(Constraint, $"Device unexpected tags. Expected : {ExpectedDeviceTags} | Actual: {string.Join(", ", Device.Tags)}");
+							FailedReservations.Add(Constraint, $"Device unexpected tags. Expected : {ExpectedDeviceTags} | Actual: {string.Join(", ", DeviceDefinition.Tags)}");
 							continue;
 						}
 
-						string ActualDeviceTags = Device.Tags[0];
+						string ActualDeviceTags = DeviceDefinition.Tags[0];
 						if (!ActualDeviceTags.Equals(ExpectedDeviceTags, StringComparison.OrdinalIgnoreCase))
 						{
 							FailedReservations.Add(Constraint, $"Device tag mismatch. Expected : {ExpectedDeviceTags} | Actual {ActualDeviceTags}");
@@ -258,7 +272,7 @@ namespace Gauntlet.SelfTest
 
 					case TestConstraintType.Platform:
 						UnrealTargetPlatform ExpectedPlatform = Constraint.Platform!.Value;
-						UnrealTargetPlatform ActualPlatform = UnrealTargetPlatform.Parse(Device.Type.Replace("-DevKit", "", StringComparison.OrdinalIgnoreCase));
+						UnrealTargetPlatform ActualPlatform = DeviceDefinition.Platform!.Value;
 						if (ExpectedPlatform != ActualPlatform)
 						{
 							FailedReservations.Add(Constraint, $"Platform mismatch. Expected : {ExpectedPlatform} | Actual {ActualPlatform}");
