Added Cosmetic Validation uploads for the Island Regression test

#rnx

diff --git a/Engine/Source/Programs/AutomationTool/Gauntlet/Framework/RpcFramework/Gauntlet.RpcLibrary.cs b/Engine/Source/Programs/AutomationTool/Gauntlet/Framework/RpcFramework/Gauntlet.RpcLibrary.cs
--- a/Engine/Source/Programs/AutomationTool/Gauntlet/Framework/RpcFramework/Gauntlet.RpcLibrary.cs
+++ b/Engine/Source/Programs/AutomationTool/Gauntlet/Framework/RpcFramework/Gauntlet.RpcLibrary.cs
@@ -676,6 +676,21 @@ namespace Gauntlet
 			}
 		}
 
+		public static bool EnableExposureCapture(this RpcTarget RpcTarget)
+		{
+			const string RPCName = "EnableExposureCapture";
+			using (HttpResponseMessage RpcResponse = RpcExecutor.CallRpc(RpcTarget, RPCName))
+			{
+				RpcLibrary.SimpleResponse Response = RpcLibrary.GetSimpleResponse(RpcResponse);
+				if (!Response.Succeeded)
+				{
+					throw new TestException($"RPC '{RPCName}' failed: {Response.Value}");
+				}
+
+				return true;
+			}
+		}
+
 		public static bool AutomationControllerClearReports(this RpcTarget InTarget)
 		{
 			const string RPCName = "AutomationControllerClearReports";
diff --git a/Engine/Source/Runtime/Engine/Private/SceneView.cpp b/Engine/Source/Runtime/Engine/Private/SceneView.cpp
--- a/Engine/Source/Runtime/Engine/Private/SceneView.cpp
+++ b/Engine/Source/Runtime/Engine/Private/SceneView.cpp
@@ -2886,6 +2886,16 @@ float FSceneView::GetLastEyeAdaptationExposure() const
 	return 0.0f;
 }
 
+/** Returns the eye adaptation luminance or 0.0f if it doesn't exist. */
+float FSceneView::GetLastAverageSceneLuminance() const
+{
+	if (EyeAdaptationViewState)
+	{
+		return EyeAdaptationViewState->GetLastAverageSceneLuminance();
+	}
+	return 0.0f;
+}
+
 const FSceneView* FSceneView::GetPrimarySceneView() const
 {
 	// It is valid for this function to return itself if it's already the primary view.
diff --git a/Engine/Source/Runtime/Engine/Public/SceneManagement.h b/Engine/Source/Runtime/Engine/Public/SceneManagement.h
--- a/Engine/Source/Runtime/Engine/Public/SceneManagement.h
+++ b/Engine/Source/Runtime/Engine/Public/SceneManagement.h
@@ -206,6 +206,9 @@ public:
 	/** Returns the eye adaptation exposure. */
 	virtual float GetLastEyeAdaptationExposure() const = 0;
 
+	/** Returns the average scene luminance. */
+	virtual float GetLastAverageSceneLuminance() const = 0;
+
 	virtual void SetExposureCompensationCurve(class UCurveFloat* NewCurve) {}
 
 	virtual void SetSequencerState(ESequencerState InSequencerState) = 0;
diff --git a/Engine/Source/Runtime/Engine/Public/SceneView.h b/Engine/Source/Runtime/Engine/Public/SceneView.h
--- a/Engine/Source/Runtime/Engine/Public/SceneView.h
+++ b/Engine/Source/Runtime/Engine/Public/SceneView.h
@@ -2036,6 +2036,9 @@ public:
 	/** Returns the eye adaptation exposure or 0.0f if it doesn't exist. */
 	ENGINE_API float GetLastEyeAdaptationExposure() const;
 
+	/** Returns the average luminance or 0.0f if it doesn't exist. */
+	ENGINE_API float GetLastAverageSceneLuminance() const;
+
 	/** Get the primary view associated with the secondary view. */
 	ENGINE_API const FSceneView* GetPrimarySceneView() const;
 
