Decreasing severity of message when not finding image files to convet or resize. This will prevent warnings in tests where no image is generated
#rnx

diff --git a/Engine/Source/Programs/AutomationTool/Gauntlet/Framework/Gauntlet.Utils.cs b/Engine/Source/Programs/AutomationTool/Gauntlet/Framework/Gauntlet.Utils.cs
--- a/Engine/Source/Programs/AutomationTool/Gauntlet/Framework/Gauntlet.Utils.cs
+++ b/Engine/Source/Programs/AutomationTool/Gauntlet/Framework/Gauntlet.Utils.cs
@@ -1729,7 +1729,7 @@ namespace Gauntlet
 
 				if (Files.Count() == 0)
 				{
-					Log.Warning("Could not find files at {0} to resize", InDirectory);
+					Log.Verbose("Could not find files at {0} to resize", InDirectory);
 					return false;
 				}
 
@@ -1765,7 +1765,7 @@ namespace Gauntlet
 
 				if (Files.Count() == 0)
 				{
-					Log.Warning("Could not find files at {0} to resize", InDirectory);
+					Log.Verbose("Could not find files at {0} to convert", InDirectory);
 					return false;
 				}
 
