Initialization failure is to be raised as phase interrupted outcome

#rnx

diff --git a/Engine/Source/Programs/AutomationTool/Gauntlet/Unreal/Base/Gauntlet.UnrealTestNode.cs b/Engine/Source/Programs/AutomationTool/Gauntlet/Unreal/Base/Gauntlet.UnrealTestNode.cs
--- a/Engine/Source/Programs/AutomationTool/Gauntlet/Unreal/Base/Gauntlet.UnrealTestNode.cs
+++ b/Engine/Source/Programs/AutomationTool/Gauntlet/Unreal/Base/Gauntlet.UnrealTestNode.cs
@@ -1409,6 +1409,7 @@ namespace Gauntlet
 							if (Item.ProcessResult != UnrealProcessResult.ExitOk)
 							{
 								RelevantStream.AddError($"{RoleType}: {Item.Summary}");
+								InitPhase.SetOutcome(TestPhaseOutcome.Interrupted);
 							}
 						}
 						// Since we are going through the Role, lets add the related device
@@ -1508,7 +1509,11 @@ namespace Gauntlet
 				MainStream.AddEvent(Event);
 			}
 			// Outcome
-			ClosurePhase.SetOutcome(GetTestResult());
+			TestResult Result = GetTestResult();
+			if (Result != TestResult.Passed)
+			{
+				ClosurePhase.SetOutcome(Result);
+			}
 		}
 
 		private AutomatedTestSessionData.TestDevice AddRoleDeviceToSessionReport(UnrealRoleResult RoleResult, AutomatedTestSessionData TestReport)
@@ -1588,6 +1593,7 @@ namespace Gauntlet
 								EventStream.AddEvent(Event);
 							}
 							EventStream.AddError($"{RoleType}: {Item.Summary}");
+							InitPhase.SetOutcome(TestPhaseOutcome.Interrupted);
 						}
 						// Since we are going through the Role, lets add the related device
 						var ReportDevice = AddRoleDeviceToSessionReport(Item, SessionReport);
