[Test Report] Just make sure path are Unix valid instead of converting to URI

#rnx

diff --git a/Engine/Source/Programs/AutomationTool/Gauntlet/Framework/Gauntlet.Utils.cs b/Engine/Source/Programs/AutomationTool/Gauntlet/Framework/Gauntlet.Utils.cs
--- a/Engine/Source/Programs/AutomationTool/Gauntlet/Framework/Gauntlet.Utils.cs
+++ b/Engine/Source/Programs/AutomationTool/Gauntlet/Framework/Gauntlet.Utils.cs
@@ -1978,14 +1978,13 @@ namespace Gauntlet
 		}
 
 		/// <summary>
-		/// Convert file path to a Uri valid and Url encoded string
+		/// Convert file path to Unix-style path
 		/// </summary>
 		/// <param name="Path"></param>
 		/// <returns></returns>
-		static public string ConvertPathToUri(string Path)
+		static public string ConvertPathToUnix(string Path)
 		{
-			// Uri.AbsolutPath remove only 'file://', the extra / still need to be removed, enhance using string.Substring(1) 
-			return new Uri($"file:///{Path}").AbsolutePath.Substring(1);
+			return Path.Replace("\\", "/");
 		}
 	}
 
diff --git a/Engine/Source/Programs/AutomationTool/Gauntlet/Framework/Utils/Gauntlet.HordeReport.cs b/Engine/Source/Programs/AutomationTool/Gauntlet/Framework/Utils/Gauntlet.HordeReport.cs
--- a/Engine/Source/Programs/AutomationTool/Gauntlet/Framework/Utils/Gauntlet.HordeReport.cs
+++ b/Engine/Source/Programs/AutomationTool/Gauntlet/Framework/Utils/Gauntlet.HordeReport.cs
@@ -1312,7 +1312,10 @@ namespace Gauntlet
 					tags = null;
 					errorFingerprint = null;
 					eventStream = new TestEventStream(session?.report, this);
-					AttachToSession(InSession);
+					if (InSession != null)
+					{
+						AttachToSession(InSession);
+					}
 				}
 
 				/// <summary>
@@ -1528,7 +1531,7 @@ namespace Gauntlet
 					if (Device != null)
 					{
 						string LogPathForHorde = Path.GetRelativePath(Globals.UnrealRootDir, Path.Combine(OutputArtifactPath, LogName));
-						Device.appInstanceLogPath = FileUtils.ConvertPathToUri(LogPathForHorde);
+						Device.appInstanceLogPath = FileUtils.ConvertPathToUnix(LogPathForHorde);
 						return true;
 					}
 				}
@@ -1674,6 +1677,7 @@ namespace Gauntlet
 					stream.Add(Event);
 					switch (Event.Level)
 					{
+						case Logging.LogLevel.Critical:
 						case Logging.LogLevel.Error:
 							phase.errorFingerprint = (phase.errorFingerprint ?? "") + GenerateHash(Event.Message);
 							break;
@@ -1895,7 +1899,7 @@ namespace Gauntlet
 				if (AttachArtifact(FilePath, Name))
 				{
 					FilePath = Path.Combine(OutputArtifactPath, Name ?? Path.GetFileName(FilePath));
-					FilePath = FileUtils.ConvertPathToUri(Path.GetRelativePath(Globals.UnrealRootDir, FilePath));
+					FilePath = FileUtils.ConvertPathToUnix(Path.GetRelativePath(Globals.UnrealRootDir, FilePath));
 					return true;
 				}
 				return false;
diff --git a/Engine/Source/Programs/AutomationTool/Gauntlet/Unreal/Automation/UE.Automation.cs b/Engine/Source/Programs/AutomationTool/Gauntlet/Unreal/Automation/UE.Automation.cs
--- a/Engine/Source/Programs/AutomationTool/Gauntlet/Unreal/Automation/UE.Automation.cs
+++ b/Engine/Source/Programs/AutomationTool/Gauntlet/Unreal/Automation/UE.Automation.cs
@@ -876,7 +876,7 @@ namespace UE
 								if (TempReport.AttachArtifact(Artifact.LogPath, LogName) && Artifact.SessionRole.RoleType == MainRole.Type)
 								{
 									// Reference last run instance log
-									JsonTestPassResults.Devices.Last().AppInstanceLog = Gauntlet.FileUtils.ConvertPathToUri(Path.GetRelativePath(Globals.UnrealRootDir, Path.Combine(HordeArtifactPath, LogName)));
+									JsonTestPassResults.Devices.Last().AppInstanceLog = Gauntlet.FileUtils.ConvertPathToUnix(Path.GetRelativePath(Globals.UnrealRootDir, Path.Combine(HordeArtifactPath, LogName)));
 								}
 							}
 							JsonTestPassResults.WriteToJson();
diff --git a/Engine/Source/Programs/AutomationTool/Gauntlet/Unreal/Base/Gauntlet.UnrealTestNode.cs b/Engine/Source/Programs/AutomationTool/Gauntlet/Unreal/Base/Gauntlet.UnrealTestNode.cs
--- a/Engine/Source/Programs/AutomationTool/Gauntlet/Unreal/Base/Gauntlet.UnrealTestNode.cs
+++ b/Engine/Source/Programs/AutomationTool/Gauntlet/Unreal/Base/Gauntlet.UnrealTestNode.cs
@@ -1454,7 +1454,7 @@ namespace Gauntlet
 					string LogPath = RoleResult.Artifacts.LogPath;
 					if (!string.IsNullOrEmpty(LogPath) && File.Exists(LogPath))
 					{
-						string LogName = FileUtils.ConvertPathToUri(Path.GetRelativePath(Path.GetFullPath(ArtifactPath), Path.GetFullPath(LogPath)));
+						string LogName = FileUtils.ConvertPathToUnix(Path.GetRelativePath(Path.GetFullPath(ArtifactPath), Path.GetFullPath(LogPath)));
 						if (HordeTestReport is BaseHordeReport Report)
 						{
 							string DeviceKey = RoleResult.Artifacts.AppInstance.Device.Name.Replace(".", "-");
@@ -1692,7 +1692,7 @@ namespace Gauntlet
 						string DeviceKey = Artifact.AppInstance.Device.Name.Replace(".", "-"); // Key must not contains dot
 						if (!string.IsNullOrEmpty(Artifact.LogPath))
 						{
-							string LogName = FileUtils.ConvertPathToUri(Path.GetRelativePath(Path.GetFullPath(Globals.Params.ParseParam("UseTestDataV2") ? ArtifactPath : Context.Options.LogDir), Path.GetFullPath(Artifact.LogPath)));
+							string LogName = FileUtils.ConvertPathToUnix(Path.GetRelativePath(Path.GetFullPath(Globals.Params.ParseParam("UseTestDataV2") ? ArtifactPath : Context.Options.LogDir), Path.GetFullPath(Artifact.LogPath)));
 							if (HordeTestPassResults is BaseHordeReport Report)
 							{
 								Report.AttachDeviceLog(DeviceKey, Artifact.LogPath, LogName);
