Add 30 secs delay to WaitForExit() to prevent job hangs on process exit

diff --git a/Engine/Source/Programs/AutomationTool/AutomationUtils/ProcessUtils.cs b/Engine/Source/Programs/AutomationTool/AutomationUtils/ProcessUtils.cs
--- a/Engine/Source/Programs/AutomationTool/AutomationUtils/ProcessUtils.cs
+++ b/Engine/Source/Programs/AutomationTool/AutomationUtils/ProcessUtils.cs
@@ -544,7 +544,7 @@ namespace AutomationTool
 
 				if (bProcTerminated)
 				{
-					Proc?.WaitForExit();
+					Proc?.WaitForExit(30000);
 				}
 
 				if (!bProcTerminated)
@@ -567,7 +567,7 @@ namespace AutomationTool
 			}
 
 			// wait until the process terminates
-			Proc?.WaitForExit();
+			Proc?.WaitForExit(30000);
 
 			lock (ProcSyncObject)
 			{
