Adding a custom staging handler to allow projects to target specific game stores. It provides staging, packaging, deployment and launching support.

#rnx

diff --git a/Engine/Source/Programs/AutomationTool/AutomationUtils/CustomStagingHandler.cs b/Engine/Source/Programs/AutomationTool/AutomationUtils/CustomStagingHandler.cs
new file mode 100644
--- /dev/null
+++ b/Engine/Source/Programs/AutomationTool/AutomationUtils/CustomStagingHandler.cs
@@ -0,0 +1,110 @@
+// Copyright Epic Games, Inc. All Rights Reserved.
+using AutomationTool;
+using System;
+using System.Collections.Generic;
+using System.Linq;
+using System.Reflection;
+using static AutomationTool.CommandUtils;
+
+/// <summary>
+/// Support for generating custom packages, for example when targeting a specific game store
+/// </summary>
+public abstract class CustomStagingHandler
+{
+	/// <summary>
+	/// Determines whether this custom packaging handler should be used in the given deployment context
+	/// </summary>
+	protected abstract bool TryInitialize(ProjectParams Params, DeploymentContext SC);
+
+	/// <summary>
+	/// Hook for adding additional files for staging.
+	/// </summary>
+	public virtual void GetFilesToStage(ProjectParams Params, DeploymentContext SC) { }
+
+	/// <summary>
+	/// Hook for adding additional files for staging when packaging a DLC plugin
+	/// </summary>
+	public virtual void GetFilesToStageForDLC(ProjectParams Params, DeploymentContext SC) { }
+
+	/// <summary>
+	/// Hook for when the build has been staged and the list of staged files is available
+	/// </summary>
+	public virtual void PostStagingFileCopy(ProjectParams Params, DeploymentContext SC) { }
+
+	/// <summary>
+	/// Hook for package generation
+	/// </summary>
+	public virtual void Package(ProjectParams Params, DeploymentContext SC, int WorkingCL) { }
+
+	/// <summary>
+	/// Hook for adding additional files for archiving.
+	/// </summary>
+	public virtual void GetFilesToArchive(ProjectParams Params, DeploymentContext SC ) { }
+
+
+	/// <summary>
+	/// Hook for deploying a package we have created
+	/// </summary>
+	public virtual bool TryDeploy(ProjectParams Params, DeploymentContext SC) => false;
+
+	/// <summary>
+	/// Hook for launching a package we have created
+	/// </summary>
+	public virtual IProcessResult TryRunClient(ERunOptions ClientRunFlags, string ClientApp, string ClientCmdLine, ProjectParams Params, DeploymentContext SC) => null;
+
+
+
+	#region Private/boilerplate
+
+	/// <summary>
+	/// Validate and instantiate the given custom staging handlers for the current project. (Note this is returning a List not an IEnumerable to allow the caller to .ForEach() on the handlers)
+	/// </summary>
+	public static List<CustomStagingHandler> GetCustomStagingHandlers( ProjectParams Params, DeploymentContext SC)
+	{
+		// one-time static init. NB. Can't use the partial struct convention because automation classes are all in different assemblies
+		if (HandlerRegistry == null)
+		{
+			HandlerRegistry = new();
+
+			foreach (Assembly LoadedAssembly in ScriptManager.AllScriptAssemblies)
+			{
+				try
+				{
+					IEnumerable<Type> HandlerTypes = LoadedAssembly.GetTypes().Where(X => !X.IsAbstract && X.IsSubclassOf(typeof(CustomStagingHandler)));
+					foreach (Type HandlerType in HandlerTypes)
+					{
+						HandlerRegistry.Add(HandlerType);
+					}
+				}
+				catch(ReflectionTypeLoadException)
+				{
+				}
+			}
+		}
+
+		// collect all handlers that are valid for the current project
+		List<CustomStagingHandler> AvailableHandlers = [];
+		foreach (Type HandlerType in HandlerRegistry)
+		{
+			try
+			{
+				CustomStagingHandler Handler = (CustomStagingHandler)Activator.CreateInstance(HandlerType);
+				if (Handler != null && Handler.TryInitialize(Params, SC))
+				{
+					AvailableHandlers.Add(Handler);
+				}
+			}
+			catch (Exception)
+			{
+			}
+		}
+
+		return AvailableHandlers;
+	}
+
+	private static List<Type> HandlerRegistry = null;
+
+	#endregion
+}
+
+
diff --git a/Engine/Source/Programs/AutomationTool/AutomationUtils/DeploymentContext.cs b/Engine/Source/Programs/AutomationTool/AutomationUtils/DeploymentContext.cs
--- a/Engine/Source/Programs/AutomationTool/AutomationUtils/DeploymentContext.cs
+++ b/Engine/Source/Programs/AutomationTool/AutomationUtils/DeploymentContext.cs
@@ -157,10 +157,15 @@ public class DeploymentContext //: ProjectParams
 	public string CustomConfig;
 
 	/// <summary>
-	/// Allows a platform to change how it is packaged, staged and deployed - for example, when packaging for a specific game store
+	/// Allows a platform to change how pak files are generated and how it is staged, packaged and deployed - for example, when packaging for a specific game store
 	/// </summary>
 	public CustomDeploymentHandler CustomDeployment = null;
 
+	/// <summary>
+	/// Support for generating custom packages, for example when targeting a specific game store
+	/// </summary>
+	public List<CustomStagingHandler> CustomStagingHandlers = [];
+
 	/// <summary>
 	/// This is the root directory that contains the engine: d:\a\UE\
 	/// </summary>
diff --git a/Engine/Source/Programs/AutomationTool/Scripts/ArchiveCommand.Automation.cs b/Engine/Source/Programs/AutomationTool/Scripts/ArchiveCommand.Automation.cs
--- a/Engine/Source/Programs/AutomationTool/Scripts/ArchiveCommand.Automation.cs
+++ b/Engine/Source/Programs/AutomationTool/Scripts/ArchiveCommand.Automation.cs
@@ -25,6 +25,7 @@ namespace AutomationScripts
 			{
 				ThisPlatform.GetFilesToArchive(Params, SC);
 			}
+			SC.CustomStagingHandlers.ForEach(H => H.GetFilesToArchive(Params, SC));
 
 			//@todo add any archive meta data files as needed
 
diff --git a/Engine/Source/Programs/AutomationTool/Scripts/CopyBuildToStagingDirectory.Automation.cs b/Engine/Source/Programs/AutomationTool/Scripts/CopyBuildToStagingDirectory.Automation.cs
--- a/Engine/Source/Programs/AutomationTool/Scripts/CopyBuildToStagingDirectory.Automation.cs
+++ b/Engine/Source/Programs/AutomationTool/Scripts/CopyBuildToStagingDirectory.Automation.cs
@@ -1544,6 +1544,7 @@ namespace AutomationScripts
 				{
 					ThisPlatform.GetFilesToStageForDLC(Params, SC);
 				}
+				SC.CustomStagingHandlers.ForEach(H => H.GetFilesToStageForDLC(Params, SC));
 
 				bCreatePluginManifest = true;
 			}
@@ -1556,6 +1557,7 @@ namespace AutomationScripts
 						ThisPlatform.GetFilesToDeployOrStage(Params, SC);
 					}
 					SC.CustomDeployment?.PostGetFilesToDeployOrStage(Params, SC);
+					SC.CustomStagingHandlers.ForEach(H => H.GetFilesToStage(Params, SC));
 
 					// Stage any extra runtime dependencies from the receipts
 					foreach (StageTarget Target in SC.StageTargets)
@@ -5830,6 +5832,7 @@ namespace AutomationScripts
 				var ThisPlatform = SC.StageTargetPlatform;
 				ThisPlatform.PostStagingFileCopy(Params, SC);
 			}
+			SC.CustomStagingHandlers.ForEach(H => H.PostStagingFileCopy(Params, SC));
 		}
 
 		private static DirectoryReference GetIntermediateCommandlineDir(DeploymentContext SC)
@@ -6283,6 +6286,8 @@ namespace AutomationScripts
 				LogDeploymentContext(SC);
 
 				SC.CustomDeployment = SC.StageTargetPlatform?.GetCustomDeploymentHandler(Params, SC);
+				SC.CustomStagingHandlers = CustomStagingHandler.GetCustomStagingHandlers(Params, SC);
+
 
 				if (Params.PreModifyDeploymentContextCallback != null)
 				{
diff --git a/Engine/Source/Programs/AutomationTool/Scripts/DeployCommand.Automation.cs b/Engine/Source/Programs/AutomationTool/Scripts/DeployCommand.Automation.cs
--- a/Engine/Source/Programs/AutomationTool/Scripts/DeployCommand.Automation.cs
+++ b/Engine/Source/Programs/AutomationTool/Scripts/DeployCommand.Automation.cs
@@ -7,6 +7,7 @@ using System.Reflection;
 using AutomationTool;
 using UnrealBuildTool;
 using Microsoft.Extensions.Logging;
+using System.Linq;
 
 namespace AutomationScripts
 {
@@ -55,9 +56,13 @@ namespace AutomationScripts
 					}
 					else
 					{
-						if (SC.CustomDeployment == null || !SC.CustomDeployment.Deploy(Params, SC))
+						bool bSuccess = SC.CustomStagingHandlers.Any( H => H.TryDeploy(Params, SC) );
+						if (!bSuccess )
 						{
-							SC.StageTargetPlatform.Deploy(Params, SC);
+							if (SC.CustomDeployment == null || !SC.CustomDeployment.Deploy(Params, SC))
+							{
+								SC.StageTargetPlatform.Deploy(Params, SC);
+							}
 						}
 					}
 				}
diff --git a/Engine/Source/Programs/AutomationTool/Scripts/PackageCommand.Automation.cs b/Engine/Source/Programs/AutomationTool/Scripts/PackageCommand.Automation.cs
--- a/Engine/Source/Programs/AutomationTool/Scripts/PackageCommand.Automation.cs
+++ b/Engine/Source/Programs/AutomationTool/Scripts/PackageCommand.Automation.cs
@@ -51,6 +51,7 @@ namespace AutomationScripts
 								SC.StageTargetPlatform.Package(Params, SC, WorkingCL);
 							}
 							SC.CustomDeployment?.PostPackage(Params, SC, WorkingCL);
+							SC.CustomStagingHandlers.ForEach(H => H.Package(Params, SC, WorkingCL));
 						}
 					}
 
diff --git a/Engine/Source/Programs/AutomationTool/Scripts/RunProjectCommand.Automation.cs b/Engine/Source/Programs/AutomationTool/Scripts/RunProjectCommand.Automation.cs
--- a/Engine/Source/Programs/AutomationTool/Scripts/RunProjectCommand.Automation.cs
+++ b/Engine/Source/Programs/AutomationTool/Scripts/RunProjectCommand.Automation.cs
@@ -299,7 +299,21 @@ namespace AutomationScripts
 
 		private static IProcessResult RunClientInternal( ERunOptions ClientRunFlags, string ClientApp, string ClientCmdLine, ProjectParams Params, DeploymentContext SC )
 		{
-			IProcessResult CustomResult = SC.CustomDeployment?.RunClient(ClientRunFlags, ClientApp, ClientCmdLine, Params, SC);
+			IProcessResult CustomResult = null;
+			foreach (CustomStagingHandler PackagingHandler in SC.CustomStagingHandlers)
+			{
+				CustomResult = PackagingHandler.TryRunClient(ClientRunFlags, ClientApp, ClientCmdLine, Params, SC);
+				if (CustomResult != null)
+				{
+					break;
+				}
+			}
+
+			if (CustomResult == null)
+			{
+				CustomResult = SC.CustomDeployment?.RunClient(ClientRunFlags, ClientApp, ClientCmdLine, Params, SC);
+			}
+
 			return CustomResult ?? SC.StageTargetPlatform.RunClient(ClientRunFlags, ClientApp, ClientCmdLine, Params, SC);
 		}
 
