Removing 30s timeout on still living proceses

#rnx

diff --git a/Engine/Source/Programs/AutomationTool/AutomationUtils/ProcessUtils.cs b/Engine/Source/Programs/AutomationTool/AutomationUtils/ProcessUtils.cs
--- a/Engine/Source/Programs/AutomationTool/AutomationUtils/ProcessUtils.cs
+++ b/Engine/Source/Programs/AutomationTool/AutomationUtils/ProcessUtils.cs
@@ -567,7 +567,7 @@ namespace AutomationTool
 			}
 
 			// wait until the process terminates
-			Proc?.WaitForExit(30000);
+			Proc?.WaitForExit();
 
 			lock (ProcSyncObject)
 			{
