- Allow MCP to gzip files
- Add FileSize param to the upload return payload
- Generate receipt file
- Fix up backend upload of receipts
- Fix up build scripts for receipt generation

diff --git a/Engine/Source/Programs/AutomationTool/AutomationUtils/MCPPublic.cs b/Engine/Source/Programs/AutomationTool/AutomationUtils/MCPPublic.cs
--- a/Engine/Source/Programs/AutomationTool/AutomationUtils/MCPPublic.cs
+++ b/Engine/Source/Programs/AutomationTool/AutomationUtils/MCPPublic.cs
@@ -2193,12 +2193,13 @@ namespace EpicGames.MCP.Automation
 		/// <param name="Identifier">The identifier or filename of the file to write.</param>
 		/// <param name="Contents">A byte array containing the data to write.</param>
 		/// <param name="ContentType">The MIME type of the file being uploaded. If left NULL, will be determined server-side by cloud provider.</param>
+		/// <param name="bGzip">If true, the data to be uploaded will be gzipped, and the appropraite Content-Encoding value set on the uploaded file in S3.</param>
 		/// <param name="bOverwrite">If true, will overwrite an existing file.  If false, will throw an exception if the file exists.</param>
 		/// <param name="bMakePublic">Specifies whether the file should be made publicly readable.</param>
 		/// <param name="bQuiet">If set to true, all log output for the operation is supressed.</param>
 		/// <param name="Metadata">If not null, key-value pairs of metadata to be applied to the object.</param>
 		/// <returns>A PostFileResult indicating whether the call was successful, and the URL to the uploaded file.</returns>
-		abstract public PostFileResult PostFile(string Container, string Identifier, byte[] Contents, string ContentType = null, bool bOverwrite = true, bool bMakePublic = false, bool bQuiet = false, IDictionary<string, object> Metadata = null);
+		abstract public PostFileResult PostFile(string Container, string Identifier, byte[] Contents, string ContentType = null, bool bGzip = false, bool bOverwrite = true, bool bMakePublic = false, bool bQuiet = false, IDictionary<string, object> Metadata = null);
 
 		/// <summary>
 		/// Posts a file to the cloud storage provider asynchronously.
@@ -2207,12 +2208,13 @@ namespace EpicGames.MCP.Automation
 		/// <param name="Identifier">The identifier or filename of the file to write.</param>
 		/// <param name="Contents">A byte array containing the data to write.</param>
 		/// <param name="ContentType">The MIME type of the file being uploaded. If left NULL, will be determined server-side by cloud provider.</param>
+		/// <param name="bGzip">If true, the data to be uploaded will be gzipped, and the appropraite Content-Encoding value set on the uploaded file in S3.</param>
 		/// <param name="bOverwrite">If true, will overwrite an existing file.  If false, will throw an exception if the file exists.</param>
 		/// <param name="bMakePublic">Specifies whether the file should be made publicly readable.</param>
         /// <param name="bQuiet">If set to true, all log output for the operation is supressed.</param>
 		/// <param name="Metadata">If not null, key-value pairs of metadata to be applied to the object.</param>
 		/// <returns>A PostFileResult indicating whether the call was successful, and the URL to the uploaded file.</returns>
-		abstract public Task<PostFileResult> PostFileAsync(string Container, string Identifier, byte[] Contents, string ContentType = null, bool bOverwrite = true, bool bMakePublic = false, bool bQuiet = false, IDictionary<string, object> Metadata = null);
+		abstract public Task<PostFileResult> PostFileAsync(string Container, string Identifier, byte[] Contents, string ContentType = null, bool bGzip = false, bool bOverwrite = true, bool bMakePublic = false, bool bQuiet = false, IDictionary<string, object> Metadata = null);
 
 		/// <summary>
 		/// Posts a file to the cloud storage provider.
@@ -2221,12 +2223,13 @@ namespace EpicGames.MCP.Automation
 		/// <param name="Identifier">The identifier or filename of the file to write.</param>
 		/// <param name="SourceFilePath">The full path of the file to upload.</param>
 		/// <param name="ContentType">The MIME type of the file being uploaded. If left NULL, will be determined server-side by cloud provider.</param>
+		/// <param name="bGzip">If true, the data to be uploaded will be gzipped, and the appropraite Content-Encoding value set on the uploaded file in S3.</param>
 		/// <param name="bOverwrite">If true, will overwrite an existing file.  If false, will throw an exception if the file exists.</param>
 		/// <param name="bMakePublic">Specifies whether the file should be made publicly readable.</param>
 		/// <param name="bQuiet">If set to true, all log output for the operation is supressed.</param>
 		/// <param name="Metadata">If not null, key-value pairs of metadata to be applied to the object.</param>
 		/// <returns>A PostFileResult indicating whether the call was successful, and the URL to the uploaded file.</returns>
-		abstract public PostFileResult PostFile(string Container, string Identifier, string SourceFilePath, string ContentType = null, bool bOverwrite = true, bool bMakePublic = false, bool bQuiet = false, IDictionary<string, object> Metadata = null);
+		abstract public PostFileResult PostFile(string Container, string Identifier, string SourceFilePath, string ContentType = null, bool bGzip = false, bool bOverwrite = true, bool bMakePublic = false, bool bQuiet = false, IDictionary<string, object> Metadata = null);
 
 		/// <summary>
 		/// Posts a file to the cloud storage provider asynchronously.
@@ -2235,12 +2238,13 @@ namespace EpicGames.MCP.Automation
 		/// <param name="Identifier">The identifier or filename of the file to write.</param>
 		/// <param name="SourceFilePath">The full path of the file to upload.</param>
 		/// <param name="ContentType">The MIME type of the file being uploaded. If left NULL, will be determined server-side by cloud provider.</param>
+		/// <param name="bGzip">If true, the data to be uploaded will be gzipped, and the appropraite Content-Encoding value set on the uploaded file in S3.</param>
 		/// <param name="bOverwrite">If true, will overwrite an existing file.  If false, will throw an exception if the file exists.</param>
 		/// <param name="bMakePublic">Specifies whether the file should be made publicly readable.</param>
         /// <param name="bQuiet">If set to true, all log output for the operation is supressed.</param>
 		/// <param name="Metadata">If not null, key-value pairs of metadata to be applied to the object.</param>
 		/// <returns>A PostFileResult indicating whether the call was successful, and the URL to the uploaded file.</returns>
-		abstract public Task<PostFileResult> PostFileAsync(string Container, string Identifier, string SourceFilePath, string ContentType = null, bool bOverwrite = true, bool bMakePublic = false, bool bQuiet = false, IDictionary<string, object> Metadata = null);
+		abstract public Task<PostFileResult> PostFileAsync(string Container, string Identifier, string SourceFilePath, string ContentType = null, bool bGzip = false, bool bOverwrite = true, bool bMakePublic = false, bool bQuiet = false, IDictionary<string, object> Metadata = null);
 
 		/// <summary>
 		/// Posts a file to the cloud storage provider using multiple connections.
@@ -2251,14 +2255,15 @@ namespace EpicGames.MCP.Automation
 		/// <param name="NumConcurrentConnections">The number of concurrent connections to use during uploading.</param>
 		/// <param name="PartSizeMegabytes">The size of each part that is uploaded. Minimum (and default) is 5 MB.</param>
 		/// <param name="ContentType">The MIME type of the file being uploaded. If left NULL, will be determined server-side by cloud provider.</param>
+		/// <param name="bGzip">If true, the data to be uploaded will be gzipped, and the appropraite Content-Encoding value set on the uploaded file in S3.</param>
 		/// <param name="bOverwrite">If true, will overwrite an existing file. If false, will throw an exception if the file exists.</param>
 		/// <param name="bMakePublic">Specifies whether the file should be made publicly readable.</param>
 		/// <param name="bQuiet">If set to true, all log output for the operation is supressed.</param>
 		/// <param name="Metadata">If not null, key-value pairs of metadata to be applied to the object.</param>
 		/// <returns>A PostFileResult indicating whether the call was successful, and the URL to the uploaded file.</returns>
-		public PostFileResult PostMultipartFile(string Container, string Identifier, string SourceFilePath, int NumConcurrentConnections, decimal PartSizeMegabytes = 5.0m, string ContentType = null, bool bOverwrite = true, bool bMakePublic = false, bool bQuiet = false, IDictionary<string, object> Metadata = null)
+		public PostFileResult PostMultipartFile(string Container, string Identifier, string SourceFilePath, int NumConcurrentConnections, decimal PartSizeMegabytes = 5.0m, string ContentType = null, bool bGzip = false, bool bOverwrite = true, bool bMakePublic = false, bool bQuiet = false, IDictionary<string, object> Metadata = null)
 		{
-			return PostMultipartFileAsync(Container, Identifier, SourceFilePath, NumConcurrentConnections, PartSizeMegabytes, ContentType, bOverwrite, bMakePublic, bQuiet, Metadata).Result;
+			return PostMultipartFileAsync(Container, Identifier, SourceFilePath, NumConcurrentConnections, PartSizeMegabytes, ContentType, bGzip, bOverwrite, bMakePublic, bQuiet, Metadata).Result;
 		}
 
 		/// <summary>
@@ -2270,12 +2275,13 @@ namespace EpicGames.MCP.Automation
 		/// <param name="NumConcurrentConnections">The number of concurrent connections to use during uploading.</param>
 		/// <param name="PartSizeMegabytes">The size of each part that is uploaded. Minimum (and default) is 5 MB.</param>
 		/// <param name="ContentType">The MIME type of the file being uploaded. If left NULL, will be determined server-side by cloud provider.</param>
+		/// <param name="bGzip">If true, the data to be uploaded will be gzipped, and the appropraite Content-Encoding value set on the uploaded file in S3.</param>
 		/// <param name="bOverwrite">If true, will overwrite an existing file. If false, will throw an exception if the file exists.</param>
 		/// <param name="bMakePublic">Specifies whether the file should be made publicly readable.</param>
 		/// <param name="bQuiet">If set to true, all log output for the operation is supressed.</param>
 		/// <param name="Metadata">If not null, key-value pairs of metadata to be applied to the object.</param>
 		/// <returns>A PostFileResult indicating whether the call was successful, and the URL to the uploaded file.</returns>
-		abstract public Task<PostFileResult> PostMultipartFileAsync(string Container, string Identifier, string SourceFilePath, int NumConcurrentConnections, decimal PartSizeMegabytes = 5.0m, string ContentType = null, bool bOverwrite = true, bool bMakePublic = false, bool bQuiet = false, IDictionary<string, object> Metadata = null);
+		abstract public Task<PostFileResult> PostMultipartFileAsync(string Container, string Identifier, string SourceFilePath, int NumConcurrentConnections, decimal PartSizeMegabytes = 5.0m, string ContentType = null, bool bGzip = false, bool bOverwrite = true, bool bMakePublic = false, bool bQuiet = false, IDictionary<string, object> Metadata = null);
 
 
 
@@ -2456,7 +2462,12 @@ namespace EpicGames.MCP.Automation
 			/// <summary>
 			/// Set to true if the write succeeds, false otherwise.
 			/// </summary>
+			
 			public bool bSuccess { get; set; }
+			/// <summary>
+			/// Upload file size
+			/// </summary>
+			public long FileSize { get; set; }
 		}
 
 		/// <summary>
