Increase timeout from 10 to 20 minutes for apps being ready to avoid timeout issues when running automated tests

diff --git a/Engine/Source/Programs/AutomationTool/Gauntlet/Framework/RpcFramework/Gauntlet.RpcLibrary.cs b/Engine/Source/Programs/AutomationTool/Gauntlet/Framework/RpcFramework/Gauntlet.RpcLibrary.cs
--- a/Engine/Source/Programs/AutomationTool/Gauntlet/Framework/RpcFramework/Gauntlet.RpcLibrary.cs
+++ b/Engine/Source/Programs/AutomationTool/Gauntlet/Framework/RpcFramework/Gauntlet.RpcLibrary.cs
@@ -557,7 +557,7 @@ namespace Gauntlet
 
 		}
 
-		public static void WaitForAllTargetsToReportReady(List<RpcTarget> InTargets, int TimeoutInSeconds = 600, int TimeBetweenRetriesInSeconds = 60)
+		public static void WaitForAllTargetsToReportReady(List<RpcTarget> InTargets, int TimeoutInSeconds = 1200, int TimeBetweenRetriesInSeconds = 60)
 		{
 			string ReadyMessageCategory = "TargetStatus";
 			string ReadyMessageContents = "FEngineLoopInitComplete";
