Update PRS Importer to make inline metadata format consistent with csv profiler metadata.

- csvid should be all uppercase.
- integers should be written as integers and not strings.

#rnx

diff --git a/Engine/Source/Programs/AutomationTool/Gauntlet/Framework/Utils/Gauntlet.CsvImporter.cs b/Engine/Source/Programs/AutomationTool/Gauntlet/Framework/Utils/Gauntlet.CsvImporter.cs
--- a/Engine/Source/Programs/AutomationTool/Gauntlet/Framework/Utils/Gauntlet.CsvImporter.cs
+++ b/Engine/Source/Programs/AutomationTool/Gauntlet/Framework/Utils/Gauntlet.CsvImporter.cs
@@ -26,7 +26,7 @@ namespace Gauntlet
 		/// <summary>
 		/// Get as a completed dict;
 		/// </summary>
-		public Dictionary<string, string> ToDict();
+		public Dictionary<string, dynamic> ToDict();
 	}
 
 	/// <summary>
@@ -61,7 +61,7 @@ namespace Gauntlet
 		/// Any additional metadata associated with this csv.
 		/// Corresponds to 'CustomFields' in the datasource config.
 		/// </summary>
-		public Dictionary<string, string> CustomFields { get; }
+		public Dictionary<string, dynamic> CustomFields { get; }
 
 		/// <summary>
 		/// Metadata for this import which varies on the import type.
@@ -102,7 +102,7 @@ namespace Gauntlet
 
 		public string CsvFilename { get; private set; } = null;
 		public string LogFilename { get; set; } = null;
-		public Dictionary<string, string> CustomFields { get; private set; }
+		public Dictionary<string, dynamic> CustomFields { get; private set; }
 		public ICsvMetadata Metadata { get; set; } = null;
 		public IReadOnlyDictionary<string, string> AdditionalFiles { get; private set; }
 		public Dictionary<string, double> AdditionalStatAverages { get; private set; } = null;
@@ -114,13 +114,13 @@ namespace Gauntlet
 		/// <param name="CsvFilename">The path to the csv to import.</param>
 		/// <param name="LogFilename">Optional path to a corresponding log to bundle with this csv.</param>
 		/// <param name="CustomFields">Optional metadata corresponding to 'CustomFields' in the datasource config.</param>
-		public CsvImportEntry(string CsvFilename, string LogFilename = null, Dictionary<string, string> CustomFields = null)
+		public CsvImportEntry(string CsvFilename, string LogFilename = null, Dictionary<string, dynamic> CustomFields = null)
 		{
 			this.ImportType = CsvImportType.CsvFile;
 			this.ImportName = Path.GetFileName(CsvFilename);
 			this.CsvFilename = CsvFilename;
 			this.LogFilename = LogFilename;
-			this.CustomFields = CustomFields ?? new Dictionary<string, string>();
+			this.CustomFields = CustomFields ?? new Dictionary<string, dynamic>();
 		}
 
 		/// <summary>
@@ -129,7 +129,7 @@ namespace Gauntlet
 		/// <param name="CsvFilename">The path to the csv to import.</param>
 		/// <param name="LogFilename">Optional path to a corresponding log to bundle with this csv.</param>
 		/// <param name="CustomFields">Optional metadata corresponding to 'CustomFields' in the datasource config.</param>
-		public static CsvImportEntry CreateCsvFileImportEntry(string CsvFilename, string LogFilename = null, Dictionary<string, string> CustomFields = null)
+		public static CsvImportEntry CreateCsvFileImportEntry(string CsvFilename, string LogFilename = null, Dictionary<string, dynamic> CustomFields = null)
 		{
 			return new CsvImportEntry(CsvFilename, LogFilename, CustomFields);
 		}
@@ -145,7 +145,7 @@ namespace Gauntlet
 		/// It's expected that these stats have been sanitized by the user. Any invalid stats will cause an ArgumentException to be thrown.
 		/// You can use the CsvStatSanitizer helper class to sanitize your stats.
 		/// </param>
-		public static CsvImportEntry CreateInlineImportEntry(string ImportName, ICsvMetadata CsvMetadata, Dictionary<string, double> AdditionalStatAverages = null, Dictionary<string, string> CustomFields = null)
+		public static CsvImportEntry CreateInlineImportEntry(string ImportName, ICsvMetadata CsvMetadata, Dictionary<string, double> AdditionalStatAverages = null, Dictionary<string, dynamic> CustomFields = null)
 		{
 			if (CsvMetadata == null)
 			{
@@ -168,7 +168,7 @@ namespace Gauntlet
 				ImportType = CsvImportType.Inline,
 				ImportName = ImportName,
 				Metadata = CsvMetadata,
-				CustomFields = CustomFields ?? new Dictionary<string, string>(),
+				CustomFields = CustomFields ?? new Dictionary<string, dynamic>(),
 				CsvFilename = null,
 				AdditionalStatAverages = AdditionalStatAverages
 			};
