Set critical failure engine test as interrupted phase
+ check Fatal error in role result when no Engine json test report was generated

#rnx

diff --git a/Engine/Source/Programs/AutomationTool/Gauntlet/Framework/Utils/Gauntlet.HordeReport.cs b/Engine/Source/Programs/AutomationTool/Gauntlet/Framework/Utils/Gauntlet.HordeReport.cs
--- a/Engine/Source/Programs/AutomationTool/Gauntlet/Framework/Utils/Gauntlet.HordeReport.cs
+++ b/Engine/Source/Programs/AutomationTool/Gauntlet/Framework/Utils/Gauntlet.HordeReport.cs
@@ -2336,6 +2336,7 @@ namespace Gauntlet
 						Phase.SetTiming(PhaseTime, InTestResult.Duration);
 						Phase.SetOutcome(InTestResult.State);
 						TestEventStream Stream = Phase.GetStream();
+						bool hasCriticalFailure = false;
 						foreach (UnrealAutomationEntry InTestEntry in InTestResult.Entries)
 						{
 							string Artifact = InTestEntry.Event.Artifact;
@@ -2365,6 +2366,11 @@ namespace Gauntlet
 								}
 							}
 							Stream.AddEvent(InTestEntry.AsLogEvent());
+							hasCriticalFailure = hasCriticalFailure || InTestEntry.Event.IsCriticalFailure;
+						}
+						if (hasCriticalFailure)
+						{
+							Phase.SetOutcome(TestPhaseOutcome.Interrupted);
 						}
 					}
 				}
diff --git a/Engine/Source/Programs/AutomationTool/Gauntlet/Unreal/Base/Gauntlet.UnrealTestNode.cs b/Engine/Source/Programs/AutomationTool/Gauntlet/Unreal/Base/Gauntlet.UnrealTestNode.cs
--- a/Engine/Source/Programs/AutomationTool/Gauntlet/Unreal/Base/Gauntlet.UnrealTestNode.cs
+++ b/Engine/Source/Programs/AutomationTool/Gauntlet/Unreal/Base/Gauntlet.UnrealTestNode.cs
@@ -1382,7 +1382,8 @@ namespace Gauntlet
 		/// <param name="Result"></param>
 		protected virtual ITestReport CreateSimpleReportForHorde(TestResult Result)
 		{
-			Horde.OnStepSummary += () => Horde.GenerateArtifactLink(ArtifactPath);
+			string ArtifactPathForReport = ArtifactPath;
+			Horde.OnStepSummary += () => Horde.GenerateArtifactLink(ArtifactPathForReport);
 			BaseHordeReport HordeTestReport = null;
 			if (Globals.Params.ParseParam("UseTestDataV2"))
 			{
@@ -1671,11 +1672,23 @@ namespace Gauntlet
 				{
 					var InitPhase = SessionReport.GetPhase("Initialization");
 					InitPhase.GetStream().AddWarning("Could not find Unreal Automated test report at {FilePath}. No test report details available.", JsonReportPath);
-					// Check if there is any Engine Test Error found to add more context
-					UnrealRoleResult RelevantRoleResult = RoleResults?.FirstOrDefault(I => I.ProcessResult == UnrealProcessResult.EngineTestError);
+					// Check if there is any Engine Test Error or Fatal Error found to add more context
+					UnrealRoleResult RelevantRoleResult = RoleResults?.FirstOrDefault(I => I.ProcessResult == UnrealProcessResult.EngineTestError || I.ProcessResult == UnrealProcessResult.EncounteredFatalError);
 					if (RelevantRoleResult != null)
 					{
-						InitPhase.GetStream().AddError(RelevantRoleResult.Summary);
+						string RoleType = RelevantRoleResult.Artifacts.SessionRole.RoleType.ToString();
+						var Stream = InitPhase.GetStream();
+						foreach(UnrealTestEvent Item in RelevantRoleResult.Events)
+						{
+							Stream.AddEvent(Item);
+						}
+						Stream.AddError("{Role}: {Summary}", RoleType, RelevantRoleResult.Summary);
+
+						// In case of a fatal error, it is an interrupted outcome
+						if (RelevantRoleResult.ProcessResult == UnrealProcessResult.EncounteredFatalError)
+						{
+							InitPhase.SetOutcome(TestPhaseOutcome.Interrupted);
+						}
 					}
 				}
 			}
