Fix System.IndexOutOfRangeException while adding multiple localization targets

#tests Run NewPluginLocalizationTarget command for plugin with existing target, specify additional target, observe no crash

diff --git a/Engine/Source/Programs/AutomationTool/Localization/LocalizationTargetEditor/Commands/NewPluginLocalizationTargetCommand.cs b/Engine/Source/Programs/AutomationTool/Localization/LocalizationTargetEditor/Commands/NewPluginLocalizationTargetCommand.cs
--- a/Engine/Source/Programs/AutomationTool/Localization/LocalizationTargetEditor/Commands/NewPluginLocalizationTargetCommand.cs
+++ b/Engine/Source/Programs/AutomationTool/Localization/LocalizationTargetEditor/Commands/NewPluginLocalizationTargetCommand.cs
@@ -193,7 +193,7 @@ namespace EpicGames.Localization
 			else
 			{
 				int numLocalizationTargets = descriptor.LocalizationTargets.Length;
-				LocalizationTargetDescriptor[] newLocalizationTargets = new LocalizationTargetDescriptor[numLocalizationTargets];
+				LocalizationTargetDescriptor[] newLocalizationTargets = new LocalizationTargetDescriptor[numLocalizationTargets + 1];
 				Array.Copy(descriptor.LocalizationTargets, newLocalizationTargets, numLocalizationTargets);
 				newLocalizationTargets[numLocalizationTargets] = locTargetDescriptorToAdd;
 				descriptor.LocalizationTargets = newLocalizationTargets;
