BuildGraph: Fixed source code indexing on Win64.

Restored the SrcSrv node which seems to have been accidentally removed last year during some refactoring of the build graph script. I've also fixed SymStore not to fail when encountering files outside of the workspace root (which is an edge case we have internally). However, in this case indexing sources using SrcSrv instead of SymStore seems to make more sense as it ensures that our PDBs always have the source information and not only when using the symbol server.

#rnx

diff --git a/Engine/Source/Programs/AutomationTool/Win/WinPlatform.Automation.cs b/Engine/Source/Programs/AutomationTool/Win/WinPlatform.Automation.cs
--- a/Engine/Source/Programs/AutomationTool/Win/WinPlatform.Automation.cs
+++ b/Engine/Source/Programs/AutomationTool/Win/WinPlatform.Automation.cs
@@ -793,12 +793,22 @@ public class Win64Platform : Platform
 		}
 
 		List<FileReference> CopiedSymbolFiles = new List<FileReference>();
+		List<FileReference> SkippedSymbolFiles = new List<FileReference>();
 
 		// Copy all the symbol files to the given directory.
 		foreach (FileReference File in SymbolFiles)
 		{
 			string RelativePath = File.MakeRelativeTo(Unreal.RootDirectory);
 
+			if (RelativePath == File.FullName)
+			{
+				Logger.LogInformation("Skipping source indexing for \"{File}\" because it is not under the root directory \"{RootDirectory}\"", File.FullName, Unreal.RootDirectory);
+
+				SkippedSymbolFiles.Add(File);
+
+				continue;
+			}
+
 			FileReference DestinationFile = FileReference.Combine(TargetDirectory, RelativePath);
 
 			try
@@ -816,6 +826,10 @@ public class Win64Platform : Platform
 		// Index all symbol files in one go (indexing source code from the source control is shared work between all pdbs).
 		AddSourceIndexToSymbols(CopiedSymbolFiles, SourceFiles, Branch, Change);
 
+		// If we skipped copying and source indexing of some files, we want to preserve the original references as this list
+		// is later used for symbol uploading.
+		CopiedSymbolFiles.AddRange(SkippedSymbolFiles);
+
 		return CopiedSymbolFiles;
 	}
 
