Update automatic loc text gather config generation by conditionally adding the GatherTextFromMetadata step based on the loading policy.

#tests Run Localize for Unreal Engine locally, assure gather step modified loc data

diff --git a/Engine/Source/Programs/AutomationTool/Localization/LocalizationConfigFileBuilder.cs b/Engine/Source/Programs/AutomationTool/Localization/LocalizationConfigFileBuilder.cs
--- a/Engine/Source/Programs/AutomationTool/Localization/LocalizationConfigFileBuilder.cs
+++ b/Engine/Source/Programs/AutomationTool/Localization/LocalizationConfigFileBuilder.cs
@@ -60,6 +60,29 @@ namespace EpicGames.Localization
 			configSection.AddValue("SkipGatherCache", gatherTextParams.bSkipGatherCache);
 		}
 
+		public void AddGatherTextFromMetadataSection(GatherTextFromMetaDataParams gatherTextParams)
+		{
+			LocalizationConfigFileSection configSection = _configFile.FindOrAddSection($"GatherTextStep{_gatherTextStepCount}");
+			++_gatherTextStepCount;
+
+			configSection.AddValue("CommandletClass", gatherTextParams.CommandletClass);
+
+			configSection.AddValues("ModulesToPreload", gatherTextParams.ModulesToPreload.ToArray());
+			configSection.AddValues("IncludePathFilters", gatherTextParams.IncludePathFilters.ToArray());
+			configSection.AddValues("ExcludePathFilters", gatherTextParams.ExcludePathFilters.ToArray());
+
+			configSection.AddValues("FieldTypesToInclude", gatherTextParams.FieldTypesToInclude.ToArray());
+			configSection.AddValues("FieldOwnerTypesToInclude", gatherTextParams.FieldOwnerTypesToInclude.ToArray());
+			
+			configSection.AddValues("InputKeys", gatherTextParams.InputKeys.ToArray());
+			configSection.AddValues("OutputKeys", gatherTextParams.OutputKeys.ToArray());
+
+			configSection.AddValues("InputNamespaceses", gatherTextParams.InputNamespaceses.ToArray());
+			configSection.AddValues("OutputNamespaces", gatherTextParams.OutputNamespaces.ToArray());
+
+			configSection.AddValue("ShouldGatherFromEditorOnlyData", gatherTextParams.bShouldGatherFromEditorOnlyData);
+		}
+
 		public void AddGenerateGatherManifestSection(GenerateGatherManifestParams generateGatherManifestParams)
 		{
 			LocalizationConfigFileSection configSection = _configFile.FindOrAddSection($"GatherTextStep{_gatherTextStepCount}");
diff --git a/Engine/Source/Programs/AutomationTool/Localization/LocalizationConfigFileGenerator.cs b/Engine/Source/Programs/AutomationTool/Localization/LocalizationConfigFileGenerator.cs
--- a/Engine/Source/Programs/AutomationTool/Localization/LocalizationConfigFileGenerator.cs
+++ b/Engine/Source/Programs/AutomationTool/Localization/LocalizationConfigFileGenerator.cs
@@ -3,6 +3,7 @@ using System;
 using System.Collections.Generic;
 using System.IO;
 using Microsoft.Extensions.Logging;
+using UnrealBuildTool;
 
 using static AutomationTool.CommandUtils;
 
@@ -10,6 +11,7 @@ namespace EpicGames.Localization
 {
 	public class LocalizationConfigFileGeneratorParams
 	{
+		public LocalizationTargetDescriptorLoadingPolicy LoadingPolicy { get; set; } = LocalizationTargetDescriptorLoadingPolicy.Never;
 		public string LocalizationTargetName { get; set; } = "";
 		public string LocalizationTargetRootDirectory { get; set; } = "";
 	}
@@ -82,11 +84,20 @@ namespace EpicGames.Localization
 			LocalizationConfigFileBuilder builder = new LocalizationConfigFileBuilder(generatorParam.LocalizationTargetName + ModularLocalizationConfigFileSuffixes.GatherSuffix + LocalizationFileExtensions.ModularLocalizationConfigFileExtension);
 			builder.AddCommonSettingsSection(commonSettings);
 
-			GatherTextFromSourceParams gatherTextFromSourceParams = GatherTextFromSourceParams.CreateDefault(generatorParam.LocalizationTargetName, generatorParam.LocalizationTargetRootDirectory); 
-			builder.AddGatherTextFromSourceSection(gatherTextFromSourceParams);
+			if (generatorParam.LoadingPolicy != LocalizationTargetDescriptorLoadingPolicy.ToolTips &&
+				generatorParam.LoadingPolicy != LocalizationTargetDescriptorLoadingPolicy.PropertyNames)
+			{
+				GatherTextFromSourceParams gatherTextFromSourceParams = GatherTextFromSourceParams.CreateDefault(generatorParam.LocalizationTargetName, generatorParam.LocalizationTargetRootDirectory);
+				builder.AddGatherTextFromSourceSection(gatherTextFromSourceParams);
 
-			GatherTextFromAssetsParams gatherTextFromAssetsParams = GatherTextFromAssetsParams.CreateDefault(generatorParam.LocalizationTargetName, generatorParam.LocalizationTargetRootDirectory);
-			builder.AddGatherTextFromAssetsSection(gatherTextFromAssetsParams);
+				GatherTextFromAssetsParams gatherTextFromAssetsParams = GatherTextFromAssetsParams.CreateDefault(generatorParam.LocalizationTargetName, generatorParam.LocalizationTargetRootDirectory);
+				builder.AddGatherTextFromAssetsSection(gatherTextFromAssetsParams);
+			}
+			else
+			{
+				GatherTextFromMetaDataParams gatherTextFromMetadata = GatherTextFromMetaDataParams.CreateDefault(generatorParam.LocalizationTargetRootDirectory, generatorParam.LoadingPolicy);
+				builder.AddGatherTextFromMetadataSection(gatherTextFromMetadata);
+			}
 
 			GenerateGatherManifestParams generateGatherManifestParams = GenerateGatherManifestParams.CreateDefault(generatorParam.LocalizationTargetName, generatorParam.LocalizationTargetRootDirectory);
 			builder.AddGenerateGatherManifestSection(generateGatherManifestParams);
diff --git a/Engine/Source/Programs/AutomationTool/Localization/LocalizationConfigFileParams.cs b/Engine/Source/Programs/AutomationTool/Localization/LocalizationConfigFileParams.cs
--- a/Engine/Source/Programs/AutomationTool/Localization/LocalizationConfigFileParams.cs
+++ b/Engine/Source/Programs/AutomationTool/Localization/LocalizationConfigFileParams.cs
@@ -1,6 +1,7 @@
 // Copyright Epic Games, Inc. All Rights Reserved.
 using System.Collections.Generic;
 using System.IO;
+using UnrealBuildTool;
 
 namespace EpicGames.Localization
 {
@@ -114,7 +115,31 @@ namespace EpicGames.Localization
 		public List<string> OutputNamespaces { get; set; } = new List<string>();
 		public bool bShouldGatherFromEditorOnlyData { get; set; } = false;
 
-		// @TODOLocalization: Create defaults for gathering from metada 
+		public static GatherTextFromMetaDataParams CreateDefault(string localizationTargetRootDirectory, LocalizationTargetDescriptorLoadingPolicy loadingPolicy)
+		{
+			GatherTextFromMetaDataParams gatherTextFromMetaData = new GatherTextFromMetaDataParams();
+
+			gatherTextFromMetaData.IncludePathFilters.Add(Path.Combine(localizationTargetRootDirectory, "Source",  "*"));
+
+			if (loadingPolicy == LocalizationTargetDescriptorLoadingPolicy.PropertyNames)
+			{
+				gatherTextFromMetaData.InputKeys.Add("DisplayName");
+				gatherTextFromMetaData.OutputKeys.Add("\"{FieldPath}\"");
+				gatherTextFromMetaData.OutputNamespaces.Add("UObjectDisplayNames");
+			}
+
+			if (loadingPolicy == LocalizationTargetDescriptorLoadingPolicy.ToolTips)
+			{
+				gatherTextFromMetaData.InputKeys.Add("ToolTip");
+				gatherTextFromMetaData.InputKeys.Add("ShortTooltip");
+				gatherTextFromMetaData.OutputKeys.Add("\"{FieldPath}\"");
+				gatherTextFromMetaData.OutputKeys.Add("\"{FieldPath}\"");
+				gatherTextFromMetaData.OutputNamespaces.Add("UObjectToolTips");
+			}
+
+			gatherTextFromMetaData.bShouldGatherFromEditorOnlyData = true;
+			return gatherTextFromMetaData;
+		}
 	}
 
 	public class GenerateGatherManifestParams
diff --git a/Engine/Source/Programs/AutomationTool/Scripts/Localisation.Automation.cs b/Engine/Source/Programs/AutomationTool/Scripts/Localisation.Automation.cs
--- a/Engine/Source/Programs/AutomationTool/Scripts/Localisation.Automation.cs
+++ b/Engine/Source/Programs/AutomationTool/Scripts/Localisation.Automation.cs
@@ -980,6 +980,7 @@ class Localize : BuildCommand
 					}
 					LocalizationConfigFileGenerator Generator = LocalizationConfigFileGenerator.GetGeneratorForFileFormat(LocalizationConfigFileFormat.Latest);
 					LocalizationConfigFileGeneratorParams GeneratorParams = new();
+					GeneratorParams.LoadingPolicy = Target.LoadingPolicy;
 					GeneratorParams.LocalizationTargetName = Target.Name;
 					GeneratorParams.LocalizationTargetRootDirectory = Info.Directory.MakeRelativeTo(PluginsRootDirectoryReference);
 					List<LocalizationConfigFile> PluginLocalizationConfigFiles;
