Fix incorrect parameter spelling on GatherTextFromMetadata config generation

#tests Run Localize command for any localization target with loading policy of "ToolTips" or "PropertyNames", observe Gather_*.ini file created with correct "InputNamespaces" entry

diff --git a/Engine/Source/Programs/AutomationTool/Localization/LocalizationConfigFileBuilder.cs b/Engine/Source/Programs/AutomationTool/Localization/LocalizationConfigFileBuilder.cs
--- a/Engine/Source/Programs/AutomationTool/Localization/LocalizationConfigFileBuilder.cs
+++ b/Engine/Source/Programs/AutomationTool/Localization/LocalizationConfigFileBuilder.cs
@@ -77,7 +77,7 @@ namespace EpicGames.Localization
 			configSection.AddValues("InputKeys", gatherTextParams.InputKeys.ToArray());
 			configSection.AddValues("OutputKeys", gatherTextParams.OutputKeys.ToArray());
 
-			configSection.AddValues("InputNamespaceses", gatherTextParams.InputNamespaceses.ToArray());
+			configSection.AddValues("InputNamespaces", gatherTextParams.InputNamespaceses.ToArray());
 			configSection.AddValues("OutputNamespaces", gatherTextParams.OutputNamespaces.ToArray());
 
 			configSection.AddValue("ShouldGatherFromEditorOnlyData", gatherTextParams.bShouldGatherFromEditorOnlyData);
