[AutomationTool]

Remove unnecessary | from character class in regex as the or is implicit and this includes '|' as accepted which is not correct.

#rnx

diff --git a/Engine/Source/Programs/AutomationTool/AutomationUtils/Matchers/SanitizerEventMatcher.cs b/Engine/Source/Programs/AutomationTool/AutomationUtils/Matchers/SanitizerEventMatcher.cs
--- a/Engine/Source/Programs/AutomationTool/AutomationUtils/Matchers/SanitizerEventMatcher.cs
+++ b/Engine/Source/Programs/AutomationTool/AutomationUtils/Matchers/SanitizerEventMatcher.cs
@@ -65,7 +65,7 @@ namespace AutomationUtils.Matchers
 		public static readonly Regex SummaryPattern = new Regex(
 			@"^[\s\t]*SUMMARY:\s+(?<SanitizerName>\w+)Sanitizer:\s*" +
 			@"(?<SummaryReason>[\w\s\-\(\)]+)\s" +
-			@"\(?(?<SourceFile>(\/|\w:)?.+?)(\+0[x|X][a-fA-F0-9]+)?\)?" +
+			@"\(?(?<SourceFile>(\/|\w:)?.+?)(\+0[xX][a-fA-F0-9]+)?\)?" +
 			@"(:(?<Line>\d+?))?" +
 			@"(:(?<Column>\d+?))?" +
 			@"(\s+\(BuildId:\s+[0-9a-f]+\))?" +
