Cloud Artifacts - Retry uploads with progressivley longer waits if they fail with a 502 or 504 message as these typically recover after a while. Also added similar waits for oplog export for the same exit codes.

diff --git a/Engine/Source/Programs/AutomationTool/BuildGraph/Tasks/CreateCloudArtifactTask.cs b/Engine/Source/Programs/AutomationTool/BuildGraph/Tasks/CreateCloudArtifactTask.cs
--- a/Engine/Source/Programs/AutomationTool/BuildGraph/Tasks/CreateCloudArtifactTask.cs
+++ b/Engine/Source/Programs/AutomationTool/BuildGraph/Tasks/CreateCloudArtifactTask.cs
@@ -4,6 +4,7 @@ using System;
 using System.Collections.Generic;
 using System.IO;
 using System.Linq;
+using System.Threading;
 using System.Threading.Tasks;
 using System.Xml;
 using EpicGames.Core;
@@ -499,38 +500,81 @@ namespace AutomationTool.Tasks
 			metadata.Add("totalSize", totalSize);
 			string createdAtString = DateTime.UtcNow.ToString("O");
 			metadata.Add("createdAt", createdAtString);
-			
-			CbObjectId buildId = CbObjectId.NewObjectId();
-			CbObjectId partId = CbObjectId.NewObjectId();
-
-			FileReference? traceFile = null;
-			if (_parameters.EnableTracing)
-			{
-				// this is automatically uploaded as it's in the Saved folder
-				traceFile = FileReference.Combine(Unreal.RootDirectory, $"Engine/Programs/AutomationTool/Saved/Logs/build-upload-{buildId}-{partId}.utrace");
-			}
 
+			Uri? buildUri = null;
 			// set part name to default by default - should be exposed as a option to allow for uploading separate parts
 			string partName = "default";
 
-			BuildUploader.UploadBuild(host, baseDir, ns, bucketId, buildId, partName, partId, manifestFile, metadata, accessToken, new BuildUploader.BuildInfo
+			const int MaxAttempts = 3;
+			for (int countOfRetries = 0; countOfRetries < MaxAttempts; countOfRetries++)
 			{
-				Name = name.ToString(),
-				ProjectName = project,
-				Branch = streamId,
-				BaselineBranch = baselineBranch,
-				Platform = platform
-			}, allowMultipart: _parameters.AllowMultipart, assumeHttp2: String.Equals(httpVersion, "http2-only", StringComparison.OrdinalIgnoreCase), forceClean: _parameters.Force, verify: _parameters.Verify, verbose: _parameters.Verbose, boostWorkers: _parameters.BoostWorkers, traceFile: traceFile);
+				// generate new build and part ids for each attempt so that we do not continue to append to a broken build in case of a retry
+				CbObjectId buildId = CbObjectId.NewObjectId();
+				CbObjectId partId = CbObjectId.NewObjectId();
 
-			Uri buildUri = new Uri($"{descriptorHost}/api/v2/builds/{ns}/{bucketId}/{buildId}");
+				FileReference? traceFile = null;
+				if (_parameters.EnableTracing)
+				{
+					// this is automatically uploaded as it's in the Saved folder
+					traceFile = FileReference.Combine(Unreal.RootDirectory, $"Engine/Programs/AutomationTool/Saved/Logs/build-upload-{buildId}-{partId}.utrace");
+				}
 
-			try
-			{
-				await RegisterHordeArtifactAsync(buildUri, hordeArtifactID, type, streamIdTyped, commitId, metadata);
+				try
+				{
+					BuildUploader.UploadBuild(host, baseDir, ns, bucketId, buildId, partName, partId, manifestFile, metadata, accessToken, new BuildUploader.BuildInfo
+					{
+						Name = name.ToString(),
+						ProjectName = project,
+						Branch = streamId,
+						BaselineBranch = baselineBranch,
+						Platform = platform
+					}, allowMultipart: _parameters.AllowMultipart, assumeHttp2: String.Equals(httpVersion, "http2-only", StringComparison.OrdinalIgnoreCase), forceClean: _parameters.Force, verify: _parameters.Verify, verbose: _parameters.Verbose, boostWorkers: _parameters.BoostWorkers, traceFile: traceFile);
+
+					buildUri = new Uri($"{descriptorHost}/api/v2/builds/{ns}/{bucketId}/{buildId}");
+
+					// command successful
+					break;
+				}	
+				catch (CommandUtils.CommandFailedException e)
+				{
+					// increasingly long time to wait before retrying
+					int[] waitDurationSeconds = [60, 300, 900];
+					int waitDuration = waitDurationSeconds[countOfRetries];
+					int exitCode = (int)e.ErrorCode;
+					if (exitCode is 502 or 504)
+					{
+						// zen returns http status code as exit codes if it had an issue
+						Logger.LogInformation("Http Status code {ExitCode} returned from zen. Waiting {WaitDuration}s before trying.", exitCode, waitDuration);
+						if (countOfRetries+1 == MaxAttempts)
+						{
+							// give up
+							Logger.LogError("Zen returned http status code {ExitCode} after {RetryAttempts} retries. Giving up.", exitCode, countOfRetries);
+							throw;
+						}
+						else
+						{
+							Thread.Sleep(TimeSpan.FromSeconds(waitDuration));
+							continue; // retry
+						}
+					}
+					else
+					{
+						Logger.LogWarning("Unknown exit code from zen when command failed: {ExitCode} . Not retrying.", exitCode);
+						throw;
+					}
+				}
 			}
-			catch (Exception ex)
+
+			if (buildUri != null)
 			{
-				Logger.LogWarning(ex, "Exception registering Horde artifact: {Message}", ex.Message);
+				try
+				{
+					await RegisterHordeArtifactAsync(buildUri, hordeArtifactID, type, streamIdTyped, commitId, metadata);
+				}
+				catch (Exception ex)
+				{
+					Logger.LogWarning(ex, "Exception registering Horde artifact: {Message}", ex.Message);
+				}
 			}
 		}
 
diff --git a/Engine/Source/Programs/AutomationTool/BuildGraph/Tasks/ZenExportSnapshotTask.cs b/Engine/Source/Programs/AutomationTool/BuildGraph/Tasks/ZenExportSnapshotTask.cs
--- a/Engine/Source/Programs/AutomationTool/BuildGraph/Tasks/ZenExportSnapshotTask.cs
+++ b/Engine/Source/Programs/AutomationTool/BuildGraph/Tasks/ZenExportSnapshotTask.cs
@@ -9,6 +9,7 @@ using System.Runtime.InteropServices;
 using System.Text;
 using System.Text.Json;
 using System.Text.Json.Serialization;
+using System.Threading;
 using System.Threading.Tasks;
 using System.Xml;
 using EpicGames.Core;
@@ -413,7 +414,7 @@ namespace AutomationTool.Tasks
 			writer.WriteObjectEnd();
 		}
 
-		private static bool TryRunAndLogWithoutSpew(string app, string commandLine, bool ignoreFailure)
+		private static bool TryRunAndLogWithoutSpew(string app, string commandLine, bool ignoreFailure, out int exitCode)
 		{
 			ProcessResult.SpewFilterCallbackType silentOutputFilter = new ProcessResult.SpewFilterCallbackType(line =>
 				{
@@ -429,8 +430,12 @@ namespace AutomationTool.Tasks
 				{
 					Logger.LogWarning("{Text}", e.ToString());
 				}
+
+				exitCode = (int)e.ErrorCode;
 				return false;
 			}
+
+			exitCode = 0;
 			return true;
 		}
 
@@ -438,14 +443,20 @@ namespace AutomationTool.Tasks
 		{
 			int attemptLimit = 2;
 			int attempt = 0;
+			// increasingly long time to wait before retrying
+			int[] waitDurationWebSeconds = [60, 300, 900];
+
 			while (attempt < attemptLimit)
 			{
-				if (TryRunAndLogWithoutSpew(app, commandLine, false))
+				if (TryRunAndLogWithoutSpew(app, commandLine, false, out int exitCode))
 				{
 					return true;
 				}
-				Logger.LogWarning("Attempt {AttemptNum} of exporting the oplog failed, {Action}...", attempt + 1, attempt < (attemptLimit - 1) ? "retrying" : "abandoning");
 
+				bool isGatewayError = exitCode is 502 or 504;
+				int waitDuration = isGatewayError ? waitDurationWebSeconds[attempt] : 0;
+				Logger.LogWarning("Attempt {AttemptNum} of exporting the oplog failed, {Action}...", attempt + 1, attempt < (attemptLimit - 1) ? $"retrying after {waitDuration}s" : "abandoning");
+				Thread.Sleep(TimeSpan.FromSeconds(waitDuration));
 				attempt = attempt + 1;
 			}
 			return false;
@@ -691,7 +702,7 @@ namespace AutomationTool.Tasks
 
 					StringBuilder createProjectCommandline = new StringBuilder();
 					createProjectCommandline.AppendFormat("project-create --hosturl {0} {1}", _parameters.DestinationZenHost, projectName);
-					TryRunAndLogWithoutSpew(zenExe.FullName, createProjectCommandline.ToString(), true);
+					TryRunAndLogWithoutSpew(zenExe.FullName, createProjectCommandline.ToString(), true, out int _);
 
 					oplogExportCommandline.AppendFormat(" --zen {0}", _parameters.DestinationZenHost);
 
