Oplog export - Do not wait when we have decided to not retry anymore

diff --git a/Engine/Source/Programs/AutomationTool/BuildGraph/Tasks/ZenExportSnapshotTask.cs b/Engine/Source/Programs/AutomationTool/BuildGraph/Tasks/ZenExportSnapshotTask.cs
--- a/Engine/Source/Programs/AutomationTool/BuildGraph/Tasks/ZenExportSnapshotTask.cs
+++ b/Engine/Source/Programs/AutomationTool/BuildGraph/Tasks/ZenExportSnapshotTask.cs
@@ -455,8 +455,12 @@ namespace AutomationTool.Tasks
 
 				bool isGatewayError = exitCode is 502 or 504;
 				int waitDuration = isGatewayError ? waitDurationWebSeconds[attempt] : 0;
-				Logger.LogWarning("Attempt {AttemptNum} of exporting the oplog failed, {Action}...", attempt + 1, attempt < (attemptLimit - 1) ? $"retrying after {waitDuration}s" : "abandoning");
-				Thread.Sleep(TimeSpan.FromSeconds(waitDuration));
+				bool willRetry = attempt < (attemptLimit - 1);
+				Logger.LogWarning("Attempt {AttemptNum} of exporting the oplog failed, {Action}...", attempt + 1, willRetry ? $"retrying after {waitDuration}s" : "abandoning");
+				if (willRetry)
+				{
+					Thread.Sleep(TimeSpan.FromSeconds(waitDuration));
+				}
 				attempt = attempt + 1;
 			}
 			return false;
