AutomationTool: Added support for the Virtual File System feature in UBA/UBT when indexing source code for the symbol server.

VFS is a new experimental feature, which replaces real paths during the build process with virtualized ones (the same ones regardless of the real workspace root). It is enabled in some of our projects and it broke source code indexing because we had a discrepancy between the virtualized paths in symbol files and the real ones, which were used during indexing.

Added an option to provide VFS mappings to the Build Graph tasks related to symbols. Currently, the path mapping has to be explicitly given because this information is not available outside of UBT/UBA. We may improve it in the future.

diff --git a/Engine/Source/Programs/AutomationTool/AutomationUtils/Platform.cs b/Engine/Source/Programs/AutomationTool/AutomationUtils/Platform.cs
--- a/Engine/Source/Programs/AutomationTool/AutomationUtils/Platform.cs
+++ b/Engine/Source/Programs/AutomationTool/AutomationUtils/Platform.cs
@@ -834,7 +834,7 @@ namespace AutomationTool
 
 		public virtual bool PublishSymbols(DirectoryReference SymbolStoreDirectory, List<FileReference> Files,
 			bool bIndexSources, List<FileReference> SourceFiles,
-			string Product, string Branch, int Change, string BuildVersion = null)
+			string Product, string Branch, int Change, string BuildVersion = null, string VFSMapping = null)
 		{
 			Logger.LogWarning("PublishSymbols() has not been implemented for {Arg0}", PlatformType.ToString());
 			return false;
diff --git a/Engine/Source/Programs/AutomationTool/AutomationUtils/SourceIndexingUtils.cs b/Engine/Source/Programs/AutomationTool/AutomationUtils/SourceIndexingUtils.cs
new file mode 100644
--- /dev/null
+++ b/Engine/Source/Programs/AutomationTool/AutomationUtils/SourceIndexingUtils.cs
@@ -0,0 +1,143 @@
+// Copyright Epic Games, Inc. All Rights Reserved.
+
+using AutomationTool;
+using EpicGames.Core;
+using Microsoft.Extensions.Logging;
+using System;
+using System.Collections.Generic;
+using System.Linq;
+using UnrealBuildBase;
+
+namespace AutomationUtils
+{
+	using VFSPathMapping = (DirectoryReference Local, DirectoryReference Virtualized);
+	using VFSPathMappingList = List<(DirectoryReference Local, DirectoryReference Virtualized)>;
+
+	public static class SourceIndexingUtils
+	{
+		/// <summary>
+		/// Read a list of Virtual File System path transformations from a single string.
+		/// The string may have multiple entries separate by semicolons.
+		/// Each entry is an assignment: LocalPath=VirtualPath.
+		/// LocalPath may use [Root] which is replaced by by the current workspace root (Unreal.RootDirectory).
+		/// </summary>
+		/// <param name="VFSMapping">List of path transformations.</param>
+		/// <returns></returns>
+		public static VFSPathMappingList ParseVFSMapping(string VFSMapping)
+		{
+			VFSPathMappingList Result = new();
+
+			if (string.IsNullOrEmpty(VFSMapping))
+			{
+				return Result;
+			}
+
+			foreach (string EntryRaw in VFSMapping.Split(';', StringSplitOptions.RemoveEmptyEntries))
+			{
+				string Entry = EntryRaw.Trim();
+
+				if (Entry.Length == 0)
+				{
+					continue;
+				}
+
+				int EqualPosition = Entry.IndexOf('=');
+
+				if (EqualPosition < 1 || EqualPosition == Entry.Length - 1)
+				{
+					continue;
+				}
+
+				string Local = Entry.Substring(0, EqualPosition).Trim();
+				string Virtualized = Entry.Substring(EqualPosition + 1).Trim();
+
+				Local = Local.Replace("[Root]", Unreal.RootDirectory.FullName);
+
+				Result.Add((new DirectoryReference(Local), new DirectoryReference(Virtualized)));
+			}
+
+			return Result;
+		}
+
+		/// <summary>
+		/// Helper function printing the given filesystem path mapping to the log for debug purposes.
+		/// </summary>
+		/// <param name="VFSMapping"></param>
+		public static void OutputVFSMapping(VFSPathMappingList VFSMapping)
+		{
+			if (VFSMapping.Count > 0)
+			{
+				CommandUtils.Logger.LogInformation("Applying the following Virtual File System path transformations to source indexing information:");
+
+				foreach (var VFSMappingEntry in VFSMapping)
+				{
+					CommandUtils.Logger.LogInformation("\t{Local} -> {Virtualized}", VFSMappingEntry.Local.FullName, VFSMappingEntry.Virtualized.FullName);
+				}
+			}
+		}
+
+		/// <summary>
+		/// Converts a local path to a virtualized path (as generated by Unreal Build Tool when using Virtual File System i.e. compiling with -vfs switch).
+		/// Only the first matching rule from VFSMapping is applied.
+		/// </summary>
+		/// <param name="Local">Local path to a source file e.g. C:\UE\Main\Engine\Source\Foo.cpp</param>
+		/// <param name="VFSMapping">List of path substitutions (generated by ParseVFSMapping).</param>
+		/// <returns>Virtualized path e.g. Z:\UEVFS\Root\Engine\Source\Foo.cpp. If no VFSMapping rules matched, Local is returned.</returns>
+		public static string ConvertLocalToVirtualized(string Local, VFSPathMappingList VFSMapping)
+		{
+			FileReference LocalFile = new FileReference(Local);
+
+			foreach (VFSPathMapping VFSMappingEntry in VFSMapping)
+			{
+				if (LocalFile.IsUnderDirectory(VFSMappingEntry.Local))
+				{
+					return FileReference.Combine(VFSMappingEntry.Virtualized, LocalFile.MakeRelativeTo(VFSMappingEntry.Local)).FullName;
+				}
+			}
+
+			return Local;
+		}
+
+		/// <summary>
+		/// Build a database of source code files in the current Perforce workspace.
+		/// By relying on the standard layout for UE projects i.e. the fact that source code is in Source directories,
+		/// we may very significantly reduce the about of data sent to us from the server.
+		/// </summary>
+		/// <param name="Pattern">Perforce pattern path to query e.g. //UE/Branch/.../Source/...</param>
+		/// <returns></returns>
+		public static Dictionary<string, P4HaveRecord> BuildSourceDatabase(string Pattern, string VFSMappingString)
+		{
+			List<P4HaveRecord> Files = null;
+
+			P4Connection DefaultConnection = new P4Connection(User: null, Client: null, ServerAndPort: null);
+
+			try
+			{
+				Files = DefaultConnection.HaveFiles(Pattern);
+			}
+			catch (P4Exception e)
+			{
+				CommandUtils.Logger.LogError("Failed to fetch source code information from Perforce for '{Pattern}' ({Message}).", Pattern, e.Message);
+
+				return null;
+			}
+
+			if (!string.IsNullOrEmpty(VFSMappingString))
+			{
+				VFSPathMappingList VFSMapping = ParseVFSMapping(VFSMappingString);
+
+				if (VFSMapping.Count > 0)
+				{
+					OutputVFSMapping(VFSMapping);
+
+					foreach (P4HaveRecord File in Files)
+					{
+						File.ClientFile = ConvertLocalToVirtualized(File.ClientFile, VFSMapping);
+					}
+				}
+			}
+
+			return Files.ToDictionary(file => file.ClientFile, file => file, StringComparer.InvariantCultureIgnoreCase);
+		}
+	}
+}
diff --git a/Engine/Source/Programs/AutomationTool/BuildGraph/Tasks/SymStoreTask.cs b/Engine/Source/Programs/AutomationTool/BuildGraph/Tasks/SymStoreTask.cs
--- a/Engine/Source/Programs/AutomationTool/BuildGraph/Tasks/SymStoreTask.cs
+++ b/Engine/Source/Programs/AutomationTool/BuildGraph/Tasks/SymStoreTask.cs
@@ -77,6 +77,21 @@ namespace AutomationTool.Tasks
 		/// </summary>
 		[TaskParameter(Optional = true)]
 		public string SourceFiles { get; set; }
+
+		/// <summary>
+		/// List of file path mappings to be applied during source indexing. It is needed only when using VFS (Virtual File System)
+		/// feature when compiling the binary. VFS replaces the real local paths used with standardized paths like Z:\UEVFS\...
+		/// Since these virtualized paths don't match paths used by the local workspace, we need to map them manually so that
+		/// paths in object files/symbols matched the data we have in source indexing.
+		/// It's a semicolon-separated list of path assignments e.g.
+		///      "[Root]\Samples\Games\SampleGame=Z:\UEVFS\SampleGame;[Root]=Z:\UEVFS\Root"
+		/// 
+		/// [Root] will be replaced with Unreal.RootDirectory.
+		/// You need to ensure that the mapping used here matches the conventions used by Unreal Build Tool.
+		/// We may improve it in the future but current the VFS mapping used during compilation is not available outside of UBT.
+		/// </summary>
+		[TaskParameter(Optional = true)]
+		public string VFSMapping { get; set; }
 	}
 
 	/// <summary>
@@ -128,7 +143,7 @@ namespace AutomationTool.Tasks
 			CommandUtils.OptionallyTakeLock(targetPlatform.SymbolServerRequiresLock, storeDir, TimeSpan.FromMinutes(60), () =>
 			{
 				if (!targetPlatform.PublishSymbols(storeDir, files, _parameters.IndexSources, sourceFiles,
-					_parameters.Product, _parameters.Branch, _parameters.Change, _parameters.BuildVersion))
+					_parameters.Product, _parameters.Branch, _parameters.Change, _parameters.BuildVersion, _parameters.VFSMapping))
 				{
 					throw new AutomationException("Failure publishing symbol files.");
 				}
diff --git a/Engine/Source/Programs/AutomationTool/Win/Tasks/SrcSrvTask.cs b/Engine/Source/Programs/AutomationTool/Win/Tasks/SrcSrvTask.cs
--- a/Engine/Source/Programs/AutomationTool/Win/Tasks/SrcSrvTask.cs
+++ b/Engine/Source/Programs/AutomationTool/Win/Tasks/SrcSrvTask.cs
@@ -14,7 +14,7 @@ using System.Runtime.Versioning;
 namespace Win.Automation
 {
 	/// <summary>
-	/// Parameters for a task that uploads symbols to a symbol server
+	/// Parameters for a task that adds source indexing information to symbol files (PDB).
 	/// </summary>
 	public class SrcSrvTaskParameters
 	{
@@ -41,6 +41,21 @@ namespace Win.Automation
 		/// </summary>
 		[TaskParameter]
 		public int Change;
+
+		/// <summary>
+		/// List of file path mappings to be applied during source indexing. It is needed only when using VFS (Virtual File System)
+		/// feature when compiling the binary. VFS replaces the real local paths used with standardized paths like Z:\UEVFS\...
+		/// Since these virtualized paths don't match paths used by the local workspace, we need to map them manually so that
+		/// paths in object files/symbols matched the data we have in source indexing.
+		/// It's a semicolon-separated list of path assignments e.g.
+		///      "[Root]\Samples\Games\SampleGame=Z:\UEVFS\SampleGame;[Root]=Z:\UEVFS\Root"
+		/// 
+		/// [Root] will be replaced with Unreal.RootDirectory.
+		/// You need to ensure that the mapping used here matches the conventions used by Unreal Build Tool.
+		/// We may improve it in the future but current the VFS mapping used during compilation is not available outside of UBT.
+		/// </summary>
+		[TaskParameter(Optional = true)]
+		public string VFSMapping { get; set; }
 	}
 
 	/// <summary>
@@ -79,17 +94,17 @@ namespace Win.Automation
 		{
 			FileReference[] BinaryFiles = ResolveFilespec(Unreal.RootDirectory, Parameters.BinaryFiles, TagNameToFileSet).ToArray();
 			FileReference[] SourceFiles = ResolveFilespec(Unreal.RootDirectory, Parameters.SourceFiles, TagNameToFileSet).ToArray();
-			Execute(BinaryFiles, SourceFiles, Parameters.Branch, Parameters.Change);
+			Execute(BinaryFiles, SourceFiles, Parameters.Branch, Parameters.Change, Parameters.VFSMapping);
 		}
 
 		[SupportedOSPlatform("windows")]
-		internal static void Execute(FileReference[] BinaryFiles, FileReference[] SourceFiles, string Branch, int Change)
+		internal static void Execute(FileReference[] BinaryFiles, FileReference[] SourceFiles, string Branch, int Change, string VFSMapping)
 		{
 			IEnumerable<FileReference> PdbFiles = BinaryFiles.Where(x => x.HasExtension(".pdb"));
 
 			Win64Platform WindowsPlatform = Platform.GetPlatform(UnrealTargetPlatform.Win64) as Win64Platform;
 
-			WindowsPlatform.AddSourceIndexToSymbols(PdbFiles, SourceFiles, Branch, Change);
+			WindowsPlatform.AddSourceIndexToSymbols(PdbFiles, SourceFiles, Branch, Change, VFSMapping);
 		}
 
         /// <summary>
@@ -129,18 +144,20 @@ namespace Win.Automation
 	public static partial class BgStateExtensions
 	{
 		/// <summary>
-		/// Uploads symbols to a symbol server
+		/// Add source indexing information to symbol files (PDB).
 		/// </summary>
 		/// <param name="State"></param>
 		/// <param name="BinaryFiles">List of output files. PDBs will be extracted from this list.</param>
 		/// <param name="SourceFiles">List of source files to index and embed into the PDBs.</param>
 		/// <param name="Branch">Branch to base all the depot source files from.</param>
 		/// <param name="Change">Changelist to sync files from.</param>
+		/// <param name="VFSMapping">List of Virtual File System transformations to apply to paths during source code indexing. Pass an empty string if no transformations are needed
+		/// e.g. when VFS was not used when compiling binaries. Check description of SrcSrvTaskParameters.VFSMapping for details about the expected syntax.</param>
 		/// <returns></returns>
 		[SupportedOSPlatform("windows")]
-		public static void SrcSrv(this BgContext State, HashSet<FileReference> BinaryFiles, HashSet<FileReference> SourceFiles, string Branch, int Change)
+		public static void SrcSrv(this BgContext State, HashSet<FileReference> BinaryFiles, HashSet<FileReference> SourceFiles, string Branch, int Change, string VFSMapping)
 		{
-			SrcSrvTask.Execute(BinaryFiles.ToArray(), SourceFiles.ToArray(), Branch, Change);
+			SrcSrvTask.Execute(BinaryFiles.ToArray(), SourceFiles.ToArray(), Branch, Change, VFSMapping);
 		}
 	}
 }
diff --git a/Engine/Source/Programs/AutomationTool/Win/WinPlatform.Automation.cs b/Engine/Source/Programs/AutomationTool/Win/WinPlatform.Automation.cs
--- a/Engine/Source/Programs/AutomationTool/Win/WinPlatform.Automation.cs
+++ b/Engine/Source/Programs/AutomationTool/Win/WinPlatform.Automation.cs
@@ -6,6 +6,7 @@ using System.Text;
 using System.IO;
 using System.Text.RegularExpressions;
 using AutomationTool;
+using AutomationUtils;
 using UnrealBuildTool;
 using Microsoft.Win32;
 using System.Diagnostics;
@@ -641,7 +642,7 @@ public class Win64Platform : Platform
 	[SupportedOSPlatform("windows")]
 	public override bool PublishSymbols(DirectoryReference SymbolStoreDirectory, List<FileReference> Files,
 			bool bIndexSources, List<FileReference> SourceFiles,
-			string Product, string Branch, int Change, string BuildVersion = null)
+			string Product, string Branch, int Change, string BuildVersion = null, string VFSMapping = null)
     {
 		Logger.LogInformation("Publishing symbols to \"{SymbolStoreDirectory}\" (source indexing: {bIndexSources})", SymbolStoreDirectory, bIndexSources);
 
@@ -677,7 +678,7 @@ public class Win64Platform : Platform
 					// We can't write to the original symbol files because they may have been generated by a different build step,
 					// and clobbering such build products is considered an error.
 					// For this reason, we copy the symbol files and modify the copied.
-					SymbolsAfterIndexing = CopySymbolsWithSourceIndexing(TempSymStoreIndexedDir, SymbolsToIndex, SourceFiles, Branch, Change);
+					SymbolsAfterIndexing = CopySymbolsWithSourceIndexing(TempSymStoreIndexedDir, SymbolsToIndex, SourceFiles, Branch, Change, VFSMapping);
 				}
 
 				File.WriteAllLines(TempFileName, FilesToAdd.Except(SymbolsToIndex).Union(SymbolsAfterIndexing).Select(x => x.FullName), Encoding.ASCII);
@@ -784,7 +785,7 @@ public class Win64Platform : Platform
 
 	[SupportedOSPlatform("windows")]
 	private IEnumerable<FileReference> CopySymbolsWithSourceIndexing(DirectoryReference TargetDirectory, IEnumerable<FileReference> SymbolFiles,
-	IEnumerable<FileReference> SourceFiles, string Branch, int Change)
+	IEnumerable<FileReference> SourceFiles, string Branch, int Change, string VFSMapping)
 	{
 		if (DirectoryReference.Exists(TargetDirectory))
 		{
@@ -824,7 +825,7 @@ public class Win64Platform : Platform
 		}
 
 		// Index all symbol files in one go (indexing source code from the source control is shared work between all pdbs).
-		AddSourceIndexToSymbols(CopiedSymbolFiles, SourceFiles, Branch, Change);
+		AddSourceIndexToSymbols(CopiedSymbolFiles, SourceFiles, Branch, Change, VFSMapping);
 
 		// If we skipped copying and source indexing of some files, we want to preserve the original references as this list
 		// is later used for symbol uploading.
@@ -846,33 +847,6 @@ public class Win64Platform : Platform
 		return false;
 	}
 
-	/// <summary>
-	/// Build a database of source code files in the current Perforce workspace.
-	/// By relying on the standard layout for UE projects i.e. the fact that source code is in Source directories,
-	/// we may very significantly reduce the about of data sent to us from the server.
-	/// </summary>
-	/// <param name="Pattern">Perforce pattern path to query e.g. //UE/Branch/.../Source/...</param>
-	/// <returns></returns>
-	protected static Dictionary<string, P4HaveRecord> BuildSourceDatabase(string Pattern)
-	{
-		List<P4HaveRecord> Files = null;
-
-		P4Connection DefaultConnection = new P4Connection(User: null, Client: null, ServerAndPort: null);
-
-		try
-		{
-			Files = DefaultConnection.HaveFiles(Pattern);
-		}
-		catch (P4Exception e)
-		{
-			Logger.LogError("Failed to fetch source code information from Perforce for '{Pattern}' ({Message}).", Pattern, e.Message);
-
-			return null;
-		}
-
-		return Files.ToDictionary(file => file.ClientFile, file => file, StringComparer.InvariantCultureIgnoreCase);
-	}
-
 	/// <summary>
 	/// 
 	/// </summary>
@@ -880,20 +854,27 @@ public class Win64Platform : Platform
 	/// <param name="SourceFiles"></param>
 	/// <param name="Branch"></param>
 	/// <param name="Change"></param>
+	/// <param name="VFSMappingString"></param>
 	[SupportedOSPlatform("windows")]
-	public void AddSourceIndexToSymbols(IEnumerable<FileReference> PdbFiles, IEnumerable<FileReference> SourceFiles, string Branch, int Change)
+	public void AddSourceIndexToSymbols(IEnumerable<FileReference> PdbFiles, IEnumerable<FileReference> SourceFiles, string Branch, int Change, string VFSMappingString)
 	{
 		Logger.LogInformation("Adding source control information to PDB files...");
 
 		string DepotFilter = ".../Source/...";
 
-		Dictionary<string, P4HaveRecord> SourceDatabase = BuildSourceDatabase(DepotFilter);
+		// On Windows we get real (non-virtualized) paths as SourceFiles. We pass null as VFSMapping to BuildSourceDatabase
+		// so that our source database was indexed with these real paths. We only need to virtualize them when adding source info into PDB for the debugger.
+		Dictionary<string, P4HaveRecord> SourceDatabase = SourceIndexingUtils.BuildSourceDatabase(DepotFilter, null);
 
 		if (SourceDatabase == null)
 		{
 			throw new AutomationException($"Failed to query the source code information for '{DepotFilter}'.");
 		}
 
+		var VFSMapping = SourceIndexingUtils.ParseVFSMapping(VFSMappingString);
+
+		SourceIndexingUtils.OutputVFSMapping(VFSMapping);
+
 		// Get the pdbstr.exe path, using the latest SDK version we can find.
 		FileReference PdbStrExe = GetPdbStrExe();
 
@@ -921,7 +902,7 @@ public class Win64Platform : Platform
 
 				if (SourceDatabase.TryGetValue(SourceFile.FullName, out SourceInfo))
 				{
-					Writer.WriteLine("{0}*{1}*{2}", SourceFile.FullName, SourceInfo.DepotFile.Replace("//", ""), SourceInfo.Revision);
+					Writer.WriteLine("{0}*{1}*{2}", SourceIndexingUtils.ConvertLocalToVirtualized(SourceFile.FullName, VFSMapping), SourceInfo.DepotFile.Replace("//", ""), SourceInfo.Revision);
 				}
 				else
 				{
