[Automation] Adds log auditing to the device benchmark tag command

diff --git a/Engine/Source/Programs/AutomationTool/Gauntlet/Framework/Utils/Gauntlet.ReportGenUtils.cs b/Engine/Source/Programs/AutomationTool/Gauntlet/Framework/Utils/Gauntlet.ReportGenUtils.cs
--- a/Engine/Source/Programs/AutomationTool/Gauntlet/Framework/Utils/Gauntlet.ReportGenUtils.cs
+++ b/Engine/Source/Programs/AutomationTool/Gauntlet/Framework/Utils/Gauntlet.ReportGenUtils.cs
@@ -9,12 +9,15 @@ using AutomationTool;
 using UnrealBuildTool;
 using EpicGames.Core;
 using System.Text.Json;
+using Newtonsoft.Json;
 
 namespace Gauntlet
 {
 	public class SummaryTableData
 	{
+		[JsonIgnore]
 		public string CsvId { get; set; }
+		[JsonIgnore]
 		public string DeviceProfile { get; set; }
 		public Dictionary<string, double> CsvStatAverage { get; set; }
 		public Dictionary<string, string> CsvMetadata { get; set; }
