[Gauntlet] Adds the PreConfigureDevice delegate to UnrealTestRole which enables users to configure devices for an individual role BEFORE any installations or the creation of the app install

diff --git a/Engine/Source/Programs/AutomationTool/Gauntlet/Unreal/Base/Gauntlet.UnrealSession.cs b/Engine/Source/Programs/AutomationTool/Gauntlet/Unreal/Base/Gauntlet.UnrealSession.cs
--- a/Engine/Source/Programs/AutomationTool/Gauntlet/Unreal/Base/Gauntlet.UnrealSession.cs
+++ b/Engine/Source/Programs/AutomationTool/Gauntlet/Unreal/Base/Gauntlet.UnrealSession.cs
@@ -108,7 +108,12 @@ namespace Gauntlet
 		public List<EIntendedBaseCopyDirectory> AdditionalArtifactDirectories;
 
 		/// <summary>
-		/// Role device configuration
+		/// Allows for individual roles to configure their associated Target device before build installation
+		/// </summary>
+		public ConfigureDeviceHandler PreConfigureDevice;
+
+		/// <summary>
+		/// Allows for individual roles to configure their associated Target device after build installation and AppInstall creation
 		/// </summary>
 		public ConfigureDeviceHandler ConfigureDevice;
 
@@ -1397,6 +1402,8 @@ namespace Gauntlet
                                 Device.CleanArtifacts(AppConfig);
                             }
 
+							Role.PreConfigureDevice?.Invoke(Device);
+
 							Install = Device.InstallApplication(AppConfig);
 							IDeviceUsageReporter.RecordEnd(Device.Name, Device.Platform, IDeviceUsageReporter.EventType.Install, IDeviceUsageReporter.EventState.Success);
 						}
@@ -1672,6 +1679,8 @@ namespace Gauntlet
 							Device.CleanArtifacts(AppConfig);
 						}
 
+						Role.PreConfigureDevice?.Invoke(Device);
+
 						if (AppConfig.SkipInstall)
 						{
 							Log.Info("Skipping install due to SkipInstall");
diff --git a/Engine/Source/Programs/AutomationTool/Gauntlet/Unreal/Base/Gauntlet.UnrealTestConfiguration.cs b/Engine/Source/Programs/AutomationTool/Gauntlet/Unreal/Base/Gauntlet.UnrealTestConfiguration.cs
--- a/Engine/Source/Programs/AutomationTool/Gauntlet/Unreal/Base/Gauntlet.UnrealTestConfiguration.cs
+++ b/Engine/Source/Programs/AutomationTool/Gauntlet/Unreal/Base/Gauntlet.UnrealTestConfiguration.cs
@@ -494,7 +494,12 @@ namespace Gauntlet
 		public string MapOverride { get; set; }
 
 		/// <summary>
-		/// Role device configuration
+		/// Allows for individual roles to configure their associated Target device before build installation
+		/// </summary>
+		public ConfigureDeviceHandler PreConfigureDevice;
+
+		/// <summary>
+		/// Allows for individual roles to configure their associated Target device after build installation and AppInstall creation
 		/// </summary>
 		public ConfigureDeviceHandler ConfigureDevice;
 
diff --git a/Engine/Source/Programs/AutomationTool/Gauntlet/Unreal/Base/Gauntlet.UnrealTestNode.cs b/Engine/Source/Programs/AutomationTool/Gauntlet/Unreal/Base/Gauntlet.UnrealTestNode.cs
--- a/Engine/Source/Programs/AutomationTool/Gauntlet/Unreal/Base/Gauntlet.UnrealTestNode.cs
+++ b/Engine/Source/Programs/AutomationTool/Gauntlet/Unreal/Base/Gauntlet.UnrealTestNode.cs
@@ -666,6 +666,7 @@ namespace Gauntlet
 					// copy over relevant settings from test role
 					SessionRole.FilesToCopy = TestRole.FilesToCopy;
 					SessionRole.AdditionalArtifactDirectories = TestRole.AdditionalArtifactDirectories;
+					SessionRole.ConfigureDevice = TestRole.PreConfigureDevice;
 					SessionRole.ConfigureDevice = TestRole.ConfigureDevice;
 					SessionRole.MapOverride = TestRole.MapOverride;
 					SessionRole.CompressScreenshots = TestRole.CompressScreenshots;
