[Gauntlet] Correct typo in role configuration copy

diff --git a/Engine/Source/Programs/AutomationTool/Gauntlet/Unreal/Base/Gauntlet.UnrealTestNode.cs b/Engine/Source/Programs/AutomationTool/Gauntlet/Unreal/Base/Gauntlet.UnrealTestNode.cs
--- a/Engine/Source/Programs/AutomationTool/Gauntlet/Unreal/Base/Gauntlet.UnrealTestNode.cs
+++ b/Engine/Source/Programs/AutomationTool/Gauntlet/Unreal/Base/Gauntlet.UnrealTestNode.cs
@@ -666,7 +666,7 @@ namespace Gauntlet
 					// copy over relevant settings from test role
 					SessionRole.FilesToCopy = TestRole.FilesToCopy;
 					SessionRole.AdditionalArtifactDirectories = TestRole.AdditionalArtifactDirectories;
-					SessionRole.ConfigureDevice = TestRole.PreConfigureDevice;
+					SessionRole.PreConfigureDevice = TestRole.PreConfigureDevice;
 					SessionRole.ConfigureDevice = TestRole.ConfigureDevice;
 					SessionRole.MapOverride = TestRole.MapOverride;
 					SessionRole.CompressScreenshots = TestRole.CompressScreenshots;
