[Gauntlet] Adds a couple convenience functions to the HordeDeviceReservationService for querying device tags

#was-backout 45248178

diff --git a/Engine/Source/Programs/AutomationTool/Gauntlet/Framework/Devices/Gauntlet.DevicePool.cs b/Engine/Source/Programs/AutomationTool/Gauntlet/Framework/Devices/Gauntlet.DevicePool.cs
--- a/Engine/Source/Programs/AutomationTool/Gauntlet/Framework/Devices/Gauntlet.DevicePool.cs
+++ b/Engine/Source/Programs/AutomationTool/Gauntlet/Framework/Devices/Gauntlet.DevicePool.cs
@@ -1246,7 +1246,7 @@ namespace Gauntlet
 			}
 		}
 
-		private bool TryGetDevicesReservationService(ITargetDevice Device, out IDeviceReservationService OutService)
+		public bool TryGetDevicesReservationService(ITargetDevice Device, out IDeviceReservationService OutService)
 		{
 			OutService = null;
 
diff --git a/Engine/Source/Programs/AutomationTool/Gauntlet/Framework/Devices/Gauntlet.DeviceReservationService.cs b/Engine/Source/Programs/AutomationTool/Gauntlet/Framework/Devices/Gauntlet.DeviceReservationService.cs
--- a/Engine/Source/Programs/AutomationTool/Gauntlet/Framework/Devices/Gauntlet.DeviceReservationService.cs
+++ b/Engine/Source/Programs/AutomationTool/Gauntlet/Framework/Devices/Gauntlet.DeviceReservationService.cs
@@ -410,5 +410,31 @@ namespace Gauntlet
 				Reservation.ReportDeviceError(DeviceURL, ServiceDeviceInfo[MatchingDevice].Name, ErrorMessage);
 			}
 		}
+
+		/// <summary>
+		/// Returns a list of the tags on the device, if any
+		/// </summary>
+		/// <param name="Device">The device who's tags will be polled</param>
+		/// <returns>An array containg all tags on the device</returns>
+		public List<string> GetDeviceTags(ITargetDevice Device)
+		{
+			if (ServiceDeviceInfo.TryGetValue(Device, out DeviceDefinition Definition))
+			{
+				return Definition.Tags ?? [];
+			}
+
+			return [];
+		}
+
+		/// <summary>
+		/// Checks if a device has a tag
+		/// </summary>
+		/// <param name="Device">The device who's tags will be observed</param>
+		/// <param name="Tag">Name of the tag to look for</param>
+		/// <returns>True if the device has the provided Tag</returns>
+		public bool DeviceHasTag(ITargetDevice Device, string Tag)
+		{
+			return GetDeviceTags(Device).Contains(Tag, StringComparer.OrdinalIgnoreCase);
+		}
 	}
 }
\ No newline at end of file
