[Gauntlet] Modifies signature of PreConfigureDevice to accept a session role handle

diff --git a/Engine/Source/Programs/AutomationTool/Gauntlet/Unreal/Base/Gauntlet.UnrealSession.cs b/Engine/Source/Programs/AutomationTool/Gauntlet/Unreal/Base/Gauntlet.UnrealSession.cs
--- a/Engine/Source/Programs/AutomationTool/Gauntlet/Unreal/Base/Gauntlet.UnrealSession.cs
+++ b/Engine/Source/Programs/AutomationTool/Gauntlet/Unreal/Base/Gauntlet.UnrealSession.cs
@@ -110,7 +110,7 @@ namespace Gauntlet
 		/// <summary>
 		/// Allows for individual roles to configure their associated Target device before build installation
 		/// </summary>
-		public ConfigureDeviceHandler PreConfigureDevice;
+		public PreConfigureDeviceHandler PreConfigureDevice;
 
 		/// <summary>
 		/// Allows for individual roles to configure their associated Target device after build installation and AppInstall creation
@@ -1375,7 +1375,7 @@ namespace Gauntlet
 						Device = new TargetDeviceNull(string.Format("Null{0}", Role.RoleType));
 					}
 
-					Role.PreConfigureDevice?.Invoke(Device);
+					Role.PreConfigureDevice?.Invoke(Device, Role);
 
 					IEnumerable<UnrealSessionRole> OtherRoles = SortedRoles.Where(R => R != Role);
 
@@ -1661,7 +1661,7 @@ namespace Gauntlet
 						RolesToConfigs.Remove(Role);
 						RolesToInstalls.Remove(Role);
 
-						Role.PreConfigureDevice?.Invoke(Device);
+						Role.PreConfigureDevice?.Invoke(Device, Role);
 
 						// Create the app config
 						IEnumerable<UnrealSessionRole> OtherRoles = SessionRoles.Where(Other => Other != Role);
diff --git a/Engine/Source/Programs/AutomationTool/Gauntlet/Unreal/Base/Gauntlet.UnrealTestConfiguration.cs b/Engine/Source/Programs/AutomationTool/Gauntlet/Unreal/Base/Gauntlet.UnrealTestConfiguration.cs
--- a/Engine/Source/Programs/AutomationTool/Gauntlet/Unreal/Base/Gauntlet.UnrealTestConfiguration.cs
+++ b/Engine/Source/Programs/AutomationTool/Gauntlet/Unreal/Base/Gauntlet.UnrealTestConfiguration.cs
@@ -372,6 +372,11 @@ namespace Gauntlet
 		Success
 	}
 
+	/// <summary>
+	/// Delegate for role device pre-configuration
+	/// </summary>
+	public delegate void PreConfigureDeviceHandler(ITargetDevice Device, UnrealSessionRole Role);
+
 	/// <summary>
 	/// Delegate for role device configuration
 	/// </summary>
@@ -496,7 +501,7 @@ namespace Gauntlet
 		/// <summary>
 		/// Allows for individual roles to configure their associated Target device before build installation
 		/// </summary>
-		public ConfigureDeviceHandler PreConfigureDevice;
+		public PreConfigureDeviceHandler PreConfigureDevice;
 
 		/// <summary>
 		/// Allows for individual roles to configure their associated Target device after build installation and AppInstall creation
