Add option so that localization is only staged for GameFeaturePlugins when that GameFeaturePlugin is actually cooked. This can be used when cooking GameFeaturePlugins as part of a DLC.
You might only want to include localization with the cooking DLC rather than the base game, this way all the data for that DLC can be placed in that DLC pak.

New Option for Staging:
-StageOnlyCookedLocalizationForGameFeatures

So that only the localization data for cooking GFPs will be included.

diff --git a/Engine/Source/Programs/AutomationTool/AutomationUtils/ProjectParams.cs b/Engine/Source/Programs/AutomationTool/AutomationUtils/ProjectParams.cs
--- a/Engine/Source/Programs/AutomationTool/AutomationUtils/ProjectParams.cs
+++ b/Engine/Source/Programs/AutomationTool/AutomationUtils/ProjectParams.cs
@@ -319,6 +319,7 @@ namespace AutomationTool
 			this.DLCOverrideStagedSubDir = InParams.DLCOverrideStagedSubDir;
 			this.DLCOverrideCryptoSettings = InParams.DLCOverrideCryptoSettings;
 			this.DLCOverrideEncryptionKeyGuid = InParams.DLCOverrideEncryptionKeyGuid;
+			this.StageOnlyCookedLocalizationForGameFeatures = InParams.StageOnlyCookedLocalizationForGameFeatures;
 			this.DiffCookedContentPath = InParams.DiffCookedContentPath;
             this.AdditionalBuildOptions = InParams.AdditionalBuildOptions;
             this.AdditionalCookerOptions = InParams.AdditionalCookerOptions;
@@ -554,6 +555,7 @@ namespace AutomationTool
 			bool? DLCPakPluginFile = null,
 			bool? DLCActLikePatch = null,
 			bool? DLCSkipNamedPak = null,
+			bool? StageOnlyCookedLocalizationForGameFeatures = null,
 			bool? CrashReporter = null,
 			bool? DedicatedServer = null,
 			bool? Client = null,
@@ -782,6 +784,8 @@ namespace AutomationTool
 			this.DLCOverrideCookedSubDir = ParseParamValueIfNotSpecified(Command, DLCOverrideCookedSubDir, "DLCOverrideCookedSubDir", null);
 			this.DLCOverrideStagedSubDir = ParseParamValueIfNotSpecified(Command, DLCOverrideStagedSubDir, "DLCOverrideStagedSubDir", null);
 
+			this.StageOnlyCookedLocalizationForGameFeatures = GetParamValueIfNotSpecified(Command, StageOnlyCookedLocalizationForGameFeatures, this.StageOnlyCookedLocalizationForGameFeatures, "StageOnlyCookedLocalizationForGameFeatures");
+
 			this.SkipCook = GetParamValueIfNotSpecified(Command, SkipCook, this.SkipCook, "skipcook");
 			if (this.SkipCook)
 			{
@@ -1995,7 +1999,13 @@ namespace AutomationTool
 		/// <summary>
 		/// When true the DLC Name will not be added to the generated pak files.
 		/// </summary>
-		public bool DLCSkipNamedPak;
+		public bool DLCSkipNamedPak = false;
+
+		/// <summary>
+		/// When true only the GameFeaturePlugins that are cooked will have thier GameFeaturePlugin data staged.
+		/// This check is performed by looking for a cooked GameFeatureData asset within the GameFeaturePlugin expected folder.
+		/// </summary>
+		public bool StageOnlyCookedLocalizationForGameFeatures = false;
 
 		/// <summary>
 		/// After cook completes diff the cooked content against another cooked content directory.
diff --git a/Engine/Source/Programs/AutomationTool/Scripts/CopyBuildToStagingDirectory.Automation.cs b/Engine/Source/Programs/AutomationTool/Scripts/CopyBuildToStagingDirectory.Automation.cs
--- a/Engine/Source/Programs/AutomationTool/Scripts/CopyBuildToStagingDirectory.Automation.cs
+++ b/Engine/Source/Programs/AutomationTool/Scripts/CopyBuildToStagingDirectory.Automation.cs
@@ -1398,6 +1398,27 @@ namespace AutomationScripts
 			return CookedFiles;
 		}
 
+		private static void StageLocalizationForCookedGameFeaturePlugins(DeploymentContext SC, List<string> CulturesToStage, List<KeyValuePair<StagedFileReference, FileReference>> StagedGameFeaturePlugins, List<FileReference> CookedFiles)
+		{
+			// build map of possible GFD -> Plugin
+			Dictionary<FileReference, FileReference> GameFeatureDataToPlugin = new Dictionary<FileReference, FileReference>();
+			foreach (KeyValuePair<StagedFileReference, FileReference> StagedPlugin in StagedGameFeaturePlugins)
+			{
+				string PluginNameWithExtension = StagedPlugin.Value.GetFileName();
+				string PluginContentDirectory = Path.Combine(SC.PlatformCookDir.FullName, StagedPlugin.Key.ToString().TrimEnd(PluginNameWithExtension.ToCharArray()), "Content");
+				GameFeatureDataToPlugin.Add(new FileReference(Path.Combine(PluginContentDirectory, "GameFeatureData.uasset")), StagedPlugin.Value);
+				GameFeatureDataToPlugin.Add(new FileReference(Path.Combine(PluginContentDirectory, $"{StagedPlugin.Value.GetFileNameWithoutAnyExtensions()}.uasset")), StagedPlugin.Value);
+			}
+			foreach (FileReference CookedFile in CookedFiles)
+			{
+				// GFD was cooked. Stage this localized plugin data now.
+				if (GameFeatureDataToPlugin.ContainsKey(CookedFile))
+				{
+					StageLocalizationDataForPlugin(SC, CulturesToStage, GameFeatureDataToPlugin[CookedFile]);
+				}
+			}
+		}
+
 		public static void CreateStagingManifest(ProjectParams Params, DeploymentContext SC)
 		{
 			if (!Params.Stage)
@@ -1414,6 +1435,7 @@ namespace AutomationScripts
 			}
 
 			ConfigHierarchy PlatformGameConfig = ConfigCache.ReadHierarchy(ConfigHierarchyType.Game, DirectoryReference.FromFile(Params.RawProjectPath), SC.StageTargetPlatform.IniPlatformType, SC.CustomConfig, IncludePluginsForTargetType:TargetType.Game);
+			DirectoryReference GameFeaturePluginDir = DirectoryReference.Combine(SC.ProjectRoot, "Plugins", "GameFeatures");
 
 			if (Params.HasIterateSharedCookedBuild)
 			{
@@ -1430,6 +1452,8 @@ namespace AutomationScripts
 				DirectoryReference PlatformEngineDir = DirectoryReference.Combine(SC.PlatformCookDir, "Engine");
 				DirectoryReference ProjectMetadataDir = DirectoryReference.Combine(SC.PlatformCookDir, SC.ShortProjectName, "Metadata");
 
+				List<string> CulturesToStage = GetCulturesToStage(Params, PlatformGameConfig, false);
+
 				// The .uplugin file is staged differently for different DLC
 				// The .uplugin file doesn't actually exist for mobile DLC
 				if (FileReference.Exists(Params.DLCFile))
@@ -1445,7 +1469,6 @@ namespace AutomationScripts
 					}
 
 					// Stage DLC localization targets
-					List<string> CulturesToStage = GetCulturesToStage(Params, PlatformGameConfig, false);
 					if (CulturesToStage != null)
 					{
 						StageLocalizationDataForPlugin(SC, CulturesToStage, Params.DLCFile);
@@ -1481,6 +1504,23 @@ namespace AutomationScripts
 				// Stage all the cooked data, this is the same rule as normal stage except we may skip Engine
 				List<FileReference> CookedFiles = DiscoverCookedFiles(Params, SC, /* bAlwaysFullyScanCookDir */ false);
 
+				if (Params.StageOnlyCookedLocalizationForGameFeatures && CulturesToStage != null)
+				{
+					List<KeyValuePair<StagedFileReference, FileReference>> GameFeaturePluginsToCheckIfCookedDuringDLC = new List<KeyValuePair<StagedFileReference, FileReference>>();
+					foreach (StageTarget Target in SC.StageTargets)
+					{
+						foreach (RuntimeDependency RuntimeDependency in Target.Receipt.RuntimeDependencies)
+						{
+							if (RuntimeDependency.Path.HasExtension(".uplugin") && RuntimeDependency.Path.IsUnderDirectory(GameFeaturePluginDir))
+							{
+								GameFeaturePluginsToCheckIfCookedDuringDLC.Add(new KeyValuePair<StagedFileReference, FileReference>(SC.GetStagedFileLocation(RuntimeDependency.Path), RuntimeDependency.Path));
+								continue;
+							}
+						}
+					}
+					StageLocalizationForCookedGameFeaturePlugins(SC, CulturesToStage, GameFeaturePluginsToCheckIfCookedDuringDLC, CookedFiles);
+				}
+
 				foreach (FileReference CookedFile in CookedFiles)
 				{
 					// Skip metadata directory
@@ -1779,9 +1819,19 @@ namespace AutomationScripts
 					// Stage all project localization targets
 					StageLocalizationDataForTargetsInDirectory(SC, PlatformGameConfig, CulturesToStage, DirectoryReference.Combine(SC.ProjectRoot, "Content", "Localization"));
 
+					List<KeyValuePair<StagedFileReference, FileReference>> PluginsSkippedStagingLocalization = new List<KeyValuePair<StagedFileReference, FileReference>>();
 					// Stage all plugin localization targets
 					foreach (KeyValuePair<StagedFileReference, FileReference> StagedPlugin in StagedPlugins)
 					{
+						if (Params.StageOnlyCookedLocalizationForGameFeatures)
+						{
+							// Only stage localization for this plugin if it was actually cooked. This is checked later.
+							if (StagedPlugin.Value.IsUnderDirectory(GameFeaturePluginDir))
+							{
+								PluginsSkippedStagingLocalization.Add(StagedPlugin);
+								continue;
+							}
+						}
 						StageLocalizationDataForPlugin(SC, CulturesToStage, StagedPlugin.Value);
 					}
 
@@ -1896,6 +1946,12 @@ namespace AutomationScripts
 					{
 						List<FileReference> CookedFiles = DiscoverCookedFiles(Params, SC, /* bAlwaysFullyScanCookDir */ true);
 
+						if (Params.StageOnlyCookedLocalizationForGameFeatures)
+						{
+							// now that we have the cooked files check which plugins should have their localization staged.
+							StageLocalizationForCookedGameFeaturePlugins(SC, CulturesToStage, PluginsSkippedStagingLocalization, CookedFiles);
+						}
+
 						foreach (FileReference CookedFile in CookedFiles)
 						{
 							// Skip metadata directory
