Check if Key exists before insertion.

diff --git a/Engine/Source/Programs/AutomationTool/Scripts/CopyBuildToStagingDirectory.Automation.cs b/Engine/Source/Programs/AutomationTool/Scripts/CopyBuildToStagingDirectory.Automation.cs
--- a/Engine/Source/Programs/AutomationTool/Scripts/CopyBuildToStagingDirectory.Automation.cs
+++ b/Engine/Source/Programs/AutomationTool/Scripts/CopyBuildToStagingDirectory.Automation.cs
@@ -1406,8 +1406,16 @@ namespace AutomationScripts
 			{
 				string PluginNameWithExtension = StagedPlugin.Value.GetFileName();
 				string PluginContentDirectory = Path.Combine(SC.PlatformCookDir.FullName, StagedPlugin.Key.ToString().TrimEnd(PluginNameWithExtension.ToCharArray()), "Content");
-				GameFeatureDataToPlugin.Add(new FileReference(Path.Combine(PluginContentDirectory, "GameFeatureData.uasset")), StagedPlugin.Value);
-				GameFeatureDataToPlugin.Add(new FileReference(Path.Combine(PluginContentDirectory, $"{StagedPlugin.Value.GetFileNameWithoutAnyExtensions()}.uasset")), StagedPlugin.Value);
+				FileReference PluginGFDAsset = new FileReference(Path.Combine(PluginContentDirectory, "GameFeatureData.uasset"));
+				FileReference PluginGFDAssetWithPluginName = new FileReference(Path.Combine(PluginContentDirectory, $"{StagedPlugin.Value.GetFileNameWithoutAnyExtensions()}.uasset"));
+				if (!GameFeatureDataToPlugin.ContainsKey(PluginGFDAsset))
+				{
+					GameFeatureDataToPlugin.Add(PluginGFDAsset, StagedPlugin.Value);
+				}
+				if (!GameFeatureDataToPlugin.ContainsKey(PluginGFDAssetWithPluginName))
+				{
+					GameFeatureDataToPlugin.Add(PluginGFDAssetWithPluginName, StagedPlugin.Value);
+				}
 			}
 			foreach (FileReference CookedFile in CookedFiles)
 			{
