Report timeout as test events

#rnx

diff --git a/Engine/Source/Programs/AutomationTool/Gauntlet/Unreal/Automation/UE.BootTest.cs b/Engine/Source/Programs/AutomationTool/Gauntlet/Unreal/Automation/UE.BootTest.cs
--- a/Engine/Source/Programs/AutomationTool/Gauntlet/Unreal/Automation/UE.BootTest.cs
+++ b/Engine/Source/Programs/AutomationTool/Gauntlet/Unreal/Automation/UE.BootTest.cs
@@ -132,7 +132,7 @@ namespace UE
 			// tests from idling on the farm needlessly.
 			if ((DateTime.Now - LastLogTime).TotalSeconds > kTimeOutDuration)
 			{
-				Log.Error("No logfile activity observed in last {Time:0.00} minutes. Ending test", kTimeOutDuration / 60);
+				ReportError("No logfile activity observed in last {Time:0.00} minutes. Ending test", kTimeOutDuration / 60);
 				MarkTestComplete();
 				SetUnrealTestResult(TestResult.TimedOut);
 			}
