Add -FastIterate flag, and made build iteration on Android and in general faster

diff --git a/Engine/Source/Developer/LauncherServices/Private/Launcher/LauncherWorker.cpp b/Engine/Source/Developer/LauncherServices/Private/Launcher/LauncherWorker.cpp
--- a/Engine/Source/Developer/LauncherServices/Private/Launcher/LauncherWorker.cpp
+++ b/Engine/Source/Developer/LauncherServices/Private/Launcher/LauncherWorker.cpp
@@ -683,6 +683,11 @@ FString FLauncherWorker::CreateUATCommand( const ILauncherProfileRef& InProfile,
 		CommandStart = TEXT("********** BUILD COMMAND STARTED **********");
 		// @todo: server
 	}
+	
+	if (InProfile->IsFastIterate())
+	{
+		UATCommand += TEXT(" -FastIterate -ubtargs=\"-FastIterate\"");
+	}
 
 	// snapshot import
 	if (InProfile->IsImportingZenSnapshot())
diff --git a/Engine/Source/Developer/LauncherServices/Private/Profiles/LauncherProfile.h b/Engine/Source/Developer/LauncherServices/Private/Profiles/LauncherProfile.h
--- a/Engine/Source/Developer/LauncherServices/Private/Profiles/LauncherProfile.h
+++ b/Engine/Source/Developer/LauncherServices/Private/Profiles/LauncherProfile.h
@@ -63,6 +63,7 @@ enum ELauncherVersion
 	LAUNCHERSERVICES_ADDCUSTOMPROPERTIES = 39,
 	LAUNCHERSERVICES_INCREMENTALCOOKMODE = 40,
 	LAUNCHERSERVICES_AUTOMATEDTESTS = 41,
+	LAUNCHERSERVICES_FASTITERATE = 42,
 	//ADD NEW STUFF HERE
 
 
@@ -861,6 +862,11 @@ public:
 		return bBuild;
 	}
 
+	virtual bool IsFastIterate() const override
+	{
+		return bFastIterate;
+	}
+	
 	virtual bool IsBuildingUAT() const override
 	{
 		return BuildUAT;
@@ -1393,6 +1399,7 @@ public:
 		Writer.WriteValue("CustomStringProperties", CustomStringProperties);
 		Writer.WriteValue("CustomBoolProperties", CustomBoolProperties);
 		Writer.WriteValue("IncrementalCookMode", IncrementalCookMode);
+		Writer.WriteValue("FastIterate", bFastIterate);
 
 		// serialize the default launch role
 		DefaultLaunchRole->Save(Writer, TEXT("DefaultRole"));
@@ -2187,6 +2194,10 @@ public:
 			bWantAutomatedTestBuild = Object.GetBoolField(TEXT("WantAutomatedTestBuild"));
 		}
 
+		if (Version >= LAUNCHERSERVICES_FASTITERATE)
+		{
+			bFastIterate = Object.GetBoolField(TEXT("FastIterate"));
+		}
 
 		// load the default launch role
 		TSharedPtr<FJsonObject> Role = Object.GetObjectField(TEXT("DefaultRole"));
@@ -2263,6 +2274,7 @@ public:
 
 		// default build settings
 		BuildMode = ELauncherProfileBuildModes::Auto;
+		bFastIterate = false;
 		BuildUAT = !FApp::GetEngineIsPromotedBuild() && !FApp::IsEngineInstalled();
 		BuildTargetSpecified = false;
 
@@ -2359,6 +2371,11 @@ public:
 		}
 	}
 
+	virtual void SetFastIterate(bool Enable) override
+	{
+		bFastIterate = Enable;
+	}
+	
 	virtual void SetBuildUAT(bool Build) override
 	{
 		if (BuildUAT != Build)
@@ -3492,6 +3509,9 @@ private:
 	// Holds the build mode
 	TEnumAsByte<ELauncherProfileBuildModes::Type> BuildMode;
 
+	// Holds a flag indicating whether fast iteration is enabled 
+	bool bFastIterate;
+	
 	// Holds a flag indicating whether UAT should be built
 	bool BuildUAT;
 
diff --git a/Engine/Source/Developer/LauncherServices/Public/ILauncherProfile.h b/Engine/Source/Developer/LauncherServices/Public/ILauncherProfile.h
--- a/Engine/Source/Developer/LauncherServices/Public/ILauncherProfile.h
+++ b/Engine/Source/Developer/LauncherServices/Public/ILauncherProfile.h
@@ -880,6 +880,14 @@ public:
 	 */
 	virtual bool ShouldBuild() = 0;
 
+	/**
+	 * Checks if fast iterate is enabled.
+	 *
+	 * @return true if fast iterate is enabled, false otherwise.
+	 * @see SetFastIterate
+	 */
+	virtual bool IsFastIterate() const = 0;
+	
 	/**
 	 * Checks whether UAT should be built.
 	 *
@@ -1134,6 +1142,14 @@ public:
 	 */
 	virtual void SetBuildMode(ELauncherProfileBuildModes::Type Mode) = 0;
 
+	/**
+	 * Sets fast iterate.
+	 *
+	 * @param Enable fast iterate.
+	 * @see IsFastIterate
+	 */
+	virtual void SetFastIterate(bool Enable) = 0;
+
 	/**
 	 * Sets whether to build UAT.
 	 *
diff --git a/Engine/Source/Editor/UnrealEd/Private/PlayLevelViaLauncher.cpp b/Engine/Source/Editor/UnrealEd/Private/PlayLevelViaLauncher.cpp
--- a/Engine/Source/Editor/UnrealEd/Private/PlayLevelViaLauncher.cpp
+++ b/Engine/Source/Editor/UnrealEd/Private/PlayLevelViaLauncher.cpp
@@ -214,6 +214,7 @@ void UEditorEngine::StartPlayUsingLauncherSession(FRequestPlaySessionParams& InR
 	ILauncherProfileRef LauncherProfile = LauncherServicesModule.CreateProfile(TEXT("Launch On Device"));
 	LauncherProfile->SetBuildMode(BuildMode);
 	LauncherProfile->SetBuildConfiguration(BuildConfiguration);
+	LauncherProfile->SetFastIterate(true);
 	if (InRequestParams.EditorPlaySettings && !InRequestParams.EditorPlaySettings->AdditionalLaunchParameters.IsEmpty())
 	{
 		LauncherProfile->SetAdditionalCommandLineParameters(InRequestParams.EditorPlaySettings->AdditionalLaunchParameters);
diff --git a/Engine/Source/Programs/AutomationTool/Android/AndroidFileClient.cs b/Engine/Source/Programs/AutomationTool/Android/AndroidFileClient.cs
--- a/Engine/Source/Programs/AutomationTool/Android/AndroidFileClient.cs
+++ b/Engine/Source/Programs/AutomationTool/Android/AndroidFileClient.cs
@@ -2246,6 +2246,74 @@ namespace AutomationTool
 			return Result;
 		}
 
+		public string FileReadString(string SourcePath)
+		{
+			if (ClientSocket == null)
+			{
+				return null;
+			}
+
+			try
+			{
+				// deal with LastBaseDir substitution
+				SourcePath = OptimizePath(SourcePath);
+
+				{
+					byte[] paramMessage = Encoding.UTF8.GetBytes(SourcePath + "\0");
+					long paramSize = paramMessage.Length;
+					byte[] commandMessage = CommandPacket(Command_FileRead, paramSize);
+
+					SocketSend(false, commandMessage);
+					SocketSend(false, paramMessage);
+				}
+
+				long resultSize = ResultSize();
+				if (resultSize >= 0)
+				{
+					byte[] buffer = new byte[FILE_BUFFERSIZE];
+
+					Stats_FilesRead++;
+
+					try
+					{
+						MemoryStream stringStream = new MemoryStream();
+						
+						for (long remaining = resultSize; remaining > 0;)
+						{
+							int readSize = remaining > FILE_BUFFERSIZE ? FILE_BUFFERSIZE : (int)remaining;
+							int bytesRead = SocketReceive(true, buffer, readSize);
+							remaining -= bytesRead;
+							stringStream.Write(buffer, 0, bytesRead);
+							Stats_CompressedBytesReceived += bytesRead;
+						}
+
+						stringStream.Position = 0;
+						return new StreamReader(stringStream).ReadToEnd();
+					}
+					catch (Exception e)
+					{
+						Logger.LogError("Unexpected Exception: {Arg0}", e.ToString());
+						
+						// ignore received data
+						for (long remainingSkip = resultSize; remainingSkip > 0;)
+						{
+							int readSize = remainingSkip > FILE_BUFFERSIZE ? FILE_BUFFERSIZE : (int)remainingSkip;
+							int bytesRead = SocketReceive(true, buffer, readSize);
+							remainingSkip -= bytesRead;
+						}
+					}
+				}
+			}
+			catch (Exception e)
+			{
+				Logger.LogError("Unexpected Exception: {Arg0}", e.ToString());
+				
+				CloseConnection();
+			}
+			
+			return null;
+		}
+		
 		private const int RetriesMax = 5;
 		private const int RetryMillisecondsDelay = 5000;
 
diff --git a/Engine/Source/Programs/AutomationTool/Android/AndroidPlatform.Automation.cs b/Engine/Source/Programs/AutomationTool/Android/AndroidPlatform.Automation.cs
--- a/Engine/Source/Programs/AutomationTool/Android/AndroidPlatform.Automation.cs
+++ b/Engine/Source/Programs/AutomationTool/Android/AndroidPlatform.Automation.cs
@@ -2572,17 +2572,12 @@ public class AndroidPlatform : Platform
 		}
 	}
 
-	private bool GetDontBundleLibrariesInAPK(ProjectParams Params, DeploymentContext SC, bool bVerbose = false)
-	{
-		return AndroidExports.GetDontBundleLibrariesInAPK(Params.RawProjectPath, null, SC.StageTargets[0].Receipt.Configuration, SC.Archive, false,
-			true, bVerbose ? Logger : null);
-	}
-
 	private void DeployAndroidFileServer(ProjectParams Params, DeploymentContext SC, string AFSToken)
 	{
 		ConfigHierarchy Ini = ConfigCache.ReadHierarchy(ConfigHierarchyType.Engine, DirectoryReference.FromFile(Params.RawProjectPath), UnrealTargetPlatform.Android);
 		Ini.GetBool("/Script/AndroidRuntimeSettings.AndroidRuntimeSettings", "bEnableMaliPerfCounters", out bool bDisablePerfHarden);
-		bool bDontBundleLibrariesInAPK = GetDontBundleLibrariesInAPK(Params, SC, true);
+		IAndroidDeploy Deploy = AndroidExports.CreateDeploymentHandler(Params.RawProjectPath, Params.ForcePackageData);
+		bool bDontBundleLibrariesInAPK = Deploy.GetDontBundleLibrariesInAPK(SC.StageTargets[0].Receipt.Configuration, SC.Archive, true);
 
 		bool bUseCompression;
 		bool bLogFiles;
@@ -2601,8 +2596,7 @@ public class AndroidPlatform : Platform
 			string FinalSOName = GetSOName(Params, SC.StageExecutables[0], DeviceArchitecture);
 
 			// make sure APK is up to date (this is fast if so)
-			var Deploy = AndroidExports.CreateDeploymentHandler(Params.RawProjectPath, Params.ForcePackageData);
-			if (!Params.Prebuilt)
+			if (!Params.Prebuilt && (Params.Build || !Deploy.bFastIterate))
 			{
 				string CookFlavor = SC.FinalCookPlatform.IndexOf("_") > 0 ? SC.FinalCookPlatform.Substring(SC.FinalCookPlatform.IndexOf("_")) : "";
 				string SOName = GetSONameWithoutArchitecture(Params, SC.StageExecutables[0]);
@@ -2662,7 +2656,7 @@ public class AndroidPlatform : Platform
 			bool bNeedAPKInstall = true;
 			bool bFreshInstall = false;
 			bool bFastDeploy = false;
-			if (Params.IterativeDeploy)
+			if (Params.IterativeDeploy && !Deploy.bFromMSBuild)
 			{
 				// Check for apk installed with this package name on the device
 				String InstalledResult = adb.Shell(DeviceName, "pm list packages " + PackageName);
@@ -2710,44 +2704,17 @@ public class AndroidPlatform : Platform
 			// install new APK if needed
 			if (bNeedAPKInstall)
 			{
-				// try reinstall the apk to preserve data
-				int SuccessCode = 0;
-				string InstallCommandline = "install -r " + (bFastDeploy ? "--fastdeploy \"" : "\"") + ApkName + "\"";
-				string InstallOutput = RunAndLogAdbCommand(Params, DeviceName, InstallCommandline, out SuccessCode);
-				int FailureIndex = InstallOutput.IndexOf("Failure");
-				// adb install doesn't always return an error code on failure, and instead prints "Failure", followed by an error code.
-				if (SuccessCode != 0 || FailureIndex != -1)
+				if (!Deploy.bFromMSBuild)
 				{
-					string ErrorMessage = string.Format("Installation of apk '{0}' failed", ApkName);
-					if (FailureIndex != -1)
-					{
-						string FailureString = InstallOutput.Substring(FailureIndex + 7).Trim();
-						if (FailureString != "")
-						{
-							ErrorMessage += ": " + FailureString;
-						}
-					}
-					if (ErrorMessage.Contains("OLDER_SDK"))
-					{
-						Logger.LogError("minSdkVersion is higher than Android version installed on device, possibly due to NDK API Level");
-						throw new AutomationException(ExitCode.Error_AppInstallFailed, ErrorMessage);
-					}
-
-					// try uninstalling an old app with the same identifier.
-					// NOTE: uninstall -k will preserve data/cache.. consider using instead?
-					bFreshInstall = true;
-					SuccessCode = 0;
-					adb.Shell(DeviceName, "pm uninstall " + PackageName);
-
-					// install the apk
-					InstallCommandline = "install \"" + ApkName + "\"";
-					InstallOutput = RunAndLogAdbCommand(Params, DeviceName, InstallCommandline, out SuccessCode);
-					FailureIndex = InstallOutput.IndexOf("Failure");
-
+					// try reinstall the apk to preserve data
+					int SuccessCode = 0;
+					string InstallCommandline = "install -r " + (bFastDeploy ? "--fastdeploy \"" : "\"") + ApkName + "\"";
+					string InstallOutput = RunAndLogAdbCommand(Params, DeviceName, InstallCommandline, out SuccessCode);
+					int FailureIndex = InstallOutput.IndexOf("Failure");
 					// adb install doesn't always return an error code on failure, and instead prints "Failure", followed by an error code.
 					if (SuccessCode != 0 || FailureIndex != -1)
 					{
-						ErrorMessage = string.Format("Installation of apk '{0}' failed", ApkName);
+						string ErrorMessage = string.Format("Installation of apk '{0}' failed", ApkName);
 						if (FailureIndex != -1)
 						{
 							string FailureString = InstallOutput.Substring(FailureIndex + 7).Trim();
@@ -2759,8 +2726,38 @@ public class AndroidPlatform : Platform
 						if (ErrorMessage.Contains("OLDER_SDK"))
 						{
 							Logger.LogError("minSdkVersion is higher than Android version installed on device, possibly due to NDK API Level");
+							throw new AutomationException(ExitCode.Error_AppInstallFailed, ErrorMessage);
+						}
+
+						// try uninstalling an old app with the same identifier.
+						// NOTE: uninstall -k will preserve data/cache.. consider using instead?
+						bFreshInstall = true;
+						SuccessCode = 0;
+						adb.Shell(DeviceName, "pm uninstall " + PackageName);
+
+						// install the apk
+						InstallCommandline = "install \"" + ApkName + "\"";
+						InstallOutput = RunAndLogAdbCommand(Params, DeviceName, InstallCommandline, out SuccessCode);
+						FailureIndex = InstallOutput.IndexOf("Failure");
+
+						// adb install doesn't always return an error code on failure, and instead prints "Failure", followed by an error code.
+						if (SuccessCode != 0 || FailureIndex != -1)
+						{
+							ErrorMessage = string.Format("Installation of apk '{0}' failed", ApkName);
+							if (FailureIndex != -1)
+							{
+								string FailureString = InstallOutput.Substring(FailureIndex + 7).Trim();
+								if (FailureString != "")
+								{
+									ErrorMessage += ": " + FailureString;
+								}
+							}
+							if (ErrorMessage.Contains("OLDER_SDK"))
+							{
+								Logger.LogError("minSdkVersion is higher than Android version installed on device, possibly due to NDK API Level");
+							}
+							throw new AutomationException(ExitCode.Error_AppInstallFailed, ErrorMessage);
 						}
-						throw new AutomationException(ExitCode.Error_AppInstallFailed, ErrorMessage);
 					}
 				}
 
@@ -2858,17 +2855,44 @@ public class AndroidPlatform : Platform
 			RemoteDir = client.Query("^project", true);
 
 			// write new timestamp for APK (do it here since RemoteDir now available)
-			if (bNeedAPKInstall)
+			if (bNeedAPKInstall && !Deploy.bFromMSBuild)
 			{
 				client.FileWriteString("APK: " + APKLastUpdateTime + "\n", "^ext/APKFileStamp.txt");
 			}
 
-			// always update libUnreal.so
-			// TODO potential optimization not to push it every time but compare filestamp instead to check if we need to update it
 			if (bDontBundleLibrariesInAPK)
 			{
 				string FinalSONameStripped = Path.Combine(Path.GetDirectoryName(FinalSOName), Path.GetFileNameWithoutExtension(FinalSOName) + "-stripped" + Path.GetExtension(FinalSOName));
-				client.PushFile(FinalSONameStripped, "^int/libUnreal.so", true);
+
+				long NewTimestamp = ((DateTimeOffset)File.GetLastWriteTimeUtc(FinalSONameStripped)).ToUnixTimeMilliseconds();
+				
+				if (!long.TryParse(client.FileReadString("^int/libUnreal.so.timestamp") ?? "", out long OldTimestamp) || OldTimestamp != NewTimestamp)
+				{
+					Logger.LogInformation("libUnreal.so changed...pushing to device");
+					client.PushFile(FinalSONameStripped, "^int/libUnreal.so");
+					client.FileWriteString(NewTimestamp.ToString(), "^int/libUnreal.so.timestamp");
+				}
+			}
+			else
+			{
+				client.FileDelete("^int/libUnreal.so.timestamp");
+				client.FileDelete("^int/libUnreal.so");
+			}
+
+			if (Deploy.bFromMSBuild)
+			{
+				// terminate server and close AFS connections
+				if (client != null)
+				{
+					client.TerminateServer();
+					client.CloseConnection();
+				}
+				if (client2 != null)
+				{
+					client2.TerminateServer();
+					client2.CloseConnection();
+				}
+				continue;
 			}
 
 			// update the uecommandline.txt
@@ -3208,7 +3232,8 @@ public class AndroidPlatform : Platform
     {
 		ConfigHierarchy Ini = ConfigCache.ReadHierarchy(ConfigHierarchyType.Engine, DirectoryReference.FromFile(Params.RawProjectPath), UnrealTargetPlatform.Android);
 		Ini.GetBool("/Script/AndroidRuntimeSettings.AndroidRuntimeSettings", "bEnableMaliPerfCounters", out bool bDisablePerfHarden);
-		bool bDontBundleLibrariesInAPK = GetDontBundleLibrariesInAPK(Params, SC, true);
+		IAndroidDeploy Deploy = AndroidExports.CreateDeploymentHandler(Params.RawProjectPath, Params.ForcePackageData);
+		bool bDontBundleLibrariesInAPK = Deploy.GetDontBundleLibrariesInAPK(SC.StageTargets[0].Receipt.Configuration, SC.Archive, true);
 		int AllowOverflowOBBLimit = AllowOverflowOBBFiles(SC);
 
 		foreach (string DeviceName in Params.DeviceNames.Select(DeviceName => GetSerialNumber(DeviceName)))
@@ -3218,8 +3243,7 @@ public class AndroidPlatform : Platform
             string FinalSOName = GetSOName(Params, SC.StageExecutables[0], DeviceArchitecture);
 
             // make sure APK is up to date (this is fast if so)
-            var Deploy = AndroidExports.CreateDeploymentHandler(Params.RawProjectPath, Params.ForcePackageData);
-            if (!Params.Prebuilt)
+            if (!Params.Prebuilt && (Params.Build || !Deploy.bFastIterate))
             {
                 string CookFlavor = SC.FinalCookPlatform.IndexOf("_") > 0 ? SC.FinalCookPlatform.Substring(SC.FinalCookPlatform.IndexOf("_")) : "";
 				string SOName = GetSONameWithoutArchitecture(Params, SC.StageExecutables[0]);
@@ -3255,7 +3279,7 @@ public class AndroidPlatform : Platform
             // determine if APK out of date
             string APKLastUpdateTime = new FileInfo(ApkName).LastWriteTime.ToString();
             bool bNeedAPKInstall = true;
-            if (Params.IterativeDeploy)
+            if (Params.IterativeDeploy && !Deploy.bFromMSBuild)
             {
                 // Check for apk installed with this package name on the device
                 IProcessResult InstalledResult = RunAdbCommand(Params, DeviceName, "shell pm list packages " + PackageName, null, ERunOptions.AppMustExist);
@@ -3287,7 +3311,7 @@ public class AndroidPlatform : Platform
             }
 
             // install new APK if needed
-            if (bNeedAPKInstall)
+            if (bNeedAPKInstall && !Deploy.bFromMSBuild)
             {
                 // try uninstalling an old app with the same identifier.
                 int SuccessCode = 0;
@@ -3330,24 +3354,41 @@ public class AndroidPlatform : Platform
 
 			if (bDontBundleLibrariesInAPK)
 			{
-				int SuccessCode = 0;
-
 				string FinalSOPathStripped = Path.Combine(Path.GetDirectoryName(FinalSOName), Path.GetFileNameWithoutExtension(FinalSOName) + "-stripped" + Path.GetExtension(FinalSOName));
 				string FinalSOFileNameStripped = Path.GetFileName(FinalSOPathStripped);
-
-				string PushSO = $"push -z lz4 {FinalSOPathStripped} /data/local/tmp/{FinalSOFileNameStripped}";
-				string CopySO = $"shell run-as {PackageName} cp /data/local/tmp/{FinalSOFileNameStripped} ./files/libUnreal.so";
-				string DeleteSO = $"shell rm /data/local/tmp/{FinalSOFileNameStripped}";
-				RunAndLogAdbCommand(Params, DeviceName, PushSO, out SuccessCode);
-				RunAndLogAdbCommand(Params, DeviceName, CopySO, out SuccessCode);
-				RunAndLogAdbCommand(Params, DeviceName, DeleteSO, out SuccessCode);
-
-				if (SuccessCode != 0)
+				
+				long NewTimestamp = ((DateTimeOffset)File.GetLastWriteTimeUtc(FinalSOPathStripped)).ToUnixTimeMilliseconds();
+				
+				if (!long.TryParse(RunAdbCommand(Params, DeviceName, "shell cat ./files/libUnreal.so.timestamp", null, ERunOptions.AppMustExist).Output, out long OldTimestamp) || OldTimestamp != NewTimestamp)
 				{
-					string ErrorMessage = $"Installation of '{PackageName}' failed due to failing to push libUnreal.so outside of {ApkName}";
-					throw new AutomationException(ExitCode.Error_AppInstallFailed, ErrorMessage);
+					int SuccessCode = 0;
+					string PushSO = $"push -z lz4 {FinalSOPathStripped} /data/local/tmp/{FinalSOFileNameStripped}";
+					string CopySO = $"shell run-as {PackageName} cp /data/local/tmp/{FinalSOFileNameStripped} ./files/libUnreal.so";
+					string DeleteSO = $"shell rm /data/local/tmp/{FinalSOFileNameStripped}";
+					RunAndLogAdbCommand(Params, DeviceName, PushSO, out SuccessCode);
+					RunAndLogAdbCommand(Params, DeviceName, CopySO, out SuccessCode);
+					RunAndLogAdbCommand(Params, DeviceName, DeleteSO, out SuccessCode);
+
+					if (SuccessCode != 0)
+					{
+						string ErrorMessage = $"Installation of '{PackageName}' failed due to failing to push libUnreal.so outside of {ApkName}";
+						throw new AutomationException(ExitCode.Error_AppInstallFailed, ErrorMessage);
+					}
+					
+					RunAndLogAdbCommand(Params, DeviceName, $"shell \"echo '{NewTimestamp}' > ./files/libUnreal.so.timestamp\"", out SuccessCode);
 				}
 			}
+			else
+			{
+				int SuccessCode;
+				RunAndLogAdbCommand(Params, DeviceName, "shell rm ./files/libUnreal.so.timestamp", out SuccessCode);
+				RunAndLogAdbCommand(Params, DeviceName, "shell rm ./files/libUnreal.so", out SuccessCode);
+			}
+
+			if (Deploy.bFromMSBuild)
+			{
+				continue;
+			}
 
             // update the uecommandline.txt
             // update and deploy uecommandline.txt
diff --git a/Engine/Source/Programs/AutomationTool/AutomationUtils/UnrealBuild.cs b/Engine/Source/Programs/AutomationTool/AutomationUtils/UnrealBuild.cs
--- a/Engine/Source/Programs/AutomationTool/AutomationUtils/UnrealBuild.cs
+++ b/Engine/Source/Programs/AutomationTool/AutomationUtils/UnrealBuild.cs
@@ -148,21 +148,31 @@ namespace AutomationTool
 
 		void BuildWithUBT(List<BuildTarget> Targets, Dictionary<BuildTarget, BuildManifest> TargetToManifest, bool DisableXGE, bool AllCores, bool SkipBuild)
 		{
-			List<FileReference> ManifestFiles = new List<FileReference>(Targets.Count);
+			List<FileReference> ManifestFiles = TargetToManifest != null ? new List<FileReference>(Targets.Count) : null;
 
 			StringBuilder FullCommandLine = new StringBuilder();
 			if (Targets.Count == 1)
 			{
-				ManifestFiles.Add(GetManifestFile(Targets[0], -1));
-				FullCommandLine.Append(GetTargetArguments(Targets[0], ManifestFiles[0]));
+				FileReference ManifestFile = null;
+				if (TargetToManifest != null)
+				{
+					ManifestFiles.Add(GetManifestFile(Targets[0], -1));
+					ManifestFile = ManifestFiles[0];
+				}
+				FullCommandLine.Append(GetTargetArguments(Targets[0], ManifestFile));
 			}
 			else
 			{
 				List<string> Arguments = new List<string>();
 				for (int Idx = 0; Idx < Targets.Count; Idx++)
 				{
-					ManifestFiles.Add(GetManifestFile(Targets[Idx], Idx));
-					Arguments.Add("-Target=" + GetTargetArguments(Targets[Idx], ManifestFiles[Idx]));
+					FileReference ManifestFile = null;
+					if (TargetToManifest != null)
+					{
+						ManifestFiles.Add(GetManifestFile(Targets[Idx], Idx));
+						ManifestFile = ManifestFiles[Idx];
+					}
+					Arguments.Add("-Target=" + GetTargetArguments(Targets[Idx], ManifestFile));
 				}
 				FullCommandLine.Append(CommandLine.FormatCommandLine(Arguments));
 			}
@@ -182,19 +192,25 @@ namespace AutomationTool
 
 			PrepareUBT();
 
-			for (int Idx = 0; Idx < ManifestFiles.Count; Idx++)
+			if (TargetToManifest != null)
 			{
-				CommandUtils.DeleteFile(ManifestFiles[Idx]);
+				for (int Idx = 0; Idx < ManifestFiles.Count; Idx++)
+				{
+					CommandUtils.DeleteFile(ManifestFiles[Idx]);
+				}
 			}
 
 			CommandUtils.RunUBT(CommandUtils.CmdEnv, UnrealBuildToolDll, FullCommandLine.ToString());
 
-			for (int Idx = 0; Idx < Targets.Count; Idx++)
+			if (TargetToManifest != null)
 			{
-				FileReference ManifestFile = ManifestFiles[Idx];
-				BuildManifest Manifest = AddBuildProductsFromManifest(ManifestFile);
-				CommandUtils.DeleteFile(ManifestFile);
-				TargetToManifest?.Add(Targets[Idx], Manifest);
+				for (int Idx = 0; Idx < Targets.Count; Idx++)
+				{
+					FileReference ManifestFile = ManifestFiles[Idx];
+					BuildManifest Manifest = AddBuildProductsFromManifest(ManifestFile);
+					CommandUtils.DeleteFile(ManifestFile);
+					TargetToManifest?.Add(Targets[Idx], Manifest);
+				}
 			}
 		}
 
diff --git a/Engine/Source/Programs/UnrealBuildTool/Platform/Android/AndroidExports.cs b/Engine/Source/Programs/UnrealBuildTool/Platform/Android/AndroidExports.cs
--- a/Engine/Source/Programs/UnrealBuildTool/Platform/Android/AndroidExports.cs
+++ b/Engine/Source/Programs/UnrealBuildTool/Platform/Android/AndroidExports.cs
@@ -29,12 +29,30 @@ namespace UnrealBuildTool
 	/// </summary>
 	public interface IAndroidDeploy
 	{
+		/// <summary>
+		/// 
+		/// </summary>
+		bool bFromMSBuild { get; }
+		
+		/// <summary>
+		/// 
+		/// </summary>
+		bool bFastIterate { get; }
+		
 		/// <summary>
 		/// 
 		/// </summary>
 		/// <returns></returns>
 		bool GetPackageDataInsideApk();
 
+		/// <summary>
+		/// 
+		/// </summary>
+		/// <param name="Configuration"></param>
+		/// <param name="bIsArchive"></param>
+		/// <param name="bIsFromUAT"></param>
+		bool GetDontBundleLibrariesInAPK(UnrealTargetConfiguration Configuration, bool bIsArchive, bool bIsFromUAT);
+		
 		/// <summary>
 		/// 
 		/// </summary>
@@ -116,21 +134,6 @@ namespace UnrealBuildTool
 			ToolChain.StripSymbols(SourceFile, TargetFile, Logger);
 		}
 
-		/// <summary>
-		/// 
-		/// </summary>
-		/// <param name="ProjectFile"></param>
-		/// <param name="bForceDontBundleLibrariesInAPK"></param>
-		/// <param name="Configuration"></param>
-		/// <param name="bIsArchive"></param>
-		/// <param name="bFromMSBuild"></param>
-		/// <param name="bIsFromUAT"></param>
-		/// <param name="Logger"></param>
-		public static bool GetDontBundleLibrariesInAPK(FileReference? ProjectFile, bool? bForceDontBundleLibrariesInAPK, UnrealTargetConfiguration Configuration, bool bIsArchive, bool bFromMSBuild, bool bIsFromUAT, ILogger? Logger)
-		{
-			return UEDeployAndroid.GetDontBundleLibrariesInAPK(ProjectFile, bForceDontBundleLibrariesInAPK, Configuration, bIsArchive, bFromMSBuild, bIsFromUAT, Logger);
-		}
-
 		/// <summary>
 		/// 
 		/// </summary>
diff --git a/Engine/Source/Programs/UnrealBuildTool/Platform/Android/AndroidProjectGenerator.cs b/Engine/Source/Programs/UnrealBuildTool/Platform/Android/AndroidProjectGenerator.cs
--- a/Engine/Source/Programs/UnrealBuildTool/Platform/Android/AndroidProjectGenerator.cs
+++ b/Engine/Source/Programs/UnrealBuildTool/Platform/Android/AndroidProjectGenerator.cs
@@ -180,8 +180,7 @@ namespace UnrealBuildTool
 		{
 			if (AGDEInstalled)
 			{
-				// do not need to check InPlatform since it will always be UnrealTargetPlatform.Android
-				return $" -ForceAPKGeneration" + base.GetExtraBuildArguments(InVSSettings);
+				return $" -ForceAPKGeneration -FastIterate" + base.GetExtraBuildArguments(InVSSettings);
 			}
 			else
 			{
@@ -207,11 +206,8 @@ namespace UnrealBuildTool
 					// It's critical to have AndroidDebugTarget here and for it to be before properties that use it (e.g. AndroidPostApkInstallCommands), otherwise MSBuild will evaluate it to empty string.
 					Out.AppendLine("    <AndroidDebugTarget></AndroidDebugTarget>");
 
-					// At this stage we don't know if bDontBundleLibrariesInAPK is enabled or not, so make a fail-safe check.
-					string shortAbi = GetShortAbi(InVSSettings);
-					string PushSOScript = Path.Combine(
-						Path.GetDirectoryName(NMakeOutputPath.FullName)!,
-						"Push_" + Path.GetFileNameWithoutExtension(NMakeOutputPath.FullName) + $"-{shortAbi}_so.bat");
+					string UatPath = Path.Combine(Unreal.RootDirectory.FullName, "RunUAT.bat");
+					string ShortAbi = GetShortAbi(InVSSettings);
 
 					// AGDE specifies current debug target in AndroidDebugTarget property in a form of "model:serial:arch".
 					// AndroidDebugTarget is a special property and needs to be evaluated in-line. And the push script needs the device serial as first argument to push to the correct device. 
@@ -220,7 +216,7 @@ namespace UnrealBuildTool
 					// so regardless of what value it has, Split(':')[1] will never throw an exception.
 					string GetTargetDeviceSerial = "$([System.String]::Concat($(AndroidDebugTarget), \"::\").Split(':')[1])";
 					Out.AppendLine(
-						$"    <AndroidPostApkInstallCommands>IF EXIST {PushSOScript} {PushSOScript} {GetTargetDeviceSerial};$(AndroidPostApkInstallCommands)</AndroidPostApkInstallCommands>");
+						$"    <AndroidPostApkInstallCommands>CALL \"{UatPath}\" BuildCookRun -project=\"$(SolutionDir){ProjectName}.uproject\" -FastIterate -FromMsBuild -client -skipbuild -skipcook -skipstage -deploy -config={InVSSettings.Configuration} -platform={InVSSettings.Platform} -clientarchitecture={ShortAbi} -device=\"{GetTargetDeviceSerial}\";$(AndroidPostApkInstallCommands)</AndroidPostApkInstallCommands>");
 
 					// Ensure the following properties are up to date even if .vcxproj.user already exists and are in correct order between each other.
 					VCUserFileSettings.PatchProperty("AndroidDebugTarget", true);
diff --git a/Engine/Source/Programs/UnrealBuildTool/Platform/Android/UEDeployAndroid.cs b/Engine/Source/Programs/UnrealBuildTool/Platform/Android/UEDeployAndroid.cs
--- a/Engine/Source/Programs/UnrealBuildTool/Platform/Android/UEDeployAndroid.cs
+++ b/Engine/Source/Programs/UnrealBuildTool/Platform/Android/UEDeployAndroid.cs
@@ -88,6 +88,7 @@ namespace UnrealBuildTool
 		/// </summary>
 		[CommandLine(Prefix = "-FromMsBuild", Value = "true")]
 		public bool bFromMSBuild = false;
+		bool IAndroidDeploy.bFromMSBuild => bFromMSBuild;
 
 		/// <summary>
 		/// Forcing bDontBundleLibrariesInAPK to "true" or "false" ignoring any other option.
@@ -101,6 +102,13 @@ namespace UnrealBuildTool
 		[CommandLine("-ScudoMemoryTracing", Value = "true")]
 		public bool bEnableScudoMemoryTracing = false;
 
+		/// <summary>
+		/// Makes build interation faster 
+		/// </summary>
+		[CommandLine("-FastIterate", Value = "true")]
+		public bool bFastIterate = false;
+		bool IAndroidDeploy.bFastIterate => bFastIterate;
+
 		public UEDeployAndroid(FileReference? InProjectFile, bool InForcePackageData, ILogger InLogger)
 			: base(InLogger)
 		{
@@ -3698,7 +3706,7 @@ namespace UnrealBuildTool
 
 		private bool BundleEnabled()
 		{
-			if (ForceAPKGeneration)
+			if (bFastIterate || ForceAPKGeneration)
 			{
 				return false;
 			}
@@ -4365,14 +4373,20 @@ dependencies.implementation('androidx.tracing:tracing:1.2.0')
 			File.WriteAllText(GradleBuildScriptAdditionsFilename, UPL.ProcessPluginNode(NDKArch, "buildscriptGradleAdditions", ""));
 		}
 
-		public static bool GetDontBundleLibrariesInAPK(FileReference? ProjectFile, bool? bForceDontBundleLibrariesInAPK, UnrealTargetConfiguration Configuration, bool bIsArchive, bool bFromMSBuild, bool bIsFromUAT, ILogger? Logger)
+		public bool GetDontBundleLibrariesInAPK(UnrealTargetConfiguration Configuration, bool bIsArchive, bool bIsFromUAT)
 		{
-			if (bForceDontBundleLibrariesInAPK.HasValue)
+			if (ForceDontBundleLibrariesInAPK.HasValue)
 			{
-				Logger?.LogInformation("bDontBundleLibrariesInAPK is force set to {ForceDontBundleLibrariesInAPK} via command line option.", bForceDontBundleLibrariesInAPK.Value);
-				return bForceDontBundleLibrariesInAPK.Value;
+				Logger?.LogInformation("bDontBundleLibrariesInAPK is force set to {ForceDontBundleLibrariesInAPK} via command line option.", ForceDontBundleLibrariesInAPK.Value);
+				return ForceDontBundleLibrariesInAPK.Value;
 			}
 
+			if (bFastIterate)
+			{
+				return true;
+			}
+
+			// Deprecated in 5.7, remove everything below in 5.9 and just return bFastIterate
 			// this feature requires extra step to install the app besides just .apk install
 			// limit where the feature is enabled because we're not in control of deployment step but in VS/AGDE and UAT
 			if (!(bFromMSBuild || bIsFromUAT))
@@ -4399,11 +4413,6 @@ dependencies.implementation('androidx.tracing:tracing:1.2.0')
 			return bDontBundleLibrariesInAPK;
 		}
 
-		private bool GetDontBundleLibrariesInAPK(UnrealTargetConfiguration Configuration, bool bIsArchive, bool bIsFromUAT, bool bVerbose = false)
-		{
-			return GetDontBundleLibrariesInAPK(ProjectFile, ForceDontBundleLibrariesInAPK, Configuration, bIsArchive, bFromMSBuild, bIsFromUAT, bVerbose ? Logger : null);
-		}
-
 		private bool GetJavaPerfettoMarkers(ConfigHierarchy Ini, UnrealTargetConfiguration Configuration, bool bDistribution)
 		{
 			// Ship for distribution should not include markers
@@ -4436,7 +4445,7 @@ dependencies.implementation('androidx.tracing:tracing:1.2.0')
 			}
 
 			// we do not need to really build an engine UnrealGame.apk so short-circuit it
-			if (!ForceAPKGeneration && ProjectName == "UnrealGame" && OutputPath.Replace("\\", "/").Contains("/Engine/Binaries/Android/") && Path.GetFileNameWithoutExtension(OutputPath).StartsWith("UnrealGame"))
+			if (!bFastIterate && !ForceAPKGeneration && ProjectName == "UnrealGame" && OutputPath.Replace("\\", "/").Contains("/Engine/Binaries/Android/") && Path.GetFileNameWithoutExtension(OutputPath).StartsWith("UnrealGame"))
 			{
 				if (!bSkipGradleBuild)
 				{
@@ -4728,7 +4737,7 @@ dependencies.implementation('androidx.tracing:tracing:1.2.0')
 				}
 			}
 
-			bool bDontBundleLibrariesInAPK = GetDontBundleLibrariesInAPK(Configuration, bIsArchive, bIsFromUAT, bVerbose: !bIsFromUAT);
+			bool bDontBundleLibrariesInAPK = GetDontBundleLibrariesInAPK(Configuration, bIsArchive, bIsFromUAT);
 
 			string ZenProjectStoreFileName = Path.Combine(Path.GetDirectoryName(ObbFileLocation)!, Path.GetFileNameWithoutExtension(ObbFileLocation), "ue.projectstore");
 			bool bUsesZenStreaming = File.Exists(ZenProjectStoreFileName);
@@ -4870,13 +4879,21 @@ dependencies.implementation('androidx.tracing:tracing:1.2.0')
 				foreach ((UnrealArch Arch, string _) in Archs)
 				{
 					string SOName = AndroidToolChain.InlineArchName(OutputPath, Arch);
-					if (!File.Exists(SOName))
+					FileInfo SOFile = new(SOName);
+					if (!SOFile.Exists)
 					{
 						Logger.LogWarning("Did not find compiled .so [{SOName}]", SOName);
+						continue;
 					}
 
 					string SONameStripped = Path.Combine(Path.GetDirectoryName(SOName)!, Path.GetFileNameWithoutExtension(SOName) + "-stripped" + Path.GetExtension(SOName));
-					StripDebugSymbols(SOName, SONameStripped, Arch, Logger, true);
+					FileInfo SOFileStripped = new(SONameStripped);
+					if (!SOFileStripped.Exists || SOFileStripped.CreationTime != SOFile.CreationTime || SOFileStripped.LastWriteTime != SOFile.LastWriteTime)
+					{
+						StripDebugSymbols(SOName, SONameStripped, Arch, Logger, true);
+						SOFileStripped.CreationTime = SOFile.CreationTime;
+						SOFileStripped.LastWriteTime = SOFile.LastWriteTime;
+					}
 
 					// LLDB (in Android Studio and AGDE) needs to see a file that is matching the filename on the device ("libUnreal.so") to resolve symbols,
 					// otherwise it will try to download libUnreal.so and fail due to adb pull not having permissions to access it.
@@ -4890,6 +4907,18 @@ dependencies.implementation('androidx.tracing:tracing:1.2.0')
 					LLDBExtraSymbolsDirectories.Add(LLDBExtraSymbolsDirectory);
 				}
 			}
+			else
+			{
+				foreach ((UnrealArch Arch, string _) in Archs)
+				{
+					string SOName = AndroidToolChain.InlineArchName(OutputPath, Arch);
+					string SONameStripped = Path.Combine(Path.GetDirectoryName(SOName)!, Path.GetFileNameWithoutExtension(SOName) + "-stripped" + Path.GetExtension(SOName));
+					string LLDBExtraSymbolsDirectory = Path.Combine(IntermediateAndroidPath, "LLDBSymbolsLibs", Arch.ToString().ToLower());
+					
+					SafeDeleteFile(SONameStripped);
+					SafeDeleteFile(Path.Combine(LLDBExtraSymbolsDirectory, "libUnreal.so"));
+				}
+			}
 
 			// Now we have to spin over all the arch/gpu combinations to make sure they all match
 			int BuildListComboTotal = BuildList.Length;
