Android File Server enabled for low level tests

#rnx

diff --git a/Engine/Build/Android/Java/src/com/epicgames/unreal/GameApplication.java.template b/Engine/Build/Android/Java/src/com/epicgames/unreal/GameApplication.java.template
--- a/Engine/Build/Android/Java/src/com/epicgames/unreal/GameApplication.java.template
+++ b/Engine/Build/Android/Java/src/com/epicgames/unreal/GameApplication.java.template
@@ -16,7 +16,7 @@ import java.util.concurrent.locks.ReentrantLock;
 import com.epicgames.unreal.network.NetworkChangedManager;
 //$${gameApplicationImportAdditions}$$
 
-public class GameApplication extends Application implements LifecycleObserver
+public class GameApplication extends Application implements LifecycleObserver //$${gameApplicationInterfaceAdditions}$$
 {
 	public static final boolean IS_SHIPPING_CONFIG = $${gameApplicationIsShippingConfig}$$;
 	
diff --git a/Engine/Plugins/Runtime/AndroidFileServer/Source/AndroidFileServer/Android/java/com/epicgames/unreal/IAndroidFileServerSetup.java b/Engine/Plugins/Runtime/AndroidFileServer/Source/AndroidFileServer/Android/java/com/epicgames/unreal/IAndroidFileServerSetup.java
new file mode 100644
--- /dev/null
+++ b/Engine/Plugins/Runtime/AndroidFileServer/Source/AndroidFileServer/Android/java/com/epicgames/unreal/IAndroidFileServerSetup.java
@@ -0,0 +1,9 @@
+// Copyright Epic Games, Inc. All Rights Reserved.
+
+package com.epicgames.unreal;
+
+public interface IAndroidFileServerSetup
+{
+	public boolean AndroidFileServer_Init(String filename);
+	public boolean AndroidFileServer_Verify(String Token);
+}
diff --git a/Engine/Plugins/Runtime/AndroidFileServer/Source/AndroidFileServer/Android/java/com/epicgames/unreal/RemoteFileManagerActivity.java b/Engine/Plugins/Runtime/AndroidFileServer/Source/AndroidFileServer/Android/java/com/epicgames/unreal/RemoteFileManagerActivity.java
--- a/Engine/Plugins/Runtime/AndroidFileServer/Source/AndroidFileServer/Android/java/com/epicgames/unreal/RemoteFileManagerActivity.java
+++ b/Engine/Plugins/Runtime/AndroidFileServer/Source/AndroidFileServer/Android/java/com/epicgames/unreal/RemoteFileManagerActivity.java
@@ -3,6 +3,7 @@
 package com.epicgames.unreal;
 
 import android.app.Activity;
+import android.app.Application;
 import android.app.Notification;
 import android.app.NotificationManager;
 import android.app.NotificationChannel;
@@ -15,6 +16,7 @@ import android.util.Log;
 import java.util.Set;
 
 import com.epicgames.unreal.RemoteFileManagerService;
+import com.epicgames.unreal.IAndroidFileServerSetup;
 
 
 public class RemoteFileManagerActivity extends Activity {
@@ -37,16 +39,18 @@ public class RemoteFileManagerActivity extends Activity {
 
 			Log.d(TAG, "RemoteFileManagerActivity cmd: " + cmd + ", package = " + getPackageName());
 
-			if ("init".equals(cmd))
+			IAndroidFileServerSetup AppAFSInterface = (IAndroidFileServerSetup)getApplication();
+
+			if ("init".equals(cmd) && AppAFSInterface != null)
 			{
-				GameApplication.AndroidFileServer_Init(getFilesDir().getAbsolutePath() + "/afs.cache");
+				AppAFSInterface.AndroidFileServer_Init(getFilesDir().getAbsolutePath() + "/afs.cache");
 			}
 			else if ("stop".equals(cmd))
 			{
 				stopAllThreads();
 				stopService(new Intent(this, RemoteFileManagerService.class));
 			}
-			else if ("start".equals(cmd) && GameApplication.AndroidFileServer_Verify(token)) {
+			else if ("start".equals(cmd) && AppAFSInterface != null && AppAFSInterface.AndroidFileServer_Verify(token)) {
 
 				// make a service intent with same extras
 				Intent serviceIntent = new Intent(this, RemoteFileManagerService.class);
diff --git a/Engine/Plugins/Runtime/AndroidFileServer/Source/AndroidFileServer/Android/java/com/epicgames/unreal/RemoteFileManagerService.java b/Engine/Plugins/Runtime/AndroidFileServer/Source/AndroidFileServer/Android/java/com/epicgames/unreal/RemoteFileManagerService.java
--- a/Engine/Plugins/Runtime/AndroidFileServer/Source/AndroidFileServer/Android/java/com/epicgames/unreal/RemoteFileManagerService.java
+++ b/Engine/Plugins/Runtime/AndroidFileServer/Source/AndroidFileServer/Android/java/com/epicgames/unreal/RemoteFileManagerService.java
@@ -16,7 +16,7 @@ import android.app.NotificationManager;
 import android.app.NotificationChannel;
 import android.app.PendingIntent;
 
-import com.epicgames.unreal.GameActivity;
+import com.epicgames.unreal.IAndroidFileServerSetup;
 import com.epicgames.unreal.RemoteFileManager;
 
 import java.io.BufferedReader;
@@ -70,12 +70,14 @@ public class RemoteFileManagerService extends Service {
 			return START_NOT_STICKY;
 		}
 
+		IAndroidFileServerSetup AppAFSInterface = (IAndroidFileServerSetup)getApplication();
+
 		String action = intent.getAction();
 		if (action != null && action.equals("com.epicgames.unreal.RemoteFileManager.intent.COMMAND2")) {
 			String cmd = intent.getStringExtra("cmd");
 			String token = intent.getStringExtra("token");
 			int port = intent.getIntExtra("port", SERVER_PORT);
-			if (cmd != null && cmd.equals("start") && GameApplication.AndroidFileServer_Verify(token)) {
+			if (AppAFSInterface != null && cmd != null && cmd.equals("start") && AppAFSInterface.AndroidFileServer_Verify(token)) {
 				String packageName = getPackageName();
 				PackageManager pm = getPackageManager();
 
diff --git a/Engine/Plugins/Runtime/AndroidFileServer/Source/AndroidFileServer/AndroidFileServer_UPL.xml b/Engine/Plugins/Runtime/AndroidFileServer/Source/AndroidFileServer/AndroidFileServer_UPL.xml
--- a/Engine/Plugins/Runtime/AndroidFileServer/Source/AndroidFileServer/AndroidFileServer_UPL.xml
+++ b/Engine/Plugins/Runtime/AndroidFileServer/Source/AndroidFileServer/AndroidFileServer_UPL.xml
@@ -133,15 +133,18 @@
 					<true>
 						<copyFile src="$S(PluginDir)/Android/java/com/epicgames/unreal/RemoteFileManagerActivity.java" dst="$S(BuildDir)/src/com/epicgames/unreal/RemoteFileManagerActivity.java" />
 						<copyFile src="$S(PluginDir)/Android/java/com/epicgames/unreal/RemoteFileManagerService.java" dst="$S(BuildDir)/src/com/epicgames/unreal/RemoteFileManagerService.java" />
+						<copyFile src="$S(PluginDir)/Android/java/com/epicgames/unreal/IAndroidFileServerSetup.java" dst="$S(BuildDir)/src/com/epicgames/unreal/IAndroidFileServerSetup.java" />
 					</true>
 					<false>
 						<isArch arch="arm64-v8a">
 							<deleteFiles filespec="arm64/src/com/epicgames/unreal/RemoteFileManagerActivity.java" />
 							<deleteFiles filespec="arm64/src/com/epicgames/unreal/RemoteFileManagerService.java" />
+							<deleteFiles filespec="arm64/src/com/epicgames/unreal/IAndroidFileServerSetup.java" />
 						</isArch>
 						<isArch arch="x86_64">
 							<deleteFiles filespec="x64/src/com/epicgames/unreal/RemoteFileManagerActivity.java" />
 							<deleteFiles filespec="x64/src/com/epicgames/unreal/RemoteFileManagerService.java" />
+							<deleteFiles filespec="x64/src/com/epicgames/unreal/IAndroidFileServerSetup.java" />
 						</isArch>
 					</false>
 				</if>
@@ -151,11 +154,13 @@
 					<deleteFiles filespec="arm64/src/com/epicgames/unreal/RemoteFileManager.java" />
 					<deleteFiles filespec="arm64/src/com/epicgames/unreal/RemoteFileManagerActivity.java" />
 					<deleteFiles filespec="arm64/src/com/epicgames/unreal/RemoteFileManagerService.java" />
+					<deleteFiles filespec="arm64/src/com/epicgames/unreal/IAndroidFileServerSetup.java" />
 				</isArch>
 				<isArch arch="x86_64">
 					<deleteFiles filespec="x64/src/com/epicgames/unreal/RemoteFileManager.java" />
 					<deleteFiles filespec="x64/src/com/epicgames/unreal/RemoteFileManagerActivity.java" />
 					<deleteFiles filespec="x64/src/com/epicgames/unreal/RemoteFileManagerService.java" />
+					<deleteFiles filespec="x64/src/com/epicgames/unreal/IAndroidFileServerSetup.java" />
 				</isArch>
 			</false>
 		</if>
@@ -359,6 +364,10 @@ tasks.whenTaskAdded { task ->
 		</if>
 	</settingsGradleAdditions>
 
+	<gameApplicationInterfaceAdditions>
+		<insert>, IAndroidFileServerSetup </insert>
+	</gameApplicationInterfaceAdditions>
+
 	<gameApplicationImportAdditions>
 		<if condition="bEmbedAFS">
 			<true>
@@ -390,7 +399,7 @@ import com.epicgames.unreal.RemoteFileManager;
 				<if condition="Distribution">
 					<true>
 						<insert>
-	public static boolean AndroidFileServer_Init(String filename)
+	public boolean AndroidFileServer_Init(String filename)
 	{
 		return false;
 	}
@@ -398,7 +407,7 @@ import com.epicgames.unreal.RemoteFileManager;
 					</true>
 					<false>
 						<insert>
-	public static boolean AndroidFileServer_Init(String filename)
+	public boolean AndroidFileServer_Init(String filename)
 	{
 		File file = new File(filename);
 		try
@@ -447,7 +456,7 @@ import com.epicgames.unreal.RemoteFileManager;
 				</if>
 				<insertNewline/>
 				<insert>
-	public static boolean AndroidFileServer_Verify(String Token)
+	public boolean AndroidFileServer_Verify(String Token)
 	{
 				</insert>
 				<if condition="bTokenOptional">
@@ -465,12 +474,12 @@ import com.epicgames.unreal.RemoteFileManager;
 			</true>
 			<false>
 				<insert>
-	public static boolean AndroidFileServer_Init(String filename)
+	public AndroidFileServer_Init(String filename)
 	{
 		return true;
 	}
 
-	public static boolean AndroidFileServer_Verify(String Token)
+	public boolean AndroidFileServer_Verify(String Token)
 	{
 		return false;
 	}
diff --git a/Engine/Source/Developer/LowLevelTestsRunner/Android/Java/src/com/epicgames/unreal/tests/TestApplication.java b/Engine/Source/Developer/LowLevelTestsRunner/Android/Java/src/com/epicgames/unreal/tests/TestApplication.java
--- a/Engine/Source/Developer/LowLevelTestsRunner/Android/Java/src/com/epicgames/unreal/tests/TestApplication.java
+++ b/Engine/Source/Developer/LowLevelTestsRunner/Android/Java/src/com/epicgames/unreal/tests/TestApplication.java
@@ -10,7 +10,9 @@ import androidx.lifecycle.LifecycleObserver;
 import androidx.lifecycle.OnLifecycleEvent;
 import androidx.lifecycle.ProcessLifecycleOwner;
 
-public class TestApplication extends Application implements LifecycleObserver {
+import com.epicgames.unreal.IAndroidFileServerSetup;
+
+public class TestApplication extends Application implements LifecycleObserver, IAndroidFileServerSetup {
 	private static boolean isForeground = false;
 
 	private static Context context;
@@ -64,4 +66,14 @@ public class TestApplication extends Application implements LifecycleObserver {
 	public static boolean isAppInBackground() {
 		return !isForeground;
 	}
+	
+	public boolean AndroidFileServer_Verify(String Token)
+	{
+		return true;
+	}
+	
+	public boolean AndroidFileServer_Init(String filename)
+	{
+		return true;
+	}
 }
diff --git a/Engine/Source/Developer/LowLevelTestsRunner/LowLevelTestsAndroid_UPL.xml b/Engine/Source/Developer/LowLevelTestsRunner/LowLevelTestsAndroid_UPL.xml
--- a/Engine/Source/Developer/LowLevelTestsRunner/LowLevelTestsAndroid_UPL.xml
+++ b/Engine/Source/Developer/LowLevelTestsRunner/LowLevelTestsAndroid_UPL.xml
@@ -7,14 +7,6 @@
 		<log text="Low level tests UPL script"/>
 	</init>
 	<androidManifestUpdates>
-		<setElement result="PermissionRead" value="uses-permission" />
-		<addAttribute tag="$PermissionRead" name="android:name" value="android.permission.READ_EXTERNAL_STORAGE" />
-		<addElement tag="manifest" name="PermissionRead" />
-
-		<setElement result="PermissionWrite" value="uses-permission" />
-		<addAttribute tag="$PermissionWrite" name="android:name" value="android.permission.Write_EXTERNAL_STORAGE" />
-		<addElement tag="manifest" name="PermissionWrite" />
-
 		<removeElement tag="application"/>
 		<setElement result="TestApplication" value="application" />
 		<addAttribute tag="$TestApplication" name="android:name" value="com.epicgames.unreal.tests.TestApplication" />
@@ -51,6 +43,23 @@
 
 		<addElement tag="intent-filter" name="TestActivityAction" />
 		<addElement tag="intent-filter" name="TestActivityCategory" />
+
+		<addPermission android:name="android.permission.INTERNET"/>
+		<addPermission android:name="android.permission.READ_EXTERNAL_STORAGE"/>
+		<addPermission android:name="android.permission.WRITE_EXTERNAL_STORAGE"/>
+
+		<addPermission android:name="android.permission.FOREGROUND_SERVICE"/>
+		<addPermission android:name="android.permission.FOREGROUND_SERVICE_DATA_SYNC"/>
+		<addPermission android:name="android.permission.POST_NOTIFICATIONS"/>
+
+		<addElements tag="application">
+			<service android:name="com.epicgames.unreal.RemoteFileManagerService" android:foregroundServiceType="dataSync" android:exported="true" />
+			<activity android:name="com.epicgames.unreal.RemoteFileManagerActivity" android:exported="true" android:screenOrientation="sensor" android:theme="@android:style/Theme.NoDisplay">
+				<intent-filter>
+					<action android:name="com.epicgames.unreal.RemoteFileManager.intent.COMMAND2" />
+				</intent-filter>
+			</activity>
+		</addElements>
 	</androidManifestUpdates>
 	<resourceCopies>
 		<isArch arch="arm64-v8a">
@@ -60,5 +69,6 @@
 			<deleteFiles filespec="x64/src/*.*" recursive="true" />
 		</isArch>
 		<copyDir src="$S(PluginDir)/Android/Java/src" dst="$S(BuildDir)/src"/>
+		<copyDir src="$S(EngineDir)/Plugins/Runtime/AndroidFileServer/Source/AndroidFileServer/Android/java" dst="$S(BuildDir)/src" />
 	</resourceCopies>
 </root>
diff --git a/Engine/Source/Programs/AutomationTool/Android/AndroidPlatform.Automation.cs b/Engine/Source/Programs/AutomationTool/Android/AndroidPlatform.Automation.cs
--- a/Engine/Source/Programs/AutomationTool/Android/AndroidPlatform.Automation.cs
+++ b/Engine/Source/Programs/AutomationTool/Android/AndroidPlatform.Automation.cs
@@ -920,6 +920,16 @@ public class AndroidPlatform : Platform
 		UnrealTargetConfiguration TargetConfiguration = SC != null ? SC.StageTargetConfigurations[0] : Params.ClientConfigsToBuild[0];
 		bIsShipping = TargetConfiguration == UnrealTargetConfiguration.Shipping;
 
+		if (Params.ProjectTargets.Any(Target => Target.Rules.IsTestTarget))
+		{
+			bEnablePlugin = true;
+			AFSToken = "";
+			bIncludeInShipping = false;
+			bAllowExternalStartInShipping = false;
+
+			return true;
+		}
+
 		ConfigHierarchy Ini = ConfigCache.ReadHierarchy(ConfigHierarchyType.Engine, DirectoryReference.FromFile(RawProjectPath), TargetPlatform);
 		if (!Ini.GetBool("/Script/AndroidFileServerEditor.AndroidFileServerRuntimeSettings", "bEnablePlugin", out bEnablePlugin))
 		{
diff --git a/Engine/Source/Programs/AutomationTool/Gauntlet/Platform/Android/Gauntlet.AndroidBuildSource.cs b/Engine/Source/Programs/AutomationTool/Gauntlet/Platform/Android/Gauntlet.AndroidBuildSource.cs
--- a/Engine/Source/Programs/AutomationTool/Gauntlet/Platform/Android/Gauntlet.AndroidBuildSource.cs
+++ b/Engine/Source/Programs/AutomationTool/Gauntlet/Platform/Android/Gauntlet.AndroidBuildSource.cs
@@ -36,9 +36,11 @@ namespace Gauntlet
 
 		public bool UsesExternalFilesDir { get; protected set; }
 
+		public bool UseAFSWithoutToken { get; protected set; }
+
 		public bool UsesPublicLogs { get; protected set; }
 
-		public AndroidBuild(UnrealTargetConfiguration InConfig, string InAndroidPackageName, string InApkPath, Dictionary<string, string> InFilesToInstall, BuildFlags InFlags, bool InIs32Bit, bool bInUsesExternalFilesDir, bool bInUsesPublicLogs)
+		public AndroidBuild(UnrealTargetConfiguration InConfig, string InAndroidPackageName, string InApkPath, Dictionary<string, string> InFilesToInstall, BuildFlags InFlags, bool InIs32Bit, bool bInUsesExternalFilesDir, bool bInUsesPublicLogs, bool bInUseAFSWithoutToken = false)
 		{
 			Configuration = InConfig;
 			AndroidPackageName = InAndroidPackageName;
@@ -48,6 +50,7 @@ namespace Gauntlet
 			Is32Bit = InIs32Bit;
 			UsesExternalFilesDir = bInUsesExternalFilesDir;
 			UsesPublicLogs = bInUsesPublicLogs;
+			UseAFSWithoutToken = bInUseAFSWithoutToken;
 		}
 
 		public bool CanSupportRole(UnrealTargetRole RoleType)
diff --git a/Engine/Source/Programs/AutomationTool/Gauntlet/Platform/Android/Gauntlet.TargetDeviceAndroid.cs b/Engine/Source/Programs/AutomationTool/Gauntlet/Platform/Android/Gauntlet.TargetDeviceAndroid.cs
--- a/Engine/Source/Programs/AutomationTool/Gauntlet/Platform/Android/Gauntlet.TargetDeviceAndroid.cs
+++ b/Engine/Source/Programs/AutomationTool/Gauntlet/Platform/Android/Gauntlet.TargetDeviceAndroid.cs
@@ -115,6 +115,8 @@ namespace Gauntlet
 
 		private AndroidAppInstall InstallCache = null;
 
+		private bool bUseAFSWithoutToken;
+
 		public TargetDeviceAndroid(string InDeviceName = "", AndroidDeviceData DeviceData = null, string InCachePath = null)
 		{
 			AutoParam.ApplyParamsAndDefaults(this, Globals.Params.AllArguments);
@@ -364,6 +366,8 @@ namespace Gauntlet
 				return;
 			}
 
+			bUseAFSWithoutToken = Build.UseAFSWithoutToken;
+
 			string ExternalStoragePath = StoragePath + "/UnrealGame/" + AppConfig.ProjectName;
 			string ExternalFilesPath = StoragePath + "/Android/data/" + Build.AndroidPackageName + "/files/UnrealGame/" + AppConfig.ProjectName;
 
@@ -382,6 +386,8 @@ namespace Gauntlet
 				throw new AutomationException("Unsupported build type {0} for Android!", AppConfig.Build.GetType());
 			}
 
+			bUseAFSWithoutToken = Build.UseAFSWithoutToken;
+
 			string Package = Build.AndroidPackageName;
 			KillRunningProcess(Package);
 
@@ -408,6 +414,8 @@ namespace Gauntlet
 				throw new AutomationException("Unsupported build type {0} for Android!", AppConfig.Build.GetType());
 			}
 
+			bUseAFSWithoutToken = Build.UseAFSWithoutToken;
+
 			string ExternalStoragePath = StoragePath + "/UnrealGame/" + AppConfig.ProjectName;
 			string ExternalFilesPath = StoragePath + "/Android/data/" + Build.AndroidPackageName + "/files/UnrealGame/" + AppConfig.ProjectName;
 
@@ -464,6 +472,7 @@ namespace Gauntlet
 
 			// Now, create and push the UECommandline.txt file
 			bool bUseExternalFileDirectory = (Install.AppConfig.Build as AndroidBuild)?.UsesExternalFilesDir ?? false;
+			bUseAFSWithoutToken = (Install.AppConfig.Build as AndroidBuild)?.UseAFSWithoutToken ?? false;
 			CopyCommandlineFile(Install.AppConfig, Install.AppConfig.CommandLine, Install.AndroidPackageName, Install.AppConfig.ProjectName, bUseExternalFileDirectory);
 
 			// kill any currently running instance:
@@ -478,7 +487,7 @@ namespace Gauntlet
 			RunAdbDeviceCommand("logcat --clear");
 
 			// Ensure artifact directories exist
-			if (UsingAndroidFileServer(Install.ProjectFile, Install.Configuration, out _, out string AFSToken, out _, out _, out _))
+			if (UsingAndroidFileServer(Install.ProjectFile, Install.Configuration, out _, out string AFSToken, out _, out _, out _, bUseAFSWithoutToken))
 			{
 				RunAFSDeviceCommand(string.Format("-p \"{0}\" -k \"{1}\" mkdir \"^saved\"", Install.AndroidPackageName, AFSToken));
 			}
@@ -552,8 +561,17 @@ namespace Gauntlet
 			}
 		}
 
-		public static bool UsingAndroidFileServer(FileReference RawProjectPath, UnrealTargetConfiguration TargetConfiguration, out bool bEnablePlugin, out string AFSToken, out bool bIsShipping, out bool bIncludeInShipping, out bool bAllowExternalStartInShipping)
+		public static bool UsingAndroidFileServer(FileReference RawProjectPath, UnrealTargetConfiguration TargetConfiguration, out bool bEnablePlugin, out string AFSToken, out bool bIsShipping, out bool bIncludeInShipping, out bool bAllowExternalStartInShipping, bool bUseAFSWithoutToken = false)
 		{
+			if (bUseAFSWithoutToken)
+			{
+				bEnablePlugin = true;
+				AFSToken = "";
+				bIsShipping = false;
+				bIncludeInShipping = false;
+				bAllowExternalStartInShipping = false;
+				return true;
+			}
 			if(RawProjectPath == null || !FileReference.Exists(RawProjectPath) || TargetConfiguration == UnrealTargetConfiguration.Unknown)
 			{
 				bEnablePlugin = false;
@@ -600,7 +618,9 @@ namespace Gauntlet
 				throw new AutomationException("Unsupported build type {0} for Android!", AppConfig.Build.GetType());
 			}
 
-			return UsingAndroidFileServer(AppConfig.ProjectFile, Build.Configuration, out _, out _, out _, out _, out _);
+			bUseAFSWithoutToken = Build.UseAFSWithoutToken;
+
+			return UsingAndroidFileServer(AppConfig.ProjectFile, Build.Configuration, out _, out _, out _, out _, out _, bUseAFSWithoutToken);
 		}
 
 		/// <summary>
@@ -844,7 +864,7 @@ namespace Gauntlet
 			FileReference ProjectFile = null, UnrealTargetConfiguration Configuration = UnrealTargetConfiguration.Unknown, bool IgnoreDependencies = false)
 		{
 			string AFSToken;
-			bool UseAFS = UsingAndroidFileServer(ProjectFile, Configuration, out _, out AFSToken, out _, out _, out _);
+			bool UseAFS = UsingAndroidFileServer(ProjectFile, Configuration, out _, out AFSToken, out _, out _, out _, bUseAFSWithoutToken);
 			bool IsAPK = string.Equals(Path.GetExtension(SourcePath), ".apk", StringComparison.OrdinalIgnoreCase);
 
 			// for the APK there's no easy/reliable way to get the date of the version installed, so
@@ -1085,7 +1105,9 @@ namespace Gauntlet
 					throw new AutomationException("Unsupported build type {0} for Android!", AppConfig.Build.GetType());
 				}
 
-				if (UsingAndroidFileServer(AppConfig.ProjectFile, Build.Configuration, out bAFSEnablePlugin, out AFSToken, out bIsShipping, out bAFSIncludeInShipping, out bAFSAllowExternalStartInShipping))
+				bUseAFSWithoutToken = Build.UseAFSWithoutToken;
+
+				if (UsingAndroidFileServer(AppConfig.ProjectFile, Build.Configuration, out bAFSEnablePlugin, out AFSToken, out bIsShipping, out bAFSIncludeInShipping, out bAFSAllowExternalStartInShipping, bUseAFSWithoutToken))
 				{
 					useADB = false;
 				}
@@ -1188,7 +1210,7 @@ namespace Gauntlet
 			FileReference ProjectFile = null, UnrealTargetConfiguration Configuration = UnrealTargetConfiguration.Unknown)
 		{
 			string AFSToken;
-			bool UseAFS = UsingAndroidFileServer(ProjectFile, Configuration, out _, out AFSToken, out _, out _, out _);
+			bool UseAFS = UsingAndroidFileServer(ProjectFile, Configuration, out _, out AFSToken, out _, out _, out _, bUseAFSWithoutToken);
 
 			IProcessResult AdbResult;
 			if (UseAFS)
@@ -1244,9 +1266,11 @@ namespace Gauntlet
 			bool bAFSIncludeInShipping;
 			bool bAFSAllowExternalStartInShipping;
 			bool useADB = true;
+
+			bUseAFSWithoutToken = Build.UseAFSWithoutToken;
 			if (AppConfig != null)
 			{
-				if (UsingAndroidFileServer(AppConfig.ProjectFile, Build.Configuration, out bAFSEnablePlugin, out AFSToken, out bIsShipping, out bAFSIncludeInShipping, out bAFSAllowExternalStartInShipping))
+				if (UsingAndroidFileServer(AppConfig.ProjectFile, Build.Configuration, out bAFSEnablePlugin, out AFSToken, out bIsShipping, out bAFSIncludeInShipping, out bAFSAllowExternalStartInShipping, bUseAFSWithoutToken))
 				{
 					useADB = false;
 				}
@@ -1668,9 +1692,10 @@ namespace Gauntlet
 			}
 		}
 		public void StartRecording()
-		{	
+		{
+			bool bAlwaysUseAFS = (Install.AppConfig.Build as AndroidBuild)?.UseAFSWithoutToken ?? false;
 			// Ensure artifact directories exist
-			if (TargetDeviceAndroid.UsingAndroidFileServer(Install.ProjectFile, Install.Configuration, out _, out string AFSToken, out _, out _, out _))
+			if (TargetDeviceAndroid.UsingAndroidFileServer(Install.ProjectFile, Install.Configuration, out _, out string AFSToken, out _, out _, out _, bAlwaysUseAFS))
 			{
 				Install.AndroidDevice.RunAFSDeviceCommand(string.Format("-p \"{0}\" -k \"{1}\" mkdir \"^saved/Logs\"", Install.AndroidPackageName, AFSToken));
 			}
diff --git a/Engine/Source/Programs/AutomationTool/LowLevelTests/RunLowLevelTests.cs b/Engine/Source/Programs/AutomationTool/LowLevelTests/RunLowLevelTests.cs
--- a/Engine/Source/Programs/AutomationTool/LowLevelTests/RunLowLevelTests.cs
+++ b/Engine/Source/Programs/AutomationTool/LowLevelTests/RunLowLevelTests.cs
@@ -840,7 +840,7 @@ namespace LowLevelTests
 		public IBuild CreateBuild(UnrealTargetPlatform InPlatform, UnrealTargetConfiguration InConfiguration, string InTestApp, string InBuildPath, bool bSkipStage)
 		{
 			string ExecutablePath = ILowLevelTestsBuildFactory.GetExecutable(InPlatform, InConfiguration, InTestApp, InBuildPath, GetExecutableRegex(InPlatform));
-			return new AndroidBuild(InConfiguration, $"com.epicgames.{InTestApp}", Path.Combine(InBuildPath, ExecutablePath), new Dictionary<string, string>(), BuildFlags.Packaged | BuildFlags.CanReplaceCommandLine, false, true, false);
+			return new AndroidBuild(InConfiguration, $"com.epicgames.{InTestApp}", Path.Combine(InBuildPath, ExecutablePath), new Dictionary<string, string>(), BuildFlags.Packaged | BuildFlags.CanReplaceCommandLine, false, true, false, true);
 		}
 
 		public string GetExecutableRegex(UnrealTargetPlatform InPlatform)
diff --git a/Engine/Source/Programs/UnrealBuildTool/Platform/Android/UEDeployAndroid.cs b/Engine/Source/Programs/UnrealBuildTool/Platform/Android/UEDeployAndroid.cs
--- a/Engine/Source/Programs/UnrealBuildTool/Platform/Android/UEDeployAndroid.cs
+++ b/Engine/Source/Programs/UnrealBuildTool/Platform/Android/UEDeployAndroid.cs
@@ -6570,6 +6570,7 @@ import java.util.Collection;
 
 			Dictionary<string, string> Replacements = new Dictionary<string, string>{
 				{ "//$${gameApplicationImportAdditions}$$", UPL!.ProcessPluginNode(NDKArch, "gameApplicationImportAdditions", GameApplicationImportAdditionsDefault)},
+				{ "//$${gameApplicationInterfaceAdditions}$$", UPL!.ProcessPluginNode(NDKArch, "gameApplicationInterfaceAdditions", "")},
 				{ "//$${gameApplicationOnCreateAdditions}$$", UPL.ProcessPluginNode(NDKArch, "gameApplicationOnCreateAdditions", "")},
 				{ "//$${gameApplicationAttachBaseContextAdditions}$$", UPL.ProcessPluginNode(NDKArch, "gameApplicationAttachBaseContextAdditions", "")},
 				{ "//$${gameApplicationOnLowMemoryAdditions}$$", UPL.ProcessPluginNode(NDKArch, "gameApplicationOnLowMemoryAdditions", "")},
diff --git a/Engine/Source/Runtime/Core/Public/Android/AndroidJavaEnv.h b/Engine/Source/Runtime/Core/Public/Android/AndroidJavaEnv.h
--- a/Engine/Source/Runtime/Core/Public/Android/AndroidJavaEnv.h
+++ b/Engine/Source/Runtime/Core/Public/Android/AndroidJavaEnv.h
@@ -2,6 +2,8 @@
 
 #pragma once
 
+#if USE_ANDROID_JNI
+
 #include "CoreTypes.h"
 #include "Containers/Array.h"
 #include "Containers/StaticArray.h"
@@ -1170,3 +1172,4 @@ public:
 	// Converts the java objectArray to an array of FStrings. jopbjectArray must be a String[] on the Java side
 	static CORE_API TArray<FString> ObjectArrayToFStringTArray(JNIEnv* Env, jobjectArray ObjectArray);
 };
+#endif
\ No newline at end of file
