Fixing a bug where UAT could hang waiting for output.
We now log a warning instead.

diff --git a/Engine/Source/Programs/AutomationTool/AutomationUtils/ProcessUtils.cs b/Engine/Source/Programs/AutomationTool/AutomationUtils/ProcessUtils.cs
--- a/Engine/Source/Programs/AutomationTool/AutomationUtils/ProcessUtils.cs
+++ b/Engine/Source/Programs/AutomationTool/AutomationUtils/ProcessUtils.cs
@@ -561,19 +561,46 @@ namespace AutomationTool
 					}
 				}
 			}
+
+			int waitTimeBeforeExit = 1000;
+
 			if (!(bStdOutSignalReceived && bStdErrSignalReceived))
 			{
-				Logger.LogInformation("Waited for {Seconds:n2}s for output of {AppName}, some output may be missing; we gave up.", Watch.Elapsed.TotalSeconds, AppName);
+				Logger.LogInformation("Waited for {Seconds:n2}s for output of {AppName}, some output may be missing; we gave up. Now waiting 5 minutes for process to exit.", Watch.Elapsed.TotalSeconds, AppName);
+				waitTimeBeforeExit = 300000;
 			}
 
-			// wait until the process terminates
-			Proc?.WaitForExit();
+			bool exited = true;
+
+			if (Proc != null)
+			{
+				// wait until the process terminates or we hit our timeout
+				exited = Proc.WaitForExit(waitTimeBeforeExit);
+			}
 
 			lock (ProcSyncObject)
 			{
-				if (Proc != null)
+				// if we exited return the exit code
+				if (exited)
 				{
-					ExitCode = Proc.ExitCode;
+					if (Proc != null)
+					{
+						ExitCode = Proc.ExitCode;
+					}
+				}
+				else
+				{
+					// if we failed to exit log a systemic warning and kill the process
+					Logger.LogWarning(KnownLogEvents.Systemic, "Failed to exit process {AppName}.", AppName);
+
+					try
+					{
+						Proc?.Kill();
+					}
+					catch
+					{ 
+					
+					}
 				}
 			}
 		}
