Added new systemic event match for generic uba cache response timeouts
#rnx

diff --git a/Engine/Source/Programs/AutomationTool/AutomationUtils/Matchers/SystemicEventMatcher.cs b/Engine/Source/Programs/AutomationTool/AutomationUtils/Matchers/SystemicEventMatcher.cs
--- a/Engine/Source/Programs/AutomationTool/AutomationUtils/Matchers/SystemicEventMatcher.cs
+++ b/Engine/Source/Programs/AutomationTool/AutomationUtils/Matchers/SystemicEventMatcher.cs
@@ -21,6 +21,8 @@ namespace AutomationUtils.Matchers
 
 		static readonly Regex s_missingFileList = new Regex(@"^\s*ERROR: Missing local or shared file list.*\.xml");
 
+		static readonly Regex s_ubaCacheResponseTimeout = new Regex(@"UbaCache (.*) - Timed out after .* waiting for message response from server");
+
 		public LogEventMatch? Match(ILogCursor cursor)
 		{
 			if (cursor.IsMatch(s_ddc))
@@ -57,6 +59,11 @@ namespace AutomationUtils.Matchers
 				return new LogEventBuilder(cursor).ToMatch(LogEventPriority.High, LogLevel.Error, KnownLogEvents.Systemic_MissingFileList);
 			}
 
+			if (cursor.IsMatch(s_ubaCacheResponseTimeout))
+			{
+				return new LogEventBuilder(cursor).ToMatch(LogEventPriority.High, LogLevel.Warning, KnownLogEvents.Systemic_UbaCacheResponseTimeout);
+			}
+
 			if (cursor.Contains("SignTool Error: The specified timestamp server either could not be reached or"))
 			{
 				return new LogEventBuilder(cursor).ToMatch(LogEventPriority.Normal, LogLevel.Information, KnownLogEvents.Systemic_SignToolTimeStampServer);
