Record network test map as session metadata

#rnx

diff --git a/Engine/Source/Programs/AutomationTool/Gauntlet/Unreal/Automation/UE.Networking.cs b/Engine/Source/Programs/AutomationTool/Gauntlet/Unreal/Automation/UE.Networking.cs
--- a/Engine/Source/Programs/AutomationTool/Gauntlet/Unreal/Automation/UE.Networking.cs
+++ b/Engine/Source/Programs/AutomationTool/Gauntlet/Unreal/Automation/UE.Networking.cs
@@ -145,5 +145,27 @@ namespace UE
 				}
 			}
 		}
+
+		public override ITestReport CreateReport(TestResult Result)
+		{
+			ITestReport Report = base.CreateReport(Result);
+
+			if (Report != null)
+			{
+				foreach (var Role in RoleResults.Where(R => R.ExitCode != 0))
+				{
+					var LogErrors = Role.LogSummary.Errors.Where(E => E.Level == UnrealLog.LogLevel.Error);
+					foreach (var LogError in LogErrors.Distinct())
+					{
+						Report.AddEvent(EventType.Error, LogError.Message);
+					}
+				}
+
+				TConfigClass Config = GetCachedConfiguration();
+				if (!string.IsNullOrEmpty(Config.Map)) Report.SetMetadata("Map", Config.Map);
+			}
+
+			return Report;
+		}
 	}
 }
