Add support for Phase tracking in UnrealTestNode

#rnx

diff --git a/Engine/Source/Programs/AutomationTool/Gauntlet/Framework/Base/Gauntlet.TestTracking.cs b/Engine/Source/Programs/AutomationTool/Gauntlet/Framework/Base/Gauntlet.TestTracking.cs
new file mode 100644
--- /dev/null
+++ b/Engine/Source/Programs/AutomationTool/Gauntlet/Framework/Base/Gauntlet.TestTracking.cs
@@ -0,0 +1,286 @@
+// Copyright Epic Games, Inc. All Rights Reserved.
+
+#nullable enable
+using System;
+using System.Linq;
+using System.Collections;
+using System.Collections.Generic;
+using System.Collections.Concurrent;
+using Microsoft.Extensions.Logging;
+using EpicGames.Core;
+using static Gauntlet.HordeReport.AutomatedTestSessionData;
+using AutomatedTestSessionData = Gauntlet.HordeReport.AutomatedTestSessionData;
+using AutomationTool;
+using Logging = Microsoft.Extensions.Logging;
+
+namespace Gauntlet
+{
+	/// <summary>
+	/// Test execution tracking objects
+	/// </summary>
+	namespace TestTracking
+	{
+		/// <summary>
+		/// Handler for associated phase data
+		/// </summary>
+		public interface IPhaseDataHandler
+		{
+			/// <summary>
+			/// Called when test phase is added to report
+			/// </summary>
+			/// <param name="ReportedPhase"></param>
+			public void OnAddToReport(TestPhase ReportedPhase);
+		}
+
+		/// <summary>
+		/// Object that represent a test phase
+		/// </summary>
+		public class Phase
+		{
+			public string Name;
+			public double ElapseTimeInSeconds;
+			public DateTime? StartedAt;
+			public TestPhaseOutcome Outcome;
+			private ConcurrentQueue<LogEvent> Events;
+
+			/// <summary>
+			/// Opaque object that can be associated with the phase.
+			/// The object class can implement its own handler by implementing IPhaseDataHandler interface.
+			/// OnAddToReport method will be called when the phase is attached to the report allowing
+			/// the opaque data to make modifications to the reported data.
+			/// </summary>
+			public object? Data;
+
+			public Phase(string InName, DateTime? InDateTime = null, object? InData = null)
+			{
+				Name = InName;
+				ElapseTimeInSeconds = 0;
+				StartedAt = InDateTime;
+				Outcome = TestPhaseOutcome.Unknown;
+				Events = new ConcurrentQueue<LogEvent>();
+				Data = InData;
+			}
+
+			/// <summary>
+			/// Mark the phase as finished
+			/// </summary>
+			/// <param name="InOutcome">The optional outcome. By default it will evaluted based on tracked events</param>
+			/// <returns></returns>
+			public double End(TestPhaseOutcome? InOutcome = null)
+			{
+				DateTime StartedTime = StartedAt ?? DateTime.UtcNow;
+				ElapseTimeInSeconds = (DateTime.UtcNow - StartedTime).TotalSeconds;
+				if (InOutcome != null)
+				{
+					Outcome = InOutcome.Value;
+				}
+				else
+				{
+					Outcome = Events.Any(E => E.Level == Logging.LogLevel.Error || E.Level == Logging.LogLevel.Critical) ? TestPhaseOutcome.Failed : TestPhaseOutcome.Success;
+				}
+
+				return ElapseTimeInSeconds;
+			}
+
+			/// <summary>
+			/// Add an event to the phase tracker
+			/// </summary>
+			/// <param name="Level">The log level of the event</param>
+			/// <param name="Message"></param>
+			/// <param name="Args"></param>
+			public void AddEvent(Logging.LogLevel Level, string Message, params object[] Args)
+			{
+				string? Format = null;
+				Dictionary<string, object?>? Properties = null;
+				if (Args.Any())
+				{
+					Format = Message;
+					Properties = new Dictionary<string, object?>();
+					MessageTemplate.ParsePropertyValues(Format, Args, Properties);
+					Message = MessageTemplate.Render(Format, Properties);
+				}
+				EventId EventId = Level == Logging.LogLevel.Critical ? KnownLogEvents.Gauntlet_FatalEvent : KnownLogEvents.Gauntlet_TestEvent;
+				LogEvent Event = new LogEvent(DateTime.UtcNow, Level, EventId, Message, Format, Properties, null);
+
+				Events.Enqueue(Event);
+			}
+
+			/// <summary>
+			/// Get the associated data
+			/// </summary>
+			/// <typeparam name="T"></typeparam>
+			/// <returns></returns>
+			public T? GetData<T>()
+			{
+				if (Data is T Value)
+				{
+					return Value;
+				}
+
+				return default(T);
+			}
+
+			/// <summary>
+			/// Add the information tracked by the phase to the test report
+			/// </summary>
+			/// <param name="Report"></param>
+			/// <returns></returns>
+			public TestPhase AddToReport(AutomatedTestSessionData Report)
+			{
+				if (Outcome == TestPhaseOutcome.Unknown)
+				{
+					if (StartedAt == null)
+					{
+						Outcome = TestPhaseOutcome.NotRun;
+					}
+					else
+					{
+						End(TestPhaseOutcome.Interrupted);
+					}
+				}
+
+				TestPhase ReportedPhase = Report.AddPhase(Name);
+				if (StartedAt != null)
+				{
+					ReportedPhase.SetTiming(StartedAt.Value, (float)ElapseTimeInSeconds);
+				}
+				ReportedPhase.SetOutcome(Outcome);
+				TestEventStream Stream = ReportedPhase.GetStream();
+				foreach (LogEvent Event in Events)
+				{
+					Stream.AddEvent(Event);
+				}
+
+				if (Data is IPhaseDataHandler DataHandler)
+				{
+					DataHandler.OnAddToReport(ReportedPhase);
+				}
+
+				return ReportedPhase;
+			}
+		}
+
+		/// <summary>
+		/// Object to track queued phases and enumerate through
+		/// </summary>
+		public class PhaseQueue : IEnumerable<Phase>
+		{
+			public PhaseQueue()
+			{
+				_manifest = new();
+				_queuedPhases = new();
+				_last = null;
+				_current = null;
+			}
+
+			private ConcurrentDictionary<string, Phase> _manifest;
+			private Queue<Phase> _queuedPhases;
+			private Phase? _last;
+			private Phase? _current;
+
+			/// <summary>
+			/// The last added phase to the queue
+			/// </summary>
+			public Phase? Last => _last;
+
+
+			/// <summary>
+			/// The current running phase. Use Start to identify the running phase
+			/// </summary>
+			public Phase? Current => _current;
+
+			/// <summary>
+			/// Start a phase, if not queued it will be added to the queue
+			/// </summary>
+			/// <param name="Name">The name of the phase</param>
+			/// <param name="Data">Extra object to associate with the phase</param>
+			/// <returns></returns>
+			/// <exception cref="AutomationException"></exception>
+			public Phase Start(string Name, object? Data = null)
+			{
+				bool bAlreadyAdded = false;
+				Phase StartedPhase = _manifest.AddOrUpdate(Name,
+					(Name) => new Phase(Name, DateTime.UtcNow, Data),
+					(Name, ExistingPhase) =>
+					{
+						bAlreadyAdded = true;
+						if (ExistingPhase.StartedAt != null)
+						{
+							throw new AutomationException($"Phase '{Name}' already started.");
+						}
+						ExistingPhase.StartedAt = DateTime.UtcNow;
+						if (Data != null) ExistingPhase.Data = Data;
+						return ExistingPhase;
+					}
+				);
+
+				if (!bAlreadyAdded)
+				{
+					_queuedPhases.Enqueue(StartedPhase);
+					_last = StartedPhase;
+				}
+
+				_current = StartedPhase;
+
+				return StartedPhase;
+			}
+
+			/// <summary>
+			/// Add a new phase. Raise an exception if it already exists.
+			/// </summary>
+			/// <param name="Name">The name of the phase</param>
+			/// <param name="Data">Extra object to associate with the phase</param>
+			/// <returns></returns>
+			/// <exception cref="AutomationException"></exception>
+			public Phase Add(string Name, object? Data = null)
+			{
+				Phase NewPhase = new Phase(Name, null, Data);
+				if (!_manifest.TryAdd(Name, NewPhase))
+				{
+					throw new AutomationException($"Phase '{Name}' already added. Make sure to use unique phase name.");
+				}
+
+				_queuedPhases.Enqueue(NewPhase);
+				_last = NewPhase;
+
+				return NewPhase;
+			}
+
+			/// <summary>
+			/// Get the corresponding phase
+			/// </summary>
+			/// <param name="Name"></param>
+			/// <returns></returns>
+			public Phase? Get(string Name)
+			{
+				Phase? FetchedPhase = null;
+				_manifest.TryGetValue(Name, out FetchedPhase);
+
+				return FetchedPhase;
+			}
+
+			IEnumerator IEnumerable.GetEnumerator()
+			{
+				return (IEnumerator)GetEnumerator();
+			}
+
+			/// <inheritdoc/>
+			public IEnumerator<Phase> GetEnumerator()
+			{
+				return _queuedPhases.GetEnumerator();
+			}
+
+
+			/// <summary>
+			/// Clear the queue
+			/// </summary>
+			public void Clear()
+			{
+				_queuedPhases.Clear();
+				_manifest.Clear();
+				_last = null;
+				_current = null;
+			}
+		}
+	}
+}
\ No newline at end of file
diff --git a/Engine/Source/Programs/AutomationTool/Gauntlet/Unreal/Base/Gauntlet.UnrealTestEvent.cs b/Engine/Source/Programs/AutomationTool/Gauntlet/Unreal/Base/Gauntlet.UnrealTestEvent.cs
--- a/Engine/Source/Programs/AutomationTool/Gauntlet/Unreal/Base/Gauntlet.UnrealTestEvent.cs
+++ b/Engine/Source/Programs/AutomationTool/Gauntlet/Unreal/Base/Gauntlet.UnrealTestEvent.cs
@@ -1,5 +1,6 @@
 ﻿// Copyright Epic Games, Inc. All Rights Reserved.
 
+using EpicGames.Core;
 using System;
 using System.Collections.Generic;
 using System.Linq;
@@ -65,5 +66,14 @@ namespace Gauntlet
 		// Constructor that requires all properties but DateTime
 		public UnrealTestEvent(EventSeverity InSeverity, string InSummary, IEnumerable<string> InDetails, UnrealLog.CallstackMessage InCallstack = null)
 			: this(DateTime.UtcNow, InSeverity, InSummary, InDetails, InCallstack) { }
+
+		// Constructor that use structure logging standard inputs
+		public UnrealTestEvent(EventSeverity InSeverity, string InFormat, params object[] Args)
+			: this(InSeverity, string.Empty, Enumerable.Empty<string>())
+		{
+			Dictionary<string, object> Properties = new();
+			MessageTemplate.ParsePropertyValues(InFormat, Args, Properties);
+			Summary = MessageTemplate.Render(InFormat, Properties);
+		}
 	}
 }
diff --git a/Engine/Source/Programs/AutomationTool/Gauntlet/Unreal/Base/Gauntlet.UnrealTestNode.cs b/Engine/Source/Programs/AutomationTool/Gauntlet/Unreal/Base/Gauntlet.UnrealTestNode.cs
--- a/Engine/Source/Programs/AutomationTool/Gauntlet/Unreal/Base/Gauntlet.UnrealTestNode.cs
+++ b/Engine/Source/Programs/AutomationTool/Gauntlet/Unreal/Base/Gauntlet.UnrealTestNode.cs
@@ -9,11 +9,13 @@ using System.Text.RegularExpressions;
 using System.Linq;
 using EpicGames.Core;
 using Microsoft.Extensions.Logging;
+using Logging = Microsoft.Extensions.Logging;
 using AutomationUtils.Matchers;
 using static Gauntlet.HordeReport;
 using static Gauntlet.UnrealSessionInstance;
 using Gauntlet.Utils;
 using static Gauntlet.HordeReport.AutomatedTestSessionData;
+using Gauntlet.TestTracking;
 
 namespace Gauntlet
 {
@@ -78,6 +80,11 @@ namespace Gauntlet
 		/// </summary>
 		protected PopulateTestPhasesReport OnPopulateTestPhases;
 
+		/// <summary>
+		/// Phase tracker
+		/// </summary>
+		protected PhaseQueue Phases;
+
 		/// <summary>
 		/// Returns a list of all log channels the heartbeat tick should look for.
 		/// </summary>
@@ -140,9 +147,9 @@ namespace Gauntlet
 		/// <param name="Args"></param>
 		public virtual void ReportError(string Message, params object[] Args)
 		{
-			Message = string.Format(Message, Args);
-			AddTestEvent(new UnrealTestEvent(EventSeverity.Error, Message, Enumerable.Empty<string>()));
-			if (!LogWarningsAndErrorsAfterSummary) { Log.Error(KnownLogEvents.Gauntlet_TestEvent, Message); }
+			Phases.Current?.AddEvent(Logging.LogLevel.Error, Message, Args);
+			AddTestEvent(new UnrealTestEvent(EventSeverity.Error, Message, Args));
+			if (Phases.Current == null && !LogWarningsAndErrorsAfterSummary) { Log.Error(KnownLogEvents.Gauntlet_TestEvent, Message, Args); }
 			if (GetTestStatus() == TestStatus.Complete && GetTestResult() == TestResult.Passed)
 			{
 				SetUnrealTestResult(TestResult.Failed);
@@ -156,9 +163,21 @@ namespace Gauntlet
 		/// <param name="Args"></param>
 		public virtual void ReportWarning(string Message, params object[] Args)
 		{
-			Message = string.Format(Message, Args);
-			AddTestEvent(new UnrealTestEvent(EventSeverity.Warning, Message, Enumerable.Empty<string>()));
-			if (!LogWarningsAndErrorsAfterSummary) { Log.Warning(KnownLogEvents.Gauntlet_TestEvent, Message); }
+			Phases.Current?.AddEvent(Logging.LogLevel.Warning, Message, Args);
+			AddTestEvent(new UnrealTestEvent(EventSeverity.Warning, Message, Args));
+			if (Phases.Current == null && !LogWarningsAndErrorsAfterSummary) { Log.Warning(KnownLogEvents.Gauntlet_TestEvent, Message, Args); }
+		}
+
+		/// <summary>
+		/// Report information
+		/// </summary>
+		/// <param name="Message"></param>
+		/// <param name="Args"></param>
+		public virtual void ReportInfo(string Message, params object[] Args)
+		{
+			Phases.Current?.AddEvent(Logging.LogLevel.Information, Message, Args);
+			AddTestEvent(new UnrealTestEvent(EventSeverity.Info, Message, Args));
+			if (Phases.Current == null) { Log.Info(KnownLogEvents.Gauntlet_TestEvent, Message, Args); }
 		}
 
 		// Begin UnrealTestNode properties and members
@@ -376,6 +395,7 @@ namespace Gauntlet
 			NumPasses = 0;
 			TestVersion = new Version("1.0.0");
 			ArtifactPath = string.Empty;
+			Phases = new();
 			PopulateCommandlineInfo();
 			// We format warnings ourselves so don't show these
 			LogWarningsAndErrorsAfterSummary = false;
@@ -796,6 +816,7 @@ namespace Gauntlet
 			NumPasses = InNumPasses;
 			LastHeartbeatTime = DateTime.MinValue;
 			LastActiveHeartbeatTime = DateTime.MinValue;
+			Phases.Clear();
 
 			TConfigClass Config = GetCachedConfiguration();
 
@@ -1028,6 +1049,9 @@ namespace Gauntlet
 			LastHeartbeatTime = DateTime.MinValue;
 			LastActiveHeartbeatTime = DateTime.MinValue;
 
+			//Reset phases
+			Phases.Clear();
+
 			ReserveArtifactPath();
 
 			// Relaunch the test
@@ -1383,8 +1407,13 @@ namespace Gauntlet
 		/// <param name="Result"></param>
 		protected virtual ITestReport CreateSimpleReportForHorde(TestResult Result)
 		{
-			string ArtifactPathForReport = ArtifactPath;
-			Horde.OnStepSummary += () => Horde.GenerateArtifactLink(ArtifactPathForReport);
+
+			if (Phases.Any())
+			{
+				OnPopulateTestPhases += PopulatePhaseReport;
+				OnPopulateTestPhases += (Report) => PopulateClosurePhase(Report);
+			}
+
 			BaseHordeReport HordeTestReport = null;
 			if (Globals.Params.ParseParam("UseTestDataV2"))
 			{
@@ -1474,6 +1503,28 @@ namespace Gauntlet
 			return HordeTestReport;
 		}
 
+		/// <summary>
+		/// Populate report phases from test execution
+		/// </summary>
+		/// <param name="SessionData"></param>
+		protected virtual void PopulatePhaseReport(AutomatedTestSessionData SessionData)
+		{
+			// Assume all devices were used to run all test cases
+			HashSet<string> DeviceKeys = SessionData.GetDevices().Select(D => D.key).ToHashSet();
+			// Go through all phases
+			foreach (Phase CurPhase in Phases)
+			{
+				if (string.IsNullOrEmpty(CurPhase.Name)) continue;
+				TestPhase ReportedPhase = CurPhase.AddToReport(SessionData);
+
+				if (ReportedPhase.outcome != TestPhaseOutcome.NotRun
+					&& ReportedPhase.outcome != TestPhaseOutcome.Skipped)
+				{
+					ReportedPhase.deviceKeys = DeviceKeys;
+				}
+			}
+		}
+
 		/// <summary>
 		/// Default populate closure phase
 		/// </summary>
