Add support for Phase tracking in UnrealTestNode (fix post merge)

#rnx

diff --git a/Engine/Source/Programs/AutomationTool/Gauntlet/Unreal/Base/Gauntlet.UnrealTestNode.cs b/Engine/Source/Programs/AutomationTool/Gauntlet/Unreal/Base/Gauntlet.UnrealTestNode.cs
--- a/Engine/Source/Programs/AutomationTool/Gauntlet/Unreal/Base/Gauntlet.UnrealTestNode.cs
+++ b/Engine/Source/Programs/AutomationTool/Gauntlet/Unreal/Base/Gauntlet.UnrealTestNode.cs
@@ -1407,13 +1407,14 @@ namespace Gauntlet
 		/// <param name="Result"></param>
 		protected virtual ITestReport CreateSimpleReportForHorde(TestResult Result)
 		{
-
 			if (Phases.Any())
 			{
 				OnPopulateTestPhases += PopulatePhaseReport;
 				OnPopulateTestPhases += (Report) => PopulateClosurePhase(Report);
 			}
 
+			string ArtifactPathForReport = ArtifactPath;
+			Horde.OnStepSummary += () => Horde.GenerateArtifactLink(ArtifactPathForReport);
 			BaseHordeReport HordeTestReport = null;
 			if (Globals.Params.ParseParam("UseTestDataV2"))
 			{
