[Horde][BuildHealth][Server] Add known AutomationTool_BuildGraphScript log event ID to BuildGraph unknown argument warning
#DisabledValidations CheckCopyright, CheckUnacceptableWords

diff --git a/Engine/Source/Programs/AutomationTool/BuildGraph/BgScriptReader.cs b/Engine/Source/Programs/AutomationTool/BuildGraph/BgScriptReader.cs
--- a/Engine/Source/Programs/AutomationTool/BuildGraph/BgScriptReader.cs
+++ b/Engine/Source/Programs/AutomationTool/BuildGraph/BgScriptReader.cs
@@ -428,7 +428,7 @@ namespace AutomationTool
 				if (!validArgumentNames.Contains(argumentName))
 				{
 					hasInvalidArguments = true;
-					logger.LogWarning("Unknown argument '{ArgumentName}' for '{Script}'", argumentName, file);
+					logger.LogWarning(KnownLogEvents.AutomationTool_BuildGraphScript, "Unknown argument '{ArgumentName}' for '{Script}'", argumentName, file);
 				}
 			}
 
