Fixed an issue where we were giving up waiting for a process to exit too early.

diff --git a/Engine/Source/Programs/AutomationTool/AutomationUtils/ProcessUtils.cs b/Engine/Source/Programs/AutomationTool/AutomationUtils/ProcessUtils.cs
--- a/Engine/Source/Programs/AutomationTool/AutomationUtils/ProcessUtils.cs
+++ b/Engine/Source/Programs/AutomationTool/AutomationUtils/ProcessUtils.cs
@@ -562,7 +562,7 @@ namespace AutomationTool
 				}
 			}
 
-			int waitTimeBeforeExit = 1000;
+			int waitTimeBeforeExit = Timeout.Infinite;
 
 			if (!(bStdOutSignalReceived && bStdErrSignalReceived))
 			{
