BuildGraph: Add flag to report diagnostics during execution of a step

ReportOnExecution="true" can be added to Trace, Warning and Error diagnositc elements to print out messages when a step is running.

diff --git a/Engine/Source/Programs/AutomationTool/BuildGraph/BgSchema.cs b/Engine/Source/Programs/AutomationTool/BuildGraph/BgSchema.cs
--- a/Engine/Source/Programs/AutomationTool/BuildGraph/BgSchema.cs
+++ b/Engine/Source/Programs/AutomationTool/BuildGraph/BgSchema.cs
@@ -973,6 +973,8 @@ namespace AutomationTool
 			propertyType.Name = GetTypeName(standardType);
 			propertyType.Attributes.Add(CreateSchemaAttribute("Message", ScriptSchemaStandardType.BalancedString, XmlSchemaUse.Required));
 			propertyType.Attributes.Add(CreateSchemaAttribute("If", ScriptSchemaStandardType.BalancedString, XmlSchemaUse.Optional));
+			propertyType.Attributes.Add(CreateSchemaAttribute("ReportOnExecution", ScriptSchemaStandardType.Boolean, XmlSchemaUse.Optional));
+			propertyType.Attributes.Add(CreateSchemaAttribute("ReportOnNodes", ScriptSchemaStandardType.NameList, XmlSchemaUse.Optional));
 			return propertyType;
 		}
 
diff --git a/Engine/Source/Programs/AutomationTool/BuildGraph/BgScriptReader.cs b/Engine/Source/Programs/AutomationTool/BuildGraph/BgScriptReader.cs
--- a/Engine/Source/Programs/AutomationTool/BuildGraph/BgScriptReader.cs
+++ b/Engine/Source/Programs/AutomationTool/BuildGraph/BgScriptReader.cs
@@ -1996,8 +1996,10 @@ namespace AutomationTool
 			if (await EvaluateConditionAsync(element))
 			{
 				string message = ReadAttribute(element, "Message");
+				bool reportOnExecution = ReadBooleanAttribute(element, "ReportOnExecution", false);
+				string[] reportOnNodes = ReadListAttribute(element, "ReportOnNodes");
 
-				BgDiagnosticDef diagnostic = new BgDiagnosticDef(element.Location.File.FullName, element.Location.LineNumber, level, message);
+				BgDiagnosticDef diagnostic = new(element.Location.File.FullName, element.Location.LineNumber, level, message, reportOnExecution, reportOnNodes);
 				if (_enclosingNode != null)
 				{
 					_enclosingNode.Diagnostics.Add(diagnostic);
diff --git a/Engine/Source/Programs/AutomationTool/BuildGraph/BuildGraph.cs b/Engine/Source/Programs/AutomationTool/BuildGraph/BuildGraph.cs
--- a/Engine/Source/Programs/AutomationTool/BuildGraph/BuildGraph.cs
+++ b/Engine/Source/Programs/AutomationTool/BuildGraph/BuildGraph.cs
@@ -720,6 +720,46 @@ namespace AutomationTool
 					return ExitCode.Error_Unknown;
 				}
 			}
+			else if (singleNode != null || listOnly)
+			{
+				List<BgDiagnosticDef> diagnostics = [];
+				foreach (BgDiagnosticDef diagnostic in graph.GetAllDiagnostics())
+				{
+					if (diagnostic.ReportOnExecution)
+					{
+						if (singleNode != null && diagnostic.ReportOnNodes.Length > 0)
+						{
+							if (diagnostic.ReportOnNodes.Contains(singleNode.Name, StringComparer.OrdinalIgnoreCase))
+							{
+								diagnostics.Add(diagnostic);
+							}
+						}
+						else
+						{
+							diagnostics.Add(diagnostic);
+						}
+					}
+				}
+				foreach (BgDiagnosticDef diagnostic in diagnostics)
+				{
+					if (diagnostic.Level == LogLevel.Information)
+					{
+						Logger.LogInformation("{File}({LineNumber}): {Message}", diagnostic.File, diagnostic.Line, diagnostic.Message);
+					}
+					else if (diagnostic.Level == LogLevel.Warning)
+					{
+						Logger.LogWarning("{File}({LineNumber}): warning: {Message}", diagnostic.File, diagnostic.Line, diagnostic.Message);
+					}
+					else
+					{
+						Logger.LogError("{File}({LineNumber}): error: {Message}", diagnostic.File, diagnostic.Line, diagnostic.Message);
+					}
+				}
+				if (diagnostics.Any(x => x.Level == LogLevel.Error))
+				{
+					return ExitCode.Error_Unknown;
+				}
+			}
 
 			// Export the graph to a file
 			if (exportFileName != null)
diff --git a/Engine/Source/Programs/Shared/EpicGames.BuildGraph/BgDiagnosticDef.cs b/Engine/Source/Programs/Shared/EpicGames.BuildGraph/BgDiagnosticDef.cs
--- a/Engine/Source/Programs/Shared/EpicGames.BuildGraph/BgDiagnosticDef.cs
+++ b/Engine/Source/Programs/Shared/EpicGames.BuildGraph/BgDiagnosticDef.cs
@@ -30,15 +30,27 @@ namespace EpicGames.BuildGraph
 		/// </summary>
 		public string Message { get; }
 
+		/// <summary>
+		///
+		/// </summary>
+		public bool ReportOnExecution { get; }
+
+		/// <summary>
+		/// 
+		/// </summary>
+		public string[] ReportOnNodes { get; }
+
 		/// <summary>
 		/// Constructor
 		/// </summary>
-		public BgDiagnosticDef(string file, int line, LogLevel level, string message)
+		public BgDiagnosticDef(string file, int line, LogLevel level, string message, bool reportOnExecution, string[] reportOnNodes)
 		{
 			File = file;
 			Line = line;
 			Level = level;
 			Message = message;
+			ReportOnExecution = reportOnExecution;
+			ReportOnNodes = reportOnNodes;
 		}
 	}
 }
