[Gauntlet] Disabling Horde issue generation for Device and Systemic related errors and warnings

#rnx

diff --git a/Engine/Source/Programs/AutomationTool/AutomationUtils/DeviceReservation.cs b/Engine/Source/Programs/AutomationTool/AutomationUtils/DeviceReservation.cs
--- a/Engine/Source/Programs/AutomationTool/AutomationUtils/DeviceReservation.cs
+++ b/Engine/Source/Programs/AutomationTool/AutomationUtils/DeviceReservation.cs
@@ -119,7 +119,7 @@ namespace AutomationTool.DeviceReservation
 					// There was an exception, warn if we've exceeded retry
 					if (RetryCurrent == RenewRetryMax)
 					{
-						Utils.Logger.LogWarning("Reservation renewal returned bad status: {Exception}", Ex.Message);
+						Utils.Logger.LogWarning(KnownLogEvents.Horde_BuildHealth_Ignore, "Reservation renewal returned bad status: {Exception}", Ex.Message);
 					}
 
 					// try again
@@ -135,7 +135,7 @@ namespace AutomationTool.DeviceReservation
 			}
 			catch (Exception Ex)
 			{
-				Utils.Logger.LogWarning("Reservation delete returned bad status: {Exception}", Ex.Message);
+				Utils.Logger.LogWarning(KnownLogEvents.Horde_BuildHealth_Ignore, "Reservation delete returned bad status: {Exception}", Ex.Message);
 			}
 
 		}
@@ -452,7 +452,7 @@ namespace AutomationTool.DeviceReservation
 
 					if (RetryCount == RetryMax)
 					{
-						Utils.Logger.LogWarning("Device reservation unsuccessful");
+						Utils.Logger.LogWarning(KnownLogEvents.Horde_BuildHealth_Ignore, "Device reservation unsuccessful");
 						string PoolName = string.IsNullOrEmpty(PoolID) ? "default" : PoolID;
 						string FinalMessage = $"Device reservation for {string.Join(", ", DeviceTypes)} from {PoolName} pool unsuccessful";
 						string MaxRetryWebExMessage = GetReservationErrorMessage(WebEx);
@@ -477,14 +477,14 @@ namespace AutomationTool.DeviceReservation
 						Message = WebExMessage;
 					}
 
-					Utils.Logger.LogWarning(Message);
+					Utils.Logger.LogWarning(KnownLogEvents.Horde_BuildHealth_Ignore, Message);
 					RetryCount++;
 					UnknownException = WebEx;
 				}
 				catch (Exception Ex)
 				{
 					UnknownException = Ex;
-					Utils.Logger.LogWarning("Device reservation unsuccessful: {Exception}", UnknownException.Message);
+					Utils.Logger.LogWarning(KnownLogEvents.Horde_BuildHealth_Ignore, "Device reservation unsuccessful: {Exception}", UnknownException.Message);
 				}
 			}
 		}
@@ -543,7 +543,7 @@ namespace AutomationTool.DeviceReservation
 			}
 			catch (Exception Ex)
 			{
-				Utils.Logger.LogWarning("Failed to report device {DeviceName}: {Exception}", DeviceId, Ex.Message);
+				Utils.Logger.LogWarning(KnownLogEvents.Horde_BuildHealth_Ignore, "Failed to report device {DeviceName}: {Exception}", DeviceId, Ex.Message);
 			}
 		}
 
diff --git a/Engine/Source/Programs/AutomationTool/Gauntlet/Framework/Gauntlet.TestExecutor.cs b/Engine/Source/Programs/AutomationTool/Gauntlet/Framework/Gauntlet.TestExecutor.cs
--- a/Engine/Source/Programs/AutomationTool/Gauntlet/Framework/Gauntlet.TestExecutor.cs
+++ b/Engine/Source/Programs/AutomationTool/Gauntlet/Framework/Gauntlet.TestExecutor.cs
@@ -264,7 +264,7 @@ namespace Gauntlet
 								}
 								else
 								{
-									Log.Error(KnownLogEvents.Gauntlet_TestEvent, "Test {Name} threw an exception during ready check. Ex: {Exception}", Node, ex);
+									Log.Error(KnownLogEvents.Horde_BuildHealth_Ignore, "Test {Name} threw an exception during ready check. Ex: {Exception}", Node, ex);
 									Node.AddTestEvent(new UnrealTestEvent(EventSeverity.Error, "Test Failed to Start", new List<string> { ex.Message }));
 									PendingTests[i] = null;
 									NodeInfo.TimeSetupBegan = NodeInfo.TimeSetupEnded = NodeInfo.TimeTestEnded = DateTime.Now;
@@ -336,7 +336,7 @@ namespace Gauntlet
 									{
 										TestToStart.TimeSetupEnded = TestToStart.TimeTestEnded = DateTime.Now;
 										CompletedTests.Add(TestToStart);
-										Log.Error(KnownLogEvents.Gauntlet_TestEvent, "Test {Name} failed to start", TestToStart);
+										Log.Error(KnownLogEvents.Horde_BuildHealth_Ignore, "Test {Name} failed to start", TestToStart);
 									}
 									else
 									{
@@ -427,7 +427,7 @@ namespace Gauntlet
 
 						if (Elapsed >= 5)
 						{
-							Log.Error(KnownLogEvents.Gauntlet_TestEvent, "Giving up waiting for tests after {Elapsed:00} seconds", Elapsed);
+							Log.Error(KnownLogEvents.Horde_BuildHealth_Ignore, "Giving up waiting for tests after {Elapsed:00} seconds", Elapsed);
 							break;
 						}
 					}
@@ -438,7 +438,7 @@ namespace Gauntlet
 						Log.Info("Attempting to join test start thread: {ThreadName} | {ThreadId}", Thread.Name, Thread.ManagedThreadId);
 						if (!Thread.Join(ThreadTimeout))
 						{
-							Log.Warning("Failed to join the thread {Thread} | {ThreadId} after {Timeout} seconds.",
+							Log.Warning(KnownLogEvents.Horde_BuildHealth_Ignore, "Failed to join the thread {Thread} | {ThreadId} after {Timeout} seconds.",
 								Thread.Name, Thread.ManagedThreadId, ThreadTimeout.TotalSeconds);
 						}
 					});
@@ -458,7 +458,7 @@ namespace Gauntlet
 
 						if (Res != TestResult.Failed)
 						{
-							Log.Warning(KnownLogEvents.Gauntlet_TestEvent, "Ticking of cancelled test {Name} returned {Result}", TestInfo.TestNode.Name, Res);
+							Log.Warning(KnownLogEvents.Horde_BuildHealth_Ignore, "Ticking of cancelled test {Name} returned {Result}", TestInfo.TestNode.Name, Res);
 						}
 					}
 				}
@@ -602,7 +602,7 @@ namespace Gauntlet
 			}
 			catch (Exception Ex)
 			{
-				Log.Error(KnownLogEvents.Gauntlet_TestEvent, "Test {Name} threw an exception during launch. Skipping test. Ex: {Exception}", Name, Ex);
+				Log.Error(KnownLogEvents.Horde_BuildHealth_Ignore, "Test {Name} threw an exception during launch. Skipping test. Ex: {Exception}", Name, Ex);
 			}
 
 			return false;			
@@ -852,7 +852,7 @@ namespace Gauntlet
 						Log.Info("{Name} result={Result}", TestInfo, TestInfo.FinalResult);
 						if (!string.IsNullOrEmpty(TestInfo.CancellationReason) && TestInfo.FinalResult != TestResult.Cancelled)
 						{
-							Log.Error(KnownLogEvents.Gauntlet_TestEvent, "{CancellationReason}", TestInfo.CancellationReason);
+							Log.Error(KnownLogEvents.Horde_BuildHealth_Ignore, "{CancellationReason}", TestInfo.CancellationReason);
 						}
 
 						if (!Options.DeferReports)
@@ -867,7 +867,7 @@ namespace Gauntlet
 						}
 						catch (System.Exception ex)
 						{
-							Log.Error(KnownLogEvents.Gauntlet_TestEvent, "Test {Name} threw an exception while cleaning up. Ex: {Exception}",
+							Log.Error(KnownLogEvents.Horde_BuildHealth_Ignore, "Test {Name} threw an exception while cleaning up. Ex: {Exception}",
 								TestInfo.TestNode.Name, ex);
 						}
 					}
@@ -879,12 +879,12 @@ namespace Gauntlet
 				{
 					if (TestIsRunning)
 					{
-						Log.Warning(KnownLogEvents.Gauntlet_TestEvent, "Cancelled Test {Name} threw an exception while stopping. Ex: {Exception}", 
+						Log.Warning(KnownLogEvents.Horde_BuildHealth_Ignore, "Cancelled Test {Name} threw an exception while stopping. Ex: {Exception}", 
 							TestInfo.TestNode.Name, ex);
 					}
 					else
 					{
-						Log.Error(KnownLogEvents.Gauntlet_TestEvent, "Test {Name} threw an exception while stopping. Ex: {Exception}",
+						Log.Error(KnownLogEvents.Horde_BuildHealth_Ignore, "Test {Name} threw an exception while stopping. Ex: {Exception}",
 							TestInfo.TestNode.Name, ex);
 					}
 
diff --git a/Engine/Source/Programs/Shared/EpicGames.Horde/Issues/Handlers/GauntletIssueHandler.cs b/Engine/Source/Programs/Shared/EpicGames.Horde/Issues/Handlers/GauntletIssueHandler.cs
--- a/Engine/Source/Programs/Shared/EpicGames.Horde/Issues/Handlers/GauntletIssueHandler.cs
+++ b/Engine/Source/Programs/Shared/EpicGames.Horde/Issues/Handlers/GauntletIssueHandler.cs
@@ -111,6 +111,14 @@ namespace EpicGames.Horde.Issues.Handlers
 			KnownLogEvents.Gauntlet_UnrealEngineTestEvent
 		};
 
+		/// <summary>
+		///   Gauntlet events to skip issue generation
+		/// </summary>
+		static readonly HashSet<EventId> s_skipGauntletEvents = new HashSet<EventId>
+		{
+			KnownLogEvents.Gauntlet_DeviceEvent
+		};
+
 		/// <summary>
 		///  Known summaries that are not specific enough
 		/// </summary>
@@ -271,6 +279,12 @@ namespace EpicGames.Horde.Issues.Handlers
 		{
 			if (issueEvent.EventId != null && IsMatchingEventId(issueEvent.EventId.Value))
 			{
+				if (s_skipGauntletEvents.Contains(issueEvent.EventId.Value))
+				{
+					// ignore Device event
+					return true;
+				}
+
 				bool isSevereEvent = s_knownSevereGauntletEvents.Contains(issueEvent.EventId.Value);
 				bool isTestEvent = !isSevereEvent && s_knownTestGauntletEvents.Contains(issueEvent.EventId.Value);
 				bool isErrorEvent = issueEvent.Severity >= LogLevel.Error;
