Push error to phase tracking
+ don't skip logging if a phase is running

#rnx

diff --git a/Engine/Source/Programs/AutomationTool/Gauntlet/Framework/Base/Gauntlet.TestTracking.cs b/Engine/Source/Programs/AutomationTool/Gauntlet/Framework/Base/Gauntlet.TestTracking.cs
--- a/Engine/Source/Programs/AutomationTool/Gauntlet/Framework/Base/Gauntlet.TestTracking.cs
+++ b/Engine/Source/Programs/AutomationTool/Gauntlet/Framework/Base/Gauntlet.TestTracking.cs
@@ -61,6 +61,21 @@ namespace Gauntlet
 				Data = InData;
 			}
 
+			/// <summary>
+			/// Mark the phase as started
+			/// </summary>
+			/// <param name="InData"></param>
+			/// <exception cref="AutomationException"></exception>
+			public void Start(object? InData = null)
+			{
+				if (StartedAt != null)
+				{
+					throw new AutomationException($"Phase '{Name}' already started.");
+				}
+				StartedAt = DateTime.UtcNow;
+				if (InData != null) Data = InData;
+			}
+
 			/// <summary>
 			/// Mark the phase as finished
 			/// </summary>
@@ -204,12 +219,7 @@ namespace Gauntlet
 					(Name, ExistingPhase) =>
 					{
 						bAlreadyAdded = true;
-						if (ExistingPhase.StartedAt != null)
-						{
-							throw new AutomationException($"Phase '{Name}' already started.");
-						}
-						ExistingPhase.StartedAt = DateTime.UtcNow;
-						if (Data != null) ExistingPhase.Data = Data;
+						ExistingPhase.Start(Data);
 						return ExistingPhase;
 					}
 				);
@@ -259,6 +269,20 @@ namespace Gauntlet
 				return FetchedPhase;
 			}
 
+			/// <summary>
+			/// Set current phase
+			/// </summary>
+			/// <param name="Name"></param>
+			/// <returns></returns>
+			public bool SetCurrent(string Name)
+			{
+				Phase? FetchedPhase = null;
+				_manifest.TryGetValue(Name, out FetchedPhase);
+				_current = FetchedPhase;
+
+				return FetchedPhase != null;
+			}
+
 			IEnumerator IEnumerable.GetEnumerator()
 			{
 				return (IEnumerator)GetEnumerator();
diff --git a/Engine/Source/Programs/AutomationTool/Gauntlet/Unreal/Base/Gauntlet.UnrealTestNode.cs b/Engine/Source/Programs/AutomationTool/Gauntlet/Unreal/Base/Gauntlet.UnrealTestNode.cs
--- a/Engine/Source/Programs/AutomationTool/Gauntlet/Unreal/Base/Gauntlet.UnrealTestNode.cs
+++ b/Engine/Source/Programs/AutomationTool/Gauntlet/Unreal/Base/Gauntlet.UnrealTestNode.cs
@@ -149,7 +149,7 @@ namespace Gauntlet
 		{
 			Phases.Current?.AddEvent(Logging.LogLevel.Error, Message, Args);
 			AddTestEvent(new UnrealTestEvent(EventSeverity.Error, Message, Args));
-			if (Phases.Current == null && !LogWarningsAndErrorsAfterSummary) { Log.Error(KnownLogEvents.Gauntlet_TestEvent, Message, Args); }
+			if (!LogWarningsAndErrorsAfterSummary) { Log.Error(KnownLogEvents.Gauntlet_TestEvent, Message, Args); }
 			if (GetTestStatus() == TestStatus.Complete && GetTestResult() == TestResult.Passed)
 			{
 				SetUnrealTestResult(TestResult.Failed);
@@ -165,7 +165,7 @@ namespace Gauntlet
 		{
 			Phases.Current?.AddEvent(Logging.LogLevel.Warning, Message, Args);
 			AddTestEvent(new UnrealTestEvent(EventSeverity.Warning, Message, Args));
-			if (Phases.Current == null && !LogWarningsAndErrorsAfterSummary) { Log.Warning(KnownLogEvents.Gauntlet_TestEvent, Message, Args); }
+			if (!LogWarningsAndErrorsAfterSummary) { Log.Warning(KnownLogEvents.Gauntlet_TestEvent, Message, Args); }
 		}
 
 		/// <summary>
@@ -177,7 +177,7 @@ namespace Gauntlet
 		{
 			Phases.Current?.AddEvent(Logging.LogLevel.Information, Message, Args);
 			AddTestEvent(new UnrealTestEvent(EventSeverity.Info, Message, Args));
-			if (Phases.Current == null) { Log.Info(KnownLogEvents.Gauntlet_TestEvent, Message, Args); }
+			Log.Info(KnownLogEvents.Gauntlet_TestEvent, Message, Args);
 		}
 
 		// Begin UnrealTestNode properties and members
