SmartBots - Fix RPC StartClientPilot argument mismatch
#rnx
#tests local compile and test run

diff --git a/Engine/Source/Programs/AutomationTool/Gauntlet/Framework/RpcFramework/Gauntlet.RpcExecutor.cs b/Engine/Source/Programs/AutomationTool/Gauntlet/Framework/RpcFramework/Gauntlet.RpcExecutor.cs
--- a/Engine/Source/Programs/AutomationTool/Gauntlet/Framework/RpcFramework/Gauntlet.RpcExecutor.cs
+++ b/Engine/Source/Programs/AutomationTool/Gauntlet/Framework/RpcFramework/Gauntlet.RpcExecutor.cs
@@ -821,6 +821,11 @@ namespace Gauntlet
 								CastSucceeded = ArgValue is string;
 								break;
 							}
+						case "object":
+							{
+								CastSucceeded = ArgValue is not null;
+								break;
+							}
 						default:
 							{
 								// If listed object type is something not handled in this switch, either add a new handler or we basically just
