Additional Helper Commands in GoogleSheetsHelper

diff --git a/Engine/Source/Programs/AutomationTool/AutomationUtils/GoogleSheetsHelper.cs b/Engine/Source/Programs/AutomationTool/AutomationUtils/GoogleSheetsHelper.cs
--- a/Engine/Source/Programs/AutomationTool/AutomationUtils/GoogleSheetsHelper.cs
+++ b/Engine/Source/Programs/AutomationTool/AutomationUtils/GoogleSheetsHelper.cs
@@ -162,7 +162,7 @@ namespace SheetsHelper
 		{
 			Refresh();
 
-			var Sheet = SourceSpreadsheet.Sheets.Where(S => S.Properties.Title == Title).FirstOrDefault();
+			var Sheet = SourceSpreadsheet.Sheets.FirstOrDefault(S => S.Properties.Title == Title);
 
 			if (Sheet == null)
 			{
@@ -172,6 +172,37 @@ namespace SheetsHelper
 			return new SheetWrapper(Service, SpreadsheetID, Sheet);
 		}
 
+		// Checks if the sheet with the given name exists
+		public bool SheetExists(string Title)
+		{
+			Refresh();
+
+			var Sheet = SourceSpreadsheet.Sheets.FirstOrDefault(S => S.Properties.Title == Title);
+			return Sheet != null;
+		}
+
+		// Deletes the sheet with the given name
+		public void DeleteSheet(string Title)
+		{
+			Refresh();
+
+			Sheet Sheet = SourceSpreadsheet.Sheets.FirstOrDefault(S => S.Properties.Title == Title);
+
+			if (Sheet != null)
+			{
+				var Request = new Request();
+				Request.DeleteSheet = new DeleteSheetRequest();
+				Request.DeleteSheet.SheetId = Sheet.Properties.SheetId;
+				var BatchRequest = new BatchUpdateSpreadsheetRequest();
+
+				BatchRequest.Requests = new List<Request> { Request };
+
+				SheetsUtils.TryCheckQuotaSync();
+				Service.Spreadsheets.BatchUpdate(BatchRequest, SpreadsheetID).Execute();
+			}
+		}
+
+
 		public SheetWrapper AddSheet(string Title, int Index = 0)
 		{
 			var Request = new Request();
@@ -405,7 +436,7 @@ namespace SheetsHelper
 		{
 			SheetsUtils.TryCheckQuotaSync();
 			var Response = Service.Spreadsheets.Get(SpreadsheetID).Execute();
-			SourceSheet = Response.Sheets.Where(S => S.Properties.SheetId == SheetID).FirstOrDefault();
+			SourceSheet = Response.Sheets.FirstOrDefault(S => S.Properties.SheetId == SheetID);
 		}
 
         // Returns a 2D container of the values of cells within the given range. The elements of the outer container are determined by the MajorDimension. For example: this returns a list of lists of cells in each row by default.
@@ -512,6 +543,21 @@ namespace SheetsHelper
 			SheetsUtils.TryCheckQuotaSync();
             Service.Spreadsheets.Values.BatchUpdate(batchRequest, SpreadsheetID).Execute();
 		}
-		
+
+		public void SetRow(List<IList<object>> Fields)
+		{
+			ValueRange EntryRange = new ValueRange();
+			EntryRange.MajorDimension = "ROWS";
+			EntryRange.Values = Fields;
+			EntryRange.Range = "A1";
+
+			BatchUpdateValuesRequest batchRequest = new BatchUpdateValuesRequest();
+			batchRequest.ValueInputOption = "RAW";
+			batchRequest.Data = new List<ValueRange> { EntryRange };
+
+			SheetsUtils.TryCheckQuotaSync();
+			Service.Spreadsheets.Values.BatchUpdate(batchRequest, SpreadsheetID).Execute();
+		}
+
 	}
 }
