SymStoreTask: Split file list in to batches of unique filenames to work around symstore issue when processing duplicate files in different directories

#rnx

diff --git a/Engine/Source/Programs/AutomationTool/BuildGraph/Tasks/SymStoreTask.cs b/Engine/Source/Programs/AutomationTool/BuildGraph/Tasks/SymStoreTask.cs
--- a/Engine/Source/Programs/AutomationTool/BuildGraph/Tasks/SymStoreTask.cs
+++ b/Engine/Source/Programs/AutomationTool/BuildGraph/Tasks/SymStoreTask.cs
@@ -123,7 +123,7 @@ namespace AutomationTool.Tasks
 		public override Task ExecuteAsync(JobContext job, HashSet<FileReference> buildProducts, Dictionary<string, HashSet<FileReference>> tagNameToFileSet)
 		{
 			// Find the matching files
-			List<FileReference> files = ResolveFilespec(Unreal.RootDirectory, _parameters.Files, tagNameToFileSet).ToList();
+			HashSet<FileReference> files = ResolveFilespec(Unreal.RootDirectory, _parameters.Files, tagNameToFileSet);
 
 			// Get the symbol store directory
 			DirectoryReference storeDir = ResolveDirectory(_parameters.StoreDir);
@@ -142,8 +142,21 @@ namespace AutomationTool.Tasks
 
 			CommandUtils.OptionallyTakeLock(targetPlatform.SymbolServerRequiresLock, storeDir, TimeSpan.FromMinutes(60), () =>
 			{
-				if (!targetPlatform.PublishSymbols(storeDir, files, _parameters.IndexSources, sourceFiles,
-					_parameters.Product, _parameters.Branch, _parameters.Change, _parameters.BuildVersion, _parameters.VFSMapping))
+				// symstore has issues with processing same file in multiple folders, so run multiple passes with only unique filenames
+				bool anyError = false;
+				while (files.Count > 0)
+				{
+					List<FileReference> publishFiles = files.DistinctBy(x => x.GetFileName().ToUpperInvariant()).ToList();
+					files = [.. files.Except(publishFiles)];
+
+					if (!targetPlatform.PublishSymbols(storeDir, publishFiles, _parameters.IndexSources, sourceFiles,
+						_parameters.Product, _parameters.Branch, _parameters.Change, _parameters.BuildVersion, _parameters.VFSMapping))
+					{
+						anyError = true;
+					}
+				}
+
+				if (anyError)
 				{
 					throw new AutomationException("Failure publishing symbol files.");
 				}
