AutomationToolStaging: ZenServer: Add feedback to exception message when ReadZenCookedFilesFromZenServer fails.
#rnx

diff --git a/Engine/Source/Programs/AutomationTool/Scripts/CopyBuildToStagingDirectory.Automation.cs b/Engine/Source/Programs/AutomationTool/Scripts/CopyBuildToStagingDirectory.Automation.cs
--- a/Engine/Source/Programs/AutomationTool/Scripts/CopyBuildToStagingDirectory.Automation.cs
+++ b/Engine/Source/Programs/AutomationTool/Scripts/CopyBuildToStagingDirectory.Automation.cs
@@ -967,9 +967,10 @@ namespace AutomationScripts
 
 			bool bAttemptAutoLaunchOnFailure = bAutoLaunch;
 			HttpResponseMessage HttpGetResult = null;
+			string HostName;
 			for (; ; )
 			{
-				string HostName = GetZenHostName(SC);
+				HostName = GetZenHostName(SC);
 				HttpClient HttpClient = CreateZenHttpClient(SC);
 				string Uri = string.Format("http://{0}/prj/{1}/oplog/{2}/entries", HostName, ProjectId, OplogId);
 				Uri += "?trim_by_referencedset=true";
@@ -981,7 +982,7 @@ namespace AutomationScripts
 					HttpGetResult = HttpClient.Send(Request);
 					break;
 				}
-				catch
+				catch (Exception Ex)
 				{
 					if (bAttemptAutoLaunchOnFailure)
 					{
@@ -996,13 +997,21 @@ namespace AutomationScripts
 					}
 					else
 					{
-						throw new AutomationException(String.Format("Failed sending oplog request to Zen at {0}. Ensure that the server is running.", HostName));
+						string RetryMessage = bAutoLaunch
+							? "Failed even after attempting restart."
+							: "bAutoLaunch is false so did not attempt restart. Ensure that the server is running.";
+						string ExceptionMessage = String.Format("Failed sending oplog request to Zen at {0} for oplog {1}.{2}: {3}. {4}",
+								HostName, ProjectId, OplogId, Ex.Message, RetryMessage);
+						throw new AutomationException(ExceptionMessage);
 					}
 				}
 			}
 			if (!HttpGetResult.IsSuccessStatusCode)
 			{
-				throw new AutomationException(String.Format("Failed reading oplog {0}.{1} from Zen. Ensure that cooking was successful.", ProjectId, OplogId));
+				string ExceptionMessage = String.Format("Failed reading oplog from Zen at {0} for oplog {1}.{2}, "
+					+ "HttpClient.Send returned StatusCode {3}. Ensure that cooking was successful.",
+					HostName, ProjectId, OplogId, HttpGetResult.StatusCode);
+				throw new AutomationException(ExceptionMessage);
 			}
 			Task<byte[]> ReadOplogTask = HttpGetResult.Content.ReadAsByteArrayAsync();
 			ReadOplogTask.Wait();
