Staging: ZenServer: Fix possible timeouts reading the oplog. In ReadZenCookedFilesFromZenServer, first quickly health check the server and launch if necessary, and then fetch the oplog with a timeout of 20 minutes rather than the default of 100 seconds.
Refactor CopyBuildToStaging's use of ZenServer HttpClients to move the handling of UnixSocketPath inside ZenUtils.
#rnx

diff --git a/Engine/Source/Programs/AutomationTool/AutomationUtils/ZenUtils.cs b/Engine/Source/Programs/AutomationTool/AutomationUtils/ZenUtils.cs
--- a/Engine/Source/Programs/AutomationTool/AutomationUtils/ZenUtils.cs
+++ b/Engine/Source/Programs/AutomationTool/AutomationUtils/ZenUtils.cs
@@ -12,6 +12,7 @@ using System.Collections.Generic;
 using System.IO;
 using System.Linq;
 using System.Net.Http;
+using System.Net.Sockets;
 using System.Net.WebSockets;
 using System.Runtime.InteropServices;
 using System.Text;
@@ -45,9 +46,38 @@ namespace AutomationUtils
 
 		public string HostName { get; private set; } = string.Empty;
 		public int HostPort { get; private set; } = 0;
+		/// <summary>
+		/// Return the "HostName:HostPort" string that should be used as the "servername" when creating http URLs, e.g.
+		/// $"http:///{HostName}/prj/{Project}/oplog/{Oplog}"
+		/// When using unixsockets we use a valid debugname string rather than than the "servername" used to open the
+		/// socket for the HttpClient, which is instead a path to the .sock file on disk.
+		/// @see GetRobustHostNameAndPortStrings.
+		/// </summary>
+		public string GetHttpHostNameAndPortString()
+		{
+			return UsingUnixSocket ? "unixsocket" : $"{HostName}:{HostPort}";
+		}
 		public bool IsValid { get; private set; } = false;
+		public bool UsingUnixSocket { get; private set; } = false;
 
 		public static ILogger Logger => Log.Logger;
+		public static int DefaultPort { get { return 8558; } }
+
+		public void InitializeFromHostNameAndPort(string SocketHostNameAndPort)
+		{
+			SplitHostNameAndPort(SocketHostNameAndPort, out string LocalSocketHostName, out int LocalHostPort);
+			InitializeHostNameInternal(LocalSocketHostName, LocalHostPort);
+		}
+
+		public void InitializeFromHostNameAndPort(string InSocketHostName, int InHostPort)
+		{
+			InitializeHostNameInternal(InSocketHostName, InHostPort);
+		}
+
+		public void InitializeFromHostName(string InSocketHostName)
+		{
+			InitializeHostNameInternal(InSocketHostName, DefaultPort);
+		}
 
 		public static ZenRunContext ReadFromContextFile(FileReference ContextFile)
 		{
@@ -109,15 +139,160 @@ namespace AutomationUtils
 			int PID = LockObject["pid"].AsInt32();
 			int ServerPort = LockObject["port"].AsInt32();
 
-			IsValid = PID > 0 && (ServerPort > 0 && ServerPort <= 0xffff);
+			bool ServerIsValid = PID > 0 && (ServerPort > 0 && ServerPort <= 0xffff);
+			if (ServerIsValid)
+			{
+				string ServerHostName = "localhost"; // for now only support localhost
+				InitializeHostNameInternal(ServerHostName, ServerPort);
+			}
+			else
+			{
+				HostName = string.Empty;
+				HostPort = 0;
+				IsValid = false;
+			}
+		}
+
+		private void InitializeHostNameInternal(string InHostName, int InHostPort)
+		{
+			HostName = InHostName;
+			HostPort = InHostPort;
+			IsValid = true;
+			UsingUnixSocket = IsUnixSocketPath(HostName);
+			if (UsingUnixSocket)
+			{
+				if (HostPort > 0 && HostPort != DefaultPort)
+				{
+					Logger.LogWarning($"HostPort={HostPort} was unexpectedly provided when using a UnixSocket. UnixSockets communicate with a file and do not need a port. Ignoring the port.");
+				}
+				HostPort = 0;
+			}
+		}
+
+		/// <summary>
+		/// Return an HttpClient that can be used to communicate with the ZenServer. If a UnixSocket is being used,
+		/// the SocketsHttpHandler of the HttpClient if set to read from the socket file.
+		/// </summary>
+		/// <returns>A new HttpClient, with SocketsHttpHandler set if necessary.</returns>
+		private HttpClient CreateHttpClientInternal()
+		{
+			if (!UsingUnixSocket)
+			{
+				return new HttpClient();
+			}
+			return CreateUnixSocketHttpClient(HostName);
+		}
+
+		/// <summary>
+		/// Return an HttpClient that communicates with a ZenServer at the given HostName. If the hostname specifies
+		/// a UnixSocket, the SocketsHttpHandler of the HttpClient if set to read from the socket file.
+		/// </summary>
+		/// <returns>A new HttpClient, with SocketsHttpHandler set.</returns>
+		public static HttpClient CreateHttpClient(string InSocketHostName)
+		{
+			if (!IsUnixSocketPath(InSocketHostName))
+			{
+				return new HttpClient();
+			}
+			return CreateUnixSocketHttpClient(InSocketHostName);
+		}
+
+		/// <summary>
+		/// Return an HttpClient that communicates with a ZenServer through a UnixSocketPath. Sets the SocketsHttpHandler
+		/// of the HttpClient to read from the socket file.
+		/// </summary>
+		/// <returns>A new HttpClient, with SocketsHttpHandler set.</returns>
+		public static HttpClient CreateUnixSocketHttpClient(string UnixSocketPath)
+		{
+			var SocketsHandler = new SocketsHttpHandler
+			{
+				ConnectCallback = async (Context, Token) => {
+					var Socket = new Socket(AddressFamily.Unix, SocketType.Stream, ProtocolType.IP);
+					var EndPoint = new UnixDomainSocketEndPoint(UnixSocketPath);
+					await Socket.ConnectAsync(EndPoint);
+					return new NetworkStream(Socket, ownsSocket: true);
+				}
+			};
+			return new HttpClient(SocketsHandler);
+		}
+
+		/// <summary>
+		/// Split a string "Hostname:HostPort" into the two parts. Silently handle edge cases such as a missing
+		/// HostPort and colons in the Hostname. Leading and trailing spaces are trimmed, but spaces before or
+		/// after the colon are not trimmed and are treated as part of the hostname.
+		/// Valid strings:
+		/// "hostname", "1.2.3.4", "[::1]"
+		/// "hostname:1234", "1.2.3.4:1234", "[::1]:1234"
+		/// </summary>
+		public static void SplitHostNameAndPort(string HostNameAndPort, out string OutHostName, out int OutHostPort)
+		{
+			HostNameAndPort = HostNameAndPort.Trim();
+			string[] Parts = HostNameAndPort.Split(':');
 
-			if (IsValid)
+			if (Parts.Length >= 2)
+			{
+				int PortIndex = Parts.Length - 1;
+				if (!int.TryParse(Parts[PortIndex], out OutHostPort))
+				{
+					OutHostName = HostNameAndPort;
+					OutHostPort = DefaultPort;
+				}
+				else
+				{
+					OutHostName = string.Join(':', Parts.Take(PortIndex));
+				}
+			}
+			else
 			{
-				HostPort = ServerPort;
-				HostName = "localhost"; // for now only support localhost
+				OutHostName = HostNameAndPort;
+				OutHostPort = DefaultPort;
 			}
 		}
 
+		/// <summary>
+		/// Returns the hostname and port connection information as a single string, e.g. "localhost:8558".
+		/// When not using unixsockets, OutSocketHostName and OutHttpHostName are identical, but when using unixsockets
+		/// they can be different.
+		/// 
+		/// When using unix sockets, OutSocketHostNameAndPort is the unix filepath of the .sock file and is used in
+		/// an HttpClient's HttpMessageHandler. OutHttpHostNameAndPort is used in the URI passed to HttpRequestMessage.
+		/// It is used only for display to the user because the HttpMessageHandler bypasses it, but it needs to be a
+		/// valid hostname string and can have no path separators. When the hostname specifies a unix socket we
+		/// therefore return e.g. ("/usr/path/zenserver.sock", "unixsocket").
+		/// </summary>
+		public static void GetRobustHostNameAndPortStrings(string InHostName, int InHostPort,
+			out string SocketHostNameAndPort, out string HttpHostNameAndPort)
+		{
+			if (IsUnixSocketPath(InHostName))
+			{
+				if (InHostPort > 0 && InHostPort != DefaultPort)
+				{
+					Logger.LogWarning($"HostPort={InHostPort} was unexpectedly provided when using a UnixSocket HostName=\"{InHostName}\". " +
+						"UnixSockets communicate with a file and do not need a port. Ignoring the port.");
+				}
+				SocketHostNameAndPort = InHostName;
+				HttpHostNameAndPort = "unixsocket";
+			}
+			else
+			{
+				SocketHostNameAndPort = InHostName;
+				if (InHostPort > 0)
+				{
+					SocketHostNameAndPort += ":" + InHostPort;
+				}
+				HttpHostNameAndPort = SocketHostNameAndPort;
+			}
+		}
+
+		/// <summary>
+		/// Returns whether a given HostName (string must not include a :HostPort suffix) is a unix socket path,
+		/// which requires communicating through the socket file rather than an http uri.
+		/// </summary>
+		public static bool IsUnixSocketPath(string InHostName)
+		{
+			return InHostName.EndsWith(".sock");
+		}
+
 		/// <summary>
 		/// Queries the server's health status to check if the the service is running
 		/// </summary>
@@ -129,11 +304,11 @@ namespace AutomationUtils
 				return false;
 			}
 
-			string Uri = string.Format("http://{0}:{1}/health/ready", HostName, HostPort);
+			string Uri = string.Format("http://{0}/health/ready", GetHttpHostNameAndPortString());
 			using var Request = new HttpRequestMessage(HttpMethod.Get, Uri);
 
 			HttpResponseMessage HttpGetResult = null;
-			HttpClient Client = new HttpClient();
+			HttpClient Client = CreateHttpClientInternal();
 
 			try
 			{
@@ -268,8 +443,8 @@ namespace AutomationUtils
 			}
 
 			HttpResponseMessage HttpGetResult = null;
-			HttpClient Client = new HttpClient();
-			string Uri = string.Format("http://{0}:{1}/ws", HostName, HostPort);
+			HttpClient Client = CreateHttpClientInternal();
+			string Uri = string.Format("http://{0}/ws", GetHttpHostNameAndPortString());
 			using var Request = new HttpRequestMessage(HttpMethod.Get, Uri);
 			Request.Headers.Add("Accept", "application/json");
 
@@ -460,6 +635,26 @@ namespace AutomationUtils
 			return Context.IsServiceRunning();
 		}
 
+		public static bool IsZenServerRunning(string SocketHostNameAndPort)
+		{
+			ZenRunContext ZenServerContext = new ZenRunContext();
+			ZenServerContext.InitializeFromHostNameAndPort(SocketHostNameAndPort);
+			return ZenServerContext.IsServiceRunning();
+		}
+
+		public static bool IsZenServerRunning(string SocketHostName, int HostPort)
+		{
+			ZenRunContext ZenServerContext = new ZenRunContext();
+			ZenServerContext.InitializeFromHostNameAndPort(SocketHostName, HostPort);
+			return ZenServerContext.IsServiceRunning();
+		}
+
+		public static HttpClient CreateHttpClient(string SocketHostNameAndPort)
+		{
+			ZenRunContext.SplitHostNameAndPort(SocketHostNameAndPort, out string HostName, out int HostPort);
+			return ZenRunContext.CreateHttpClient(HostName);
+		}
+
 		public static JsonSerializerOptions GetDefaultJsonSerializerOptions()
 		{
 			JsonSerializerOptions Options = new();
diff --git a/Engine/Source/Programs/AutomationTool/Scripts/CopyBuildToStagingDirectory.Automation.cs b/Engine/Source/Programs/AutomationTool/Scripts/CopyBuildToStagingDirectory.Automation.cs
--- a/Engine/Source/Programs/AutomationTool/Scripts/CopyBuildToStagingDirectory.Automation.cs
+++ b/Engine/Source/Programs/AutomationTool/Scripts/CopyBuildToStagingDirectory.Automation.cs
@@ -916,63 +916,76 @@ namespace AutomationScripts
 			return Field.AsHash();
 		}
 
-		private static HttpClient CreateZenHttpClient(DeploymentContext SC)
-		{
-			if (SC.PackageStoreData == null || SC.PackageStoreData.ZenServerStore == null || !SC.PackageStoreData.ZenServerStore.HostName.EndsWith(".sock"))
-			{
-				return new HttpClient();
-			}
-
-			string UnixSocketPath = SC.PackageStoreData.ZenServerStore.HostName;
-			var SocketsHandler = new SocketsHttpHandler {
-				ConnectCallback = async (Context, Token) => {
-					var Socket = new Socket(AddressFamily.Unix, SocketType.Stream, ProtocolType.IP);
-					var EndPoint = new UnixDomainSocketEndPoint(UnixSocketPath);
-					await Socket.ConnectAsync(EndPoint);
-					return new NetworkStream(Socket, ownsSocket:true);
-				}
-			};
-			return new HttpClient(SocketsHandler);
-		}
-
-		private static string GetZenHostName(DeploymentContext SC)
+		/// <summary>
+		/// Returns the zen hostname and port connection information as a single string, e.g. "localhost:8558".
+		/// When not using unixsockets, SocketHostName and HttpHostName are identical, but when using unixsockets
+		/// they can be different. @see ZenRunContext.GetRobustHostNameAndPortStrings.
+		/// </summary>
+		private static void GetZenHostNameAndPort(DeploymentContext SC, out string SocketHostNameAndPort,
+			out string HttpHostNameAndPort)
 		{
 			if (SC.PackageStoreData == null || SC.PackageStoreData.ZenServerStore == null)
 			{
-				return "localhost:8558";
+				SocketHostNameAndPort = $"localhost:{ZenRunContext.DefaultPort}";
+				HttpHostNameAndPort = SocketHostNameAndPort;
 			}
-
-			var ZenServer = SC.PackageStoreData.ZenServerStore;
-
-			string HostName = ZenServer.HostName;
-			if (HostName.EndsWith(".sock"))
-			{
-				HostName = "unixsocket";
-			}
-			
-			if (ZenServer.HostPort > 0)
+			else
 			{
-				HostName += ":" + ZenServer.HostPort;
+				GetZenHostNameAndPort(SC.PackageStoreData.ZenServerStore, out SocketHostNameAndPort,
+					out HttpHostNameAndPort);
 			}
-
-			return HostName;
+		}
+		/// <summary>
+		/// @see GetZenHostNameAndPort(DeploymentContext SC,...)
+		/// </summary>
+		private static void GetZenHostNameAndPort(ZenServerStoreData ZenServer, out string SocketHostNameAndPort,
+			out string HttpHostNameAndPort)
+		{
+			ZenRunContext.GetRobustHostNameAndPortStrings(ZenServer.HostName, ZenServer.HostPort,
+				out SocketHostNameAndPort, out HttpHostNameAndPort);
 		}
 
-		private static List<ZenCookedFile> ReadZenCookedFilesFromZenServer(ProjectParams Params, DeploymentContext SC, bool bAutoLaunch, string PackageStoreFileArgName, string PackageStoreFileArgValue)
+		private static List<ZenCookedFile> ReadZenCookedFilesFromZenServer(ProjectParams Params, DeploymentContext SC,
+			bool bAutoLaunch, string PackageStoreFileArgName, string PackageStoreFileArgValue)
 		{
 			List<ZenCookedFile> ZenCookedFiles = new List<ZenCookedFile>();
 
+			GetZenHostNameAndPort(SC, out string SocketHostNameAndPort, out string HttpHostNameAndPort);
 			string ProjectId = SC.PackageStoreData.ZenServerStore.ProjectId;
 			string OplogId = SC.PackageStoreData.ZenServerStore.OplogId;
 
-			bool bAttemptAutoLaunchOnFailure = bAutoLaunch;
+			if (!ZenUtils.IsZenServerRunning(SocketHostNameAndPort))
+			{
+				if (!bAutoLaunch)
+				{
+					string ExceptionMessage = string.Format(
+						"Failed sending oplog request to Zen at {0}: ZenServer is not running (or is not responding), and bAutoLaunch == false.",
+						SocketHostNameAndPort);
+					throw new AutomationException(ExceptionMessage);
+				}
+
+				string Arguments = String.Format("IoStore -StartZenServerForStage -SponsorProcessID={0} -{1}={2}",
+					Environment.ProcessId, PackageStoreFileArgName, MakePathSafeToUseWithCommandLine(PackageStoreFileArgValue));
+				if (Params.Unattended)
+				{
+					Arguments += " -unattended";
+				}
+				RunUnrealPak(Params.RawProjectPath, Arguments, "ZenAutoLaunch");
+				if (!ZenUtils.IsZenServerRunning(SocketHostNameAndPort))
+				{
+					string ExceptionMessage = string.Format(
+						"Failed sending oplog request to Zen at {0}: ZenServer is not running (or is not responding), even after attempting to AutoLaunch.",
+						SocketHostNameAndPort);
+					throw new AutomationException(ExceptionMessage);
+				}
+			}
+
 			HttpResponseMessage HttpGetResult = null;
-			string HostName;
-			for (; ; )
 			{
-				HostName = GetZenHostName(SC);
-				HttpClient HttpClient = CreateZenHttpClient(SC);
-				string Uri = string.Format("http://{0}/prj/{1}/oplog/{2}/entries", HostName, ProjectId, OplogId);
+				HttpClient HttpClient = ZenUtils.CreateHttpClient(SocketHostNameAndPort);
+				HttpClient.Timeout = TimeSpan.FromMinutes(20);
+
+				string Uri = string.Format("http://{0}/prj/{1}/oplog/{2}/entries", HttpHostNameAndPort, ProjectId, OplogId);
 				Uri += "?trim_by_referencedset=true";
 
 				using var Request = new HttpRequestMessage(HttpMethod.Get, Uri);
@@ -980,37 +993,19 @@ namespace AutomationScripts
 				try
 				{
 					HttpGetResult = HttpClient.Send(Request);
-					break;
 				}
 				catch (Exception Ex)
 				{
-					if (bAttemptAutoLaunchOnFailure)
-					{
-						string Arguments = String.Format("IoStore -StartZenServerForStage -SponsorProcessID={0} -{1}={2}", Environment.ProcessId, PackageStoreFileArgName, MakePathSafeToUseWithCommandLine(PackageStoreFileArgValue));
-						if (Params.Unattended)
-						{
-							Arguments += " -unattended";
-						}
-
-						RunUnrealPak(Params.RawProjectPath, Arguments, "ZenAutoLaunch");
-						bAttemptAutoLaunchOnFailure = false;
-					}
-					else
-					{
-						string RetryMessage = bAutoLaunch
-							? "Failed even after attempting restart."
-							: "bAutoLaunch is false so did not attempt restart. Ensure that the server is running.";
-						string ExceptionMessage = String.Format("Failed sending oplog request to Zen at {0} for oplog {1}.{2}: {3}. {4}",
-								HostName, ProjectId, OplogId, Ex.Message, RetryMessage);
-						throw new AutomationException(ExceptionMessage);
-					}
+					string ExceptionMessage = String.Format("Failed sending oplog request to Zen at {0} for oplog {1}.{2}: {3}.",
+							SocketHostNameAndPort, ProjectId, OplogId, Ex.Message);
+					throw new AutomationException(ExceptionMessage);
 				}
 			}
 			if (!HttpGetResult.IsSuccessStatusCode)
 			{
 				string ExceptionMessage = String.Format("Failed reading oplog from Zen at {0} for oplog {1}.{2}, "
 					+ "HttpClient.Send returned StatusCode {3}. Ensure that cooking was successful.",
-					HostName, ProjectId, OplogId, HttpGetResult.StatusCode);
+					SocketHostNameAndPort, ProjectId, OplogId, HttpGetResult.StatusCode);
 				throw new AutomationException(ExceptionMessage);
 			}
 			Task<byte[]> ReadOplogTask = HttpGetResult.Content.ReadAsByteArrayAsync();
@@ -1316,7 +1311,9 @@ namespace AutomationScripts
 
 		private static async Task<byte[]> ReadChunkId(HttpClient Client, ZenServerStoreData StoreData, IoHash ChunkId)
 		{
-			using var Request = new HttpRequestMessage(HttpMethod.Get, string.Format("http://{0}:{1}/prj/{2}/oplog/{3}/{4}", StoreData.HostName, StoreData.HostPort, StoreData.ProjectId, StoreData.OplogId, ChunkId));
+			GetZenHostNameAndPort(StoreData, out string SocketHostNameAndPort, out string HttpHostNameAndPort);
+			using var Request = new HttpRequestMessage(HttpMethod.Get, string.Format("http://{0}/prj/{2}/oplog/{3}/{4}",
+				HttpHostNameAndPort, StoreData.ProjectId, StoreData.OplogId, ChunkId));
 			Request.Headers.Add("Accept", "application/octet-stream");
 			return await Client.Send(Request).Content.ReadAsByteArrayAsync();
 		}
@@ -1386,7 +1383,8 @@ namespace AutomationScripts
 				{
 					CookedFiles.AddRange(DirectoryReference.EnumerateFiles(SC.PlatformCookDir, "*.projectstore", SearchOption.TopDirectoryOnly).ToList());
 				}
-				HttpClient Client = CreateZenHttpClient(SC);
+				GetZenHostNameAndPort(SC, out string SocketHostNameAndPort, out string HttpHostNameAndPort);
+				HttpClient Client = ZenUtils.CreateHttpClient(SocketHostNameAndPort);
 				List<Task> TaskList = new List<Task>();
 				foreach (ZenCookedFile CookedFile in SC.PackageStoreData.ZenCookedFiles)
 				{
@@ -4584,7 +4582,7 @@ namespace AutomationScripts
 				StringBuilder UploadArgs = new StringBuilder();
 				if (Params.Upload.ToLower() == "localzen")
 				{
-					int Port = 8558;
+					int Port = ZenRunContext.DefaultPort;
 					string ServiceUrl = string.Format("http://127.0.0.1:{0}", Port);
 					List<string> Addresses = GetHostAddresses(SC);
 					if (Addresses.Count > 0)
@@ -5113,8 +5111,10 @@ namespace AutomationScripts
 				{
 					HttpResponseMessage HttpGetResult = null;
 
-					HttpClient HttpClient = CreateZenHttpClient(SC);
-					using var Request = new HttpRequestMessage(HttpMethod.Get, string.Format("http://{0}:{1}/prj/{2}/oplog/{3}/{4}", SC.PackageStoreData.ZenServerStore.HostName, SC.PackageStoreData.ZenServerStore.HostPort, SC.PackageStoreData.ZenServerStore.ProjectId, SC.PackageStoreData.ZenServerStore.OplogId, ZenFile.ChunkId.ToString()));
+					GetZenHostNameAndPort(SC, out string SocketHostNameAndPort, out string HttpHostNameAndPort);
+					HttpClient HttpClient = ZenUtils.CreateHttpClient(SocketHostNameAndPort);
+					using var Request = new HttpRequestMessage(HttpMethod.Get, string.Format("http://{0}/prj/{2}/oplog/{3}/{4}",
+						HttpHostNameAndPort, SC.PackageStoreData.ZenServerStore.ProjectId, SC.PackageStoreData.ZenServerStore.OplogId, ZenFile.ChunkId.ToString()));
 					Request.Headers.Add("Accept", "application/octet-stream");
 					try
 					{
@@ -5122,12 +5122,14 @@ namespace AutomationScripts
 					}
 					catch
 					{
-						throw new AutomationException(String.Format("Failed sending chunk request to Zen at {0}:{1}. Ensure that the server is running.", SC.PackageStoreData.ZenServerStore.HostName, SC.PackageStoreData.ZenServerStore.HostPort));
+						throw new AutomationException(String.Format("Failed sending chunk request to Zen at {0}. Ensure that the server is running.",
+							SocketHostNameAndPort));
 					}
 
 					if (!HttpGetResult.IsSuccessStatusCode)
 					{
-						throw new AutomationException(String.Format("Failed reading chunk {0}.{1} from Zen. Ensure that cooking was successful.", SC.PackageStoreData.ZenServerStore.ProjectId, SC.PackageStoreData.ZenServerStore.OplogId));
+						throw new AutomationException(String.Format("Failed reading chunk {0}.{1} from Zen. Ensure that cooking was successful.",
+							SC.PackageStoreData.ZenServerStore.ProjectId, SC.PackageStoreData.ZenServerStore.OplogId));
 					}
 					Task<byte[]> ReadChunkTask = HttpGetResult.Content.ReadAsByteArrayAsync();
 					ReadChunkTask.Wait();
