Staging: ZenServer: Bugfix from static analysis for invalid format string in previous change.
#rnx

diff --git a/Engine/Source/Programs/AutomationTool/Scripts/CopyBuildToStagingDirectory.Automation.cs b/Engine/Source/Programs/AutomationTool/Scripts/CopyBuildToStagingDirectory.Automation.cs
--- a/Engine/Source/Programs/AutomationTool/Scripts/CopyBuildToStagingDirectory.Automation.cs
+++ b/Engine/Source/Programs/AutomationTool/Scripts/CopyBuildToStagingDirectory.Automation.cs
@@ -1312,7 +1312,7 @@ namespace AutomationScripts
 		private static async Task<byte[]> ReadChunkId(HttpClient Client, ZenServerStoreData StoreData, IoHash ChunkId)
 		{
 			GetZenHostNameAndPort(StoreData, out string SocketHostNameAndPort, out string HttpHostNameAndPort);
-			using var Request = new HttpRequestMessage(HttpMethod.Get, string.Format("http://{0}/prj/{2}/oplog/{3}/{4}",
+			using var Request = new HttpRequestMessage(HttpMethod.Get, string.Format("http://{0}/prj/{1}/oplog/{2}/{3}",
 				HttpHostNameAndPort, StoreData.ProjectId, StoreData.OplogId, ChunkId));
 			Request.Headers.Add("Accept", "application/octet-stream");
 			return await Client.Send(Request).Content.ReadAsByteArrayAsync();
@@ -5113,7 +5113,7 @@ namespace AutomationScripts
 
 					GetZenHostNameAndPort(SC, out string SocketHostNameAndPort, out string HttpHostNameAndPort);
 					HttpClient HttpClient = ZenUtils.CreateHttpClient(SocketHostNameAndPort);
-					using var Request = new HttpRequestMessage(HttpMethod.Get, string.Format("http://{0}/prj/{2}/oplog/{3}/{4}",
+					using var Request = new HttpRequestMessage(HttpMethod.Get, string.Format("http://{0}/prj/{1}/oplog/{2}/{3}",
 						HttpHostNameAndPort, SC.PackageStoreData.ZenServerStore.ProjectId, SC.PackageStoreData.ZenServerStore.OplogId, ZenFile.ChunkId.ToString()));
 					Request.Headers.Add("Accept", "application/octet-stream");
 					try
