Allow Localization Chunkd Data Generator to be enabled during DLC cooks that are a bCookingDLCBaseGamePlugin.
Ensure locmeta files generated by the cook are captured during CopyBuildToStaging.

diff --git a/Engine/Source/Editor/UnrealEd/Private/CookOnTheFlyServer.cpp b/Engine/Source/Editor/UnrealEd/Private/CookOnTheFlyServer.cpp
--- a/Engine/Source/Editor/UnrealEd/Private/CookOnTheFlyServer.cpp
+++ b/Engine/Source/Editor/UnrealEd/Private/CookOnTheFlyServer.cpp
@@ -12508,7 +12508,7 @@ void UCookOnTheFlyServer::RegisterLocalizationChunkDataGenerator()
 
 	// Localization chunking is disabled when cooking DLC as it produces output that can override the base localization data
 	// Localization chunking is disabled when we're not cooking for any languages, as there would be no output generated
-	if (IsCookingDLC() || CookByTheBookOptions->AllCulturesToCook.IsEmpty())
+	if ((GetCookingDLC() == UE::Cook::ECookingDLC::Yes) || CookByTheBookOptions->AllCulturesToCook.IsEmpty())
 	{
 		return;
 	}
diff --git a/Engine/Source/Programs/AutomationTool/Scripts/CopyBuildToStagingDirectory.Automation.cs b/Engine/Source/Programs/AutomationTool/Scripts/CopyBuildToStagingDirectory.Automation.cs
--- a/Engine/Source/Programs/AutomationTool/Scripts/CopyBuildToStagingDirectory.Automation.cs
+++ b/Engine/Source/Programs/AutomationTool/Scripts/CopyBuildToStagingDirectory.Automation.cs
@@ -1029,9 +1029,13 @@ namespace AutomationScripts
 					foreach (CbField PackageDataField in EntryField["files"].AsArray())
 					{
 						string RelativeFilename = PackageDataField["clientpath"].AsString();
-						if (RelativeFilename.EndsWith(".bin") || RelativeFilename.EndsWith(".ushaderbytecode") ||
-							RelativeFilename.EndsWith(".metallib") || RelativeFilename.EndsWith(".metalmap") ||
-							(RelativeFilename.EndsWith(".txt") && RelativeFilename.Contains("/Metadata/ChunkManifest/"))
+						if (RelativeFilename.EndsWith(".bin")
+							|| RelativeFilename.EndsWith(".ushaderbytecode")
+							|| RelativeFilename.EndsWith(".metallib")
+							|| RelativeFilename.EndsWith(".metalmap")
+							|| (RelativeFilename.EndsWith(".txt") && RelativeFilename.Contains("/Metadata/ChunkManifest/"))
+							|| (RelativeFilename.EndsWith(".locmeta") && RelativeFilename.Contains("{project}/Content/Localization/") && RelativeFilename.Contains("locchunk"))
+							|| (RelativeFilename.EndsWith(".locres") && RelativeFilename.Contains("{project}/Content/Localization/") && RelativeFilename.Contains("locchunk"))
 							)
 						{
 							RelativeFilename = RelativeFilename.Replace("/{engine}/", "Engine/");
