Fix gathering short tooltips' text for auto-generated localization configs.

#tests Run Localize command locally for any plugin with "ToolTips" specified as LoadingPolicy - GatherTextFromMetadata commandlet works correctly without crashing on iterating UObjectShortToolTips objects

diff --git a/Engine/Source/Programs/AutomationTool/Localization/LocalizationConfigFileParams.cs b/Engine/Source/Programs/AutomationTool/Localization/LocalizationConfigFileParams.cs
--- a/Engine/Source/Programs/AutomationTool/Localization/LocalizationConfigFileParams.cs
+++ b/Engine/Source/Programs/AutomationTool/Localization/LocalizationConfigFileParams.cs
@@ -115,26 +115,27 @@ namespace EpicGames.Localization
 		public List<string> OutputNamespaces { get; set; } = new List<string>();
 		public bool bShouldGatherFromEditorOnlyData { get; set; } = false;
 
-		public static GatherTextFromMetaDataParams CreateDefault(string localizationTargetRootDirectory, LocalizationTargetDescriptorLoadingPolicy loadingPolicy)
+		public static GatherTextFromMetaDataParams CreateDefault(string localizationTargetRootDirectory, LocalizationTargetDescriptorLoadingPolicy localizationPolicy)
 		{
 			GatherTextFromMetaDataParams gatherTextFromMetaData = new GatherTextFromMetaDataParams();
 
 			gatherTextFromMetaData.IncludePathFilters.Add(Path.Combine(localizationTargetRootDirectory, "Source",  "*"));
 
-			if (loadingPolicy == LocalizationTargetDescriptorLoadingPolicy.PropertyNames)
+			if (localizationPolicy == LocalizationTargetDescriptorLoadingPolicy.PropertyNames)
 			{
 				gatherTextFromMetaData.InputKeys.Add("DisplayName");
 				gatherTextFromMetaData.OutputKeys.Add("\"{FieldPath}\"");
 				gatherTextFromMetaData.OutputNamespaces.Add("UObjectDisplayNames");
 			}
 
-			if (loadingPolicy == LocalizationTargetDescriptorLoadingPolicy.ToolTips)
+			if (localizationPolicy == LocalizationTargetDescriptorLoadingPolicy.ToolTips)
 			{
 				gatherTextFromMetaData.InputKeys.Add("ToolTip");
 				gatherTextFromMetaData.InputKeys.Add("ShortTooltip");
 				gatherTextFromMetaData.OutputKeys.Add("\"{FieldPath}\"");
 				gatherTextFromMetaData.OutputKeys.Add("\"{FieldPath}\"");
 				gatherTextFromMetaData.OutputNamespaces.Add("UObjectToolTips");
+				gatherTextFromMetaData.OutputNamespaces.Add("UObjectShortToolTips");
 			}
 
 			gatherTextFromMetaData.bShouldGatherFromEditorOnlyData = true;
