[Gauntlet] Fixes issue with UnrealTargetContraints hashcode resolving different values if the tag list was empty or null

diff --git a/Engine/Source/Programs/AutomationTool/AutomationUtils/DeviceReservation.cs b/Engine/Source/Programs/AutomationTool/AutomationUtils/DeviceReservation.cs
--- a/Engine/Source/Programs/AutomationTool/AutomationUtils/DeviceReservation.cs
+++ b/Engine/Source/Programs/AutomationTool/AutomationUtils/DeviceReservation.cs
@@ -357,10 +357,10 @@ namespace AutomationTool.DeviceReservation
 
 			public override string ToString()
 			{
-				return $"Required: {(RequiredTags == null ? string.Empty : string.Join(", ", RequiredTags))}\n" +
-				       $"Preferred: {(PreferredTags == null ? string.Empty : string.Join(", ", PreferredTags))}\n" +
-				       $"Undesired: {(UndesiredTags == null ? string.Empty : string.Join(", ", UndesiredTags))}\n" +
-				       $"Blocked: {(BlockedTags == null ? string.Empty : string.Join(", ", BlockedTags))}";
+				return $"\n\tRequired: {(RequiredTags == null || !RequiredTags.Any() ? string.Empty : string.Join(", ", RequiredTags))}" +
+				       $"\n\tPreferred: {(PreferredTags == null || !PreferredTags.Any() ? string.Empty : string.Join(", ", PreferredTags))}" +
+				       $"\n\tUndesired: {(UndesiredTags == null || !UndesiredTags.Any() ? string.Empty : string.Join(", ", UndesiredTags))}" +
+				       $"\n\tBlocked: {(BlockedTags == null || !BlockedTags.Any() ? string.Empty : string.Join(", ", BlockedTags))}";
 			}
 
 			public override bool Equals(object OtherObj)
diff --git a/Engine/Source/Programs/AutomationTool/Gauntlet/Framework/Devices/Gauntlet.DevicePool.cs b/Engine/Source/Programs/AutomationTool/Gauntlet/Framework/Devices/Gauntlet.DevicePool.cs
--- a/Engine/Source/Programs/AutomationTool/Gauntlet/Framework/Devices/Gauntlet.DevicePool.cs
+++ b/Engine/Source/Programs/AutomationTool/Gauntlet/Framework/Devices/Gauntlet.DevicePool.cs
@@ -147,11 +147,15 @@ namespace Gauntlet
 				return true;
 			}
 
+			bool bNullTags = Other.Tags == null && Tags == null;
+			bool bMatchingTags = Other.Tags != null && Other.Tags.Equals(Tags);
+			bool bNullEmptyTags = (Other.Tags == null && Tags != null && !Tags.Any()) || (Other.Tags != null && !Other.Tags.Any() && Tags == null);
+
 			return Other.Platform == Platform
 				&& Other.Model.Equals(Model, StringComparison.InvariantCultureIgnoreCase)
 				&& Other.DeviceName.Equals(DeviceName, StringComparison.InvariantCultureIgnoreCase)
 				&& Other.PerfSpec == PerfSpec
-				&& ((Other.Tags == null && Tags == null) || (Other.Tags != null && Other.Tags.Equals(Tags)));
+				&& (bNullTags || bMatchingTags || bNullEmptyTags);
 		}
 
 		public override bool Equals(object Obj)
@@ -204,7 +208,7 @@ namespace Gauntlet
 			{
 				Value = string.Format("{0}:{1}", Value, DeviceName);
 			}
-			if (Tags != null)
+			if (Tags != null && Tags.Any())
 			{
 				Value = string.Format("{0}:{1}", Value, Tags.ToString());
 			}
