[Gauntlet] Adds appium integration for use by TargetDeviceIOS
Users can create a json config file and set the UE-AppiumConfigPath envvar or run gauntlet with -AppiumConfigPath=/configpath.
If this file is found, an instance of AppiumContainer will be constructed. Prior to a test launchingthe container will automatically start a server and deploy a web driver agent runner.
The primary use of this is automatically dismissing blocking system prompts.

diff --git a/Engine/Source/Programs/AutomationTool/Gauntlet/Gauntlet.Automation.csproj b/Engine/Source/Programs/AutomationTool/Gauntlet/Gauntlet.Automation.csproj
--- a/Engine/Source/Programs/AutomationTool/Gauntlet/Gauntlet.Automation.csproj
+++ b/Engine/Source/Programs/AutomationTool/Gauntlet/Gauntlet.Automation.csproj
@@ -90,6 +90,7 @@
     <ProjectReference Include="..\Android\Android.Automation.csproj" PrivateAssets="All"><Private>false</Private></ProjectReference>
   </ItemGroup>
   <ItemGroup>
+    <PackageReference Include="Appium.WebDriver" Version="8.0.0" />
     <PackageReference Include="Magick.NET-Q16-HDRI-AnyCPU" Version="14.8.2" PrivateAssets="all" />
     <PackageReference Include="Microsoft.CSharp" Version="4.7.0" />
     <PackageReference Include="MySql.Data" Version="6.10.9" PrivateAssets="all" />
diff --git a/Engine/Source/Programs/AutomationTool/Gauntlet/Platform/IOS/Gauntlet.IOSAppium.cs b/Engine/Source/Programs/AutomationTool/Gauntlet/Platform/IOS/Gauntlet.IOSAppium.cs
new file mode 100644
--- /dev/null
+++ b/Engine/Source/Programs/AutomationTool/Gauntlet/Platform/IOS/Gauntlet.IOSAppium.cs
@@ -0,0 +1,304 @@
+// Copyright Epic Games, Inc. All Rights Reserved.
+using System;
+using System.Text.RegularExpressions;
+using AutomationTool;
+using Newtonsoft.Json;
+using OpenQA.Selenium.Appium;
+using OpenQA.Selenium.Appium.iOS;
+using OpenQA.Selenium.Appium.Service;
+using static AutomationTool.CommandUtils;
+
+namespace Gauntlet
+{
+	/// <summary>
+	/// AppiumContainer is a wrapper for both an Appium server instance and an AppiumDriver.
+	/// It's primary use in Gauntlet is for automating dismissal of blocking system notifications that cannot be managed by an MDM profile.
+	/// Provided you configure your environment correctly, the container will automatically initialize as part of an IOS app instance execution.
+	///
+	/// Steps for configuring your environment:
+	///		1.	Download appium on your host. It's recommended you use the global npm installation 'npm install -g appium'.
+	///		2.  For real device testing, you will need a WebDriverAgentRunner.app which is signed with a mobile provision that includes your device.
+	///			For this you have two options:
+	///				- Build the app from source. This lets you configure bundle.id's if your signing cert only allows for certain identifiers.
+	///				- Download the precompiled app and re-sign after replacing the embedded mobileprovision.
+	///			In either case, you can find both the source and precompiled app on this page [REDACTED_PROFILE_URL]/WebDriverAgent/releases
+	///		3.	Create a JSON file that can be de-serialized to the 'AppiumContainer.Config' type. This file is used to configure the driver with information
+	///			that is specific to your team. Place this file in a location that can be read by your host.
+	///		4.	Point the container to the location of the json file you created in step 4 by doing one of the following:
+	///				- Setting the UE-AppiumConfigPath EnvVar to a qualified path or a relative path to your UE root.
+	///				- Run UAT with -AppiumConfigPath=/path set to a qualified path or a relative path to your UE root.
+	///
+	/// Once all these steps are completed, before TargetDeviceIOS.Run starts the app process, it will execute these actions:
+	///		1. Start an appium server on an available loopback port
+	///		2. Install the WebDriverAgent app
+	///		3. Start the driver with your configured settings
+	///
+	/// From there appium will automatically accept/dismiss any system prompts it encounters.
+	/// </summary>
+	public class AppiumContainer : IDisposable
+    {
+        class Config
+        {
+            /// <summary>
+            /// Name of the automation driver to use. Defaults to XCUITest if not specified
+            /// </summary>
+            public string AutomationName { get; set; } = "xcuitest";
+
+            /// <summary>
+            /// Your company's apple developer team ID
+            /// </summary>
+            public string OrgId { get; set; }
+
+            /// <summary>
+            /// Identity of your signing cert. Usually just 'Apple Development'
+            /// </summary>
+            public string SigningId { get; set; }
+
+			/// <summary>
+			/// Path to a precompiled WebDriverAgentRunner app
+			/// </summary>
+			public string WdaAppPath { get; set; }
+
+			/// <summary>
+			/// Bundle ID of the WebDriverAgent app. Ex: 'com.epicgames.WebDriverAgent'
+			/// </summary>
+			public string WdaBundleId { get; set; }
+
+			/// <summary>
+			/// Optional - Allow you to override the location of the appium executable.
+			/// Useful if you opt not to install appium to the global npm root
+			/// </summary>
+			public string AppiumLocation { get; set; }
+
+            public AppiumOptions GetCapabilities(string UUID, string PackageName)
+            {
+                AppiumOptions Capabilities = new AppiumOptions();
+                Capabilities.PlatformName = "iOS";
+				Capabilities.AutomationName = AutomationName;
+                Capabilities.AddAdditionalAppiumOption("udid", UUID);
+                Capabilities.AddAdditionalAppiumOption("bundleId", PackageName);
+                Capabilities.AddAdditionalAppiumOption("autoAcceptAlerts", true);
+                Capabilities.AddAdditionalAppiumOption("xcodeOrgId", OrgId);
+                Capabilities.AddAdditionalAppiumOption("xcodeSigningId", SigningId);
+                Capabilities.AddAdditionalAppiumOption("autoLaunch", false);
+                Capabilities.AddAdditionalAppiumOption("usePreinstalledWDA", true);
+                Capabilities.AddAdditionalAppiumOption("updatedWDABundleId", WdaBundleId);
+
+                if (Log.Level == LogLevel.Verbose)
+                {
+                    Capabilities.AddAdditionalAppiumOption("showXcodeLog", true);
+                }
+
+				return Capabilities;
+            }
+        }
+
+        private const string AppiumConfigEnvVar = "UE-AppiumConfigPath";
+
+        private const string AppiumConfigArg = "AppiumConfigPath";
+
+        /// <summary>
+        /// **REQUIRED**
+        /// Path to a json file that can be deserialized into the AppiumContainer.Config type
+        /// This allows users to specify things such as the org id, signing identity, etc.
+        /// Can be overriden by setting the UE-AppiumConfigPath envvar, or by running with -AppiumConfigPath=/path
+        /// Path can be relative or absolute
+        /// </summary>
+        public static string AppiumConfigPath = "Engine/Restricted/NotForLicensees/Extras/ThirdPartyNotUE/WebDriverAgent/AppiumConfig.json";
+
+        /// <summary>
+        /// Whether or not the appium container is configured properly for use.
+        /// This requires the appium config path and the appium server path to point to existing files
+        /// </summary>
+        public static bool Enabled => AppiumConfig != null;
+
+		/// <summary>
+		/// Lock object
+		/// </summary>
+        private static object Mutex = new();
+
+        /// <summary>
+        /// Config file used for the Driver
+        /// </summary>
+        private static Config AppiumConfig = null;
+
+		/// <summary>
+		/// Driver handle
+		/// </summary>
+        private IOSDriver Driver = null;
+
+		/// <summary>
+		/// Server handle
+		/// </summary>
+        private AppiumLocalService Server = null;
+
+		/// <summary>
+		/// UUID of the device being tests
+		/// </summary>
+        private string UUID = null;
+        static AppiumContainer()
+        {
+			if (Globals.Params.ParseParam("NoAppium"))
+			{
+				return;
+			}
+
+            string ConfigEnvVar = Environment.GetEnvironmentVariable(AppiumConfigEnvVar);
+            if (!string.IsNullOrEmpty(ConfigEnvVar))
+            {
+                AppiumConfigPath = ConfigEnvVar;
+            }
+            AppiumConfigPath = Globals.Params.ParseValue(AppiumConfigArg, AppiumConfigPath);
+
+            if (FileExists(AppiumConfigPath))
+            {
+                try
+                {
+                    AppiumConfig = JsonConvert.DeserializeObject<Config>(ReadAllText(AppiumConfigPath));
+                }
+                catch (Exception Ex)
+                {
+                    throw new AutomationException(Ex, "Failed to derserialize AppiumConfig at {0}", AppiumConfigPath);
+                }
+            }
+        }
+
+        public AppiumContainer(string UUID)
+        {
+            this.UUID = UUID;
+            ConfigureDrivers();
+        }
+
+		#region IDisposable Support
+		private bool Disposed = false;
+        ~AppiumContainer()
+        {
+            Dispose(false);
+        }
+        public void Dispose()
+        {
+            Dispose(true);
+            GC.SuppressFinalize(this);
+        }
+
+        private void Dispose(bool bDisposing)
+        {
+			if (!Disposed)
+			{
+				if (bDisposing)
+				{
+
+				}
+				Stop();
+				Disposed = true;
+			}
+        }
+        #endregion
+
+        public void Start(string PackageName)
+        {
+            Stop();
+
+            try
+            {
+				IProcessResult InstallResult;
+				using (ScopedSuspendECErrorParsing ErrorSuspension = new())
+				{
+					if (TargetDeviceIOS.UseDeviceCtl)
+					{
+						InstallResult = TargetDeviceIOS.ExecuteDevicectlCommand(
+							string.Format("device install app \"{0}\" -v", AppiumConfig.WdaAppPath), UUID, 60, AdditionalOptions: ERunOptions.NoStdOutRedirect);
+					}
+					else
+					{
+						InstallResult = TargetDeviceIOS.ExecuteIOSDeployCommand(
+							string.Format("-b \"{0}\"", AppiumConfig.WdaAppPath), UUID, 60, AdditionalOptions: ERunOptions.NoStdOutRedirect);
+					}
+				}
+
+				if (InstallResult.ExitCode != 0)
+				{
+					throw new AutomationException("Failed to install WDA app ({0}): {1}", InstallResult.ExitCode, InstallResult.Output);
+				}
+
+				lock (Mutex)
+                {
+					AppiumServiceBuilder ServerBuilder = new AppiumServiceBuilder()
+						.UsingAnyFreePort()
+						.WithIPAddress("127.0.0.1");
+
+					if (!string.IsNullOrEmpty(AppiumConfig.AppiumLocation) && FileExists(AppiumConfig.AppiumLocation))
+					{
+						ServerBuilder.WithAppiumJS(new System.IO.FileInfo(AppiumConfig.AppiumLocation));
+					}
+
+                    Server = ServerBuilder.Build();
+                    Server.Start();
+
+                    Driver = new IOSDriver(Server, AppiumConfig.GetCapabilities(UUID, PackageName));
+                }
+            }
+            catch (Exception Ex)
+            {
+                throw new AutomationException(Ex, "Failed to start appium instance for {0}: {1}", UUID, Ex.Message);
+            }
+        }
+
+        public void Stop()
+        {
+            if (Driver != null)
+            {
+                Driver.Dispose();
+                Driver = null;
+            }
+
+            if (Server != null)
+            {
+                Server.Dispose();
+                Server = null;
+            }
+
+            // TODO: Terminate WDA app - parse devicectl for pid, then kill?
+        }
+
+		private void ConfigureDrivers()
+		{
+			// Query installed drivers
+			IProcessResult Result = Run("appium", "driver list");
+			if (Result.ExitCode != 0)
+			{
+				throw new AutomationException("Failed to query appium driver list ({0}): {1}", Result.ExitCode, Result.Output);
+			}
+
+            // Trim ansii escape codes
+            string SanitizedOutput = Regex.Replace(Result.Output, "\\x1B\\[[0-9;]*[a-zA-Z]", string.Empty);
+
+			// Try to find the driver for for this container's config
+            bool bFoundDriver = false;
+			Regex DriverMatch = new Regex("(- )(.*?)(@)(\\d+((\\.\\d+)+)?)", RegexOptions.IgnoreCase | RegexOptions.Multiline);
+			foreach (Match Match in DriverMatch.Matches(SanitizedOutput))
+			{
+				string Driver = Match.Groups[2].Value;
+				string Version = Match.Groups[4].Value;
+
+				if (Driver.Equals(AppiumConfig.AutomationName, StringComparison.OrdinalIgnoreCase))
+				{
+					bFoundDriver = true;
+					Log.Verbose("Using {AppiumDriverName} appium driver version {AppiumDriverVersion}", Driver, Version);
+					break;
+				}
+			}
+
+			// Install the driver if the missing
+			if (!bFoundDriver)
+			{
+				Log.Info("Could not find appium driver {AppiumDriverName}. Attempting to install...", AppiumConfig.AutomationName);
+				Result = Run("appium", $"driver install {AppiumConfig.AutomationName.ToLower()}");
+				if (Result.ExitCode != 0)
+				{
+					throw new AutomationException("Failed to install {0} appium driver. ({1}): {2}", AppiumConfig.AutomationName, Result.ExitCode, Result.Output);
+				}
+			}
+		}
+    }
+}
\ No newline at end of file
diff --git a/Engine/Source/Programs/AutomationTool/Gauntlet/Platform/IOS/Gauntlet.TargetDeviceIOS.cs b/Engine/Source/Programs/AutomationTool/Gauntlet/Platform/IOS/Gauntlet.TargetDeviceIOS.cs
--- a/Engine/Source/Programs/AutomationTool/Gauntlet/Platform/IOS/Gauntlet.TargetDeviceIOS.cs
+++ b/Engine/Source/Programs/AutomationTool/Gauntlet/Platform/IOS/Gauntlet.TargetDeviceIOS.cs
@@ -7,7 +7,6 @@ using System.Data;
 using System.Text;
 using System.Text.RegularExpressions;
 using System.Collections.Generic;
-using System.Security.Cryptography;
 using EpicGames.Core;
 using Gauntlet.Utils;
 using AutomationTool;
@@ -72,7 +71,7 @@ namespace Gauntlet
 		/// ios-deploy will be used when this is false.
 		/// This property is set based on the installed version of xcode and the ios version of the device
 		/// </summary>
-		private static bool UseDeviceCtl;
+		public static bool UseDeviceCtl { get; private set; }
 
 		/// <summary>
 		/// Collection of devices visible to the UAT host
@@ -91,14 +90,9 @@ namespace Gauntlet
 		private static bool bStaticInitialized;
 
 		/// <summary>
-		/// Prevents separate start threads from managing connected device management at the same time (bulk only)
+		/// Lock object
 		/// </summary>
-		private static object InitializerMutex = new object();
-
-		/// <summary>
-		/// Prevents separate start threads from writing to the network build cache directory at the same time(bulk only)
-		/// </summary>
-		private static object CacheMutex = new object();
+		private static object Mutex = new object();
 
 		/// <summary>
 		/// Used to prevent re-copies of an entire network share build in the event of a device failure (bulk only)
@@ -145,8 +139,13 @@ namespace Gauntlet
 
 		private bool IsDefaultDevice = false;
 
+		/// <summary>
+		/// Handle to the device's Install
+		/// </summary>
 		private IOSAppInstall InstallCache = null;
 
+		private AppiumContainer AppiumContainer = null;
+
 		public TargetDeviceIOS(string InName, string InCachePath = null)
 		{
 			AutoParam.ApplyParamsAndDefaults(this, Globals.Params.AllArguments);
@@ -160,7 +159,7 @@ namespace Gauntlet
 
 			IsDefaultDevice = string.IsNullOrEmpty(InName) || InName.Equals("default", StringComparison.OrdinalIgnoreCase);
 			Name = IsDefaultDevice ? "default" : InName;
-			
+
 			if (IsDefaultDevice)
 			{
 				IEnumerable<ConnectedDevice> AvailableDevices = ConnectedDevices.Values.Where(Device => !Device.IsAssigned);
@@ -230,6 +229,11 @@ namespace Gauntlet
 				UseDeviceCtl = true;
 			}
 
+			if (AppiumContainer.Enabled)
+			{
+				AppiumContainer = new AppiumContainer(UUID);
+			}
+
 			LocalCachePath = InCachePath ?? Path.Combine(Globals.TempDir, "IOSDevice_" + UUID);
 		}
 
@@ -244,15 +248,16 @@ namespace Gauntlet
 		{
 			if (!DisposedValue)
 			{
-				// TODO: dispose managed state (managed objects).
+				if (AppiumContainer != null)
+				{
+					AppiumContainer.Dispose();
+				}
 				DisposedValue = true;
 			}
 		}
 
-		// This code added to correctly implement the disposable pattern.
 		public void Dispose()
 		{
-			// Do not change this code. Put cleanup code in Dispose(bool disposing) above.
 			Dispose(true);
 			GC.SuppressFinalize(this);
 		}
@@ -569,7 +574,7 @@ namespace Gauntlet
 					string DestinationBulkPath = Path.Combine(TempBuildCache, Path.GetFileName(SourceBulkPath));
 					string DestinationBulkContentsPath = Path.Combine(TempBuildCache, Path.GetFileName(SourceBulkContentsPath));
 
-					lock(CacheMutex)
+					lock (Mutex)
 					{
 						if(!HasCopiedNetworkBuild)
 						{
@@ -785,6 +790,11 @@ namespace Gauntlet
 				throw new DeviceException("Failed to copy uecommandline file to device!");
 			}
 
+			if (AppiumContainer != null)
+			{
+				AppiumContainer.Start(IOSInstall.PackageName);
+			}
+
 			Log.Info("Launching {0} on {1}", Install.Name, this);
 			Log.Verbose(IOSInstall.CommandLine);
 
@@ -941,7 +951,7 @@ namespace Gauntlet
 			return Result;
 		}
 
-		private static ILongProcessResult ExecuteDeploymentCommandNoWait(string Executable, string CommandLine, string WorkingDir = null, ERunOptions AdditionalOptions = ERunOptions.None, string LocalCache = null, bool bQuiet = false)
+		public static ILongProcessResult ExecuteDeploymentCommandNoWait(string Executable, string CommandLine, string WorkingDir = null, ERunOptions AdditionalOptions = ERunOptions.None, string LocalCache = null, bool bQuiet = false)
 		{
 			if (!AdditionalOptions.HasFlag(ERunOptions.UseShellExecute) && !File.Exists(Executable))
 			{
@@ -973,7 +983,7 @@ namespace Gauntlet
 
 		private static void InitializeDevices()
 		{
-			lock (InitializerMutex)
+			lock (Mutex)
 			{
 				if (bStaticInitialized)
 				{
@@ -1053,7 +1063,7 @@ namespace Gauntlet
 				foreach (JsonObject JsonDevice in JsonDevices)
 				{
 					ConnectedDevice IOSDevice = new ConnectedDevice(string.Empty, -1, false);
-					
+
 					if (JsonDevice.TryGetObjectField("hardwareProperties", out JsonObject HardwareProperties))
 					{
 						HardwareProperties.TryGetStringField("udid", out IOSDevice.UUID);
@@ -1134,7 +1144,7 @@ namespace Gauntlet
 
 			return bIsConnected;
 		}
-		
+
 		// DEPRECATED
 		public IAppInstall InstallApplication(UnrealAppConfig AppConfig)
 		{
