Fixed destination files potentially not having the correct permission during SafeCopyFile on non-Windows

diff --git a/Engine/Source/Programs/AutomationTool/AutomationUtils/Utils.cs b/Engine/Source/Programs/AutomationTool/AutomationUtils/Utils.cs
--- a/Engine/Source/Programs/AutomationTool/AutomationUtils/Utils.cs
+++ b/Engine/Source/Programs/AutomationTool/AutomationUtils/Utils.cs
@@ -490,9 +490,14 @@ namespace AutomationTool
 					}
 					else
 					{
-						if(File.Exists(SourceName))
+						if (File.Exists(SourceName))
 						{
 							File.Copy(SourceName, TargetName, overwrite: true);
+							if (!OperatingSystem.IsWindows())
+							{
+								// Permissions may not be the same as the source file if the umask filters it, and/or if the file is overwritten
+								File.SetUnixFileMode(TargetName, File.GetUnixFileMode(SourceName));
+							}
 						}
 						else
 						{
