binary config files honor CustomConfig

diff --git a/Engine/Source/Programs/AutomationTool/Scripts/CopyBuildToStagingDirectory.Automation.cs b/Engine/Source/Programs/AutomationTool/Scripts/CopyBuildToStagingDirectory.Automation.cs
--- a/Engine/Source/Programs/AutomationTool/Scripts/CopyBuildToStagingDirectory.Automation.cs
+++ b/Engine/Source/Programs/AutomationTool/Scripts/CopyBuildToStagingDirectory.Automation.cs
@@ -2130,6 +2130,10 @@ namespace AutomationScripts
 						String TargetPlatformName = ThisPlatform.GetCookPlatform(Params.DedicatedServer, Params.Client);
 						FileReference OutputFile = FileReference.Combine(SC.ProjectRoot, "Intermediate", "Config", TargetPlatformName, "BinaryConfig.ini");
 						String UnrealPakParams = String.Format("MakeBinaryConfig -Project=\"{0}\" -Platform={1} -OutputFile=\"{2}\" -StagedPluginsFile=\"{3}\"", SC.RawProjectPath, ConfigHierarchy.GetIniPlatformName(ThisPlatform.IniPlatformType), OutputFile.FullName, ConfigCacheInputFile.FullName);
+						if (!string.IsNullOrEmpty(SC.CustomConfig))
+						{
+							UnrealPakParams += $" -CustomConfig={SC.CustomConfig}";
+						}
 
 						RunUnrealPak(Params.RawProjectPath, UnrealPakParams, "MakeBinaryConfig");
 
