1kb chunk size for windows.

#rnx

diff --git a/Engine/Source/Programs/AutomationTool/AutomationUtils/MCPPublic.cs b/Engine/Source/Programs/AutomationTool/AutomationUtils/MCPPublic.cs
--- a/Engine/Source/Programs/AutomationTool/AutomationUtils/MCPPublic.cs
+++ b/Engine/Source/Programs/AutomationTool/AutomationUtils/MCPPublic.cs
@@ -1049,6 +1049,8 @@ namespace EpicGames.MCP.Automation
 			/// The directory which will receive the generated manifest and chunks.
 			/// </summary>
 			public string CloudDir;
+			// Byte window size for patch matching. Match spans smaller than this will not be found. Min 1kb, defaults to 8kb.
+			public int ScanWindowSize=0;
 		}
 
 		public class ListBinariesOptions
