FinalizeInstalledBuild: Correctly handle one target per architecture build receipts

#rnx

diff --git a/Engine/Source/Programs/AutomationTool/Scripts/FinalizeInstalledBuild.cs b/Engine/Source/Programs/AutomationTool/Scripts/FinalizeInstalledBuild.cs
--- a/Engine/Source/Programs/AutomationTool/Scripts/FinalizeInstalledBuild.cs
+++ b/Engine/Source/Programs/AutomationTool/Scripts/FinalizeInstalledBuild.cs
@@ -129,7 +129,7 @@ namespace AutomationTool
 						{
 							foreach (UnrealArch Arch in AllArchitectures.Architectures)
 							{
-								FileReference ReceiptFileName = TargetReceipt.GetDefaultPath(new DirectoryReference(OutputEnginePath), CurrentTargetName, CodeTargetPlatform, EngineConfiguration, AllArchitectures);
+								FileReference ReceiptFileName = TargetReceipt.GetDefaultPath(new DirectoryReference(OutputEnginePath), CurrentTargetName, CodeTargetPlatform, EngineConfiguration, new(Arch.ToString(), CodeTargetPlatform));
 								if (FileReference.Exists(ReceiptFileName))
 								{
 									string RelativeReceiptFileName = ReceiptFileName.MakeRelativeTo(new DirectoryReference(OutputDir));
