Update Microsoft.Build nuget package

#rnx

diff --git a/Engine/Source/Programs/AutomationTool/AutomationTool.csproj b/Engine/Source/Programs/AutomationTool/AutomationTool.csproj
--- a/Engine/Source/Programs/AutomationTool/AutomationTool.csproj
+++ b/Engine/Source/Programs/AutomationTool/AutomationTool.csproj
@@ -90,7 +90,7 @@
   <ItemGroup>
     <PackageReference Include="AWSSDK.SecurityToken" Version="3.7.501.14" />
     <PackageReference Include="Magick.NET-Q16-HDRI-AnyCPU" Version="14.8.2" />
-    <PackageReference Include="Microsoft.Build" Version="17.11.31" ExcludeAssets="runtime" PrivateAssets="all">
+    <PackageReference Include="Microsoft.Build" Version="17.11.48" ExcludeAssets="runtime" PrivateAssets="all">
       <IncludeAssets>compile; build; native; contentfiles; analyzers; buildtransitive</IncludeAssets>
     </PackageReference>
     <PackageReference Include="Microsoft.Build.Locator" Version="1.9.1" />
diff --git a/Engine/Source/Programs/Shared/EpicGames.MsBuild/EpicGames.MsBuild.csproj b/Engine/Source/Programs/Shared/EpicGames.MsBuild/EpicGames.MsBuild.csproj
--- a/Engine/Source/Programs/Shared/EpicGames.MsBuild/EpicGames.MsBuild.csproj
+++ b/Engine/Source/Programs/Shared/EpicGames.MsBuild/EpicGames.MsBuild.csproj
@@ -36,7 +36,7 @@
   </ItemGroup>
 
   <ItemGroup>
-    <PackageReference Include="Microsoft.Build" Version="17.11.31" ExcludeAssets="runtime" />
+    <PackageReference Include="Microsoft.Build" Version="17.11.48" ExcludeAssets="runtime" />
     <PackageReference Include="Microsoft.Build.Locator" Version="1.9.1" />
     <PackageReference Include="System.Drawing.Common" Version="8.0.20" />
   </ItemGroup>
