Unshelved from pending changelist '47519487':

Update vulnerable .NET package to the lates version

#rnx

diff --git a/Engine/Source/Programs/AutomationTool/AutomationTool.csproj b/Engine/Source/Programs/AutomationTool/AutomationTool.csproj
--- a/Engine/Source/Programs/AutomationTool/AutomationTool.csproj
+++ b/Engine/Source/Programs/AutomationTool/AutomationTool.csproj
@@ -89,7 +89,7 @@
 
   <ItemGroup>
     <PackageReference Include="AWSSDK.SecurityToken" Version="3.7.501.14" />
-    <PackageReference Include="Magick.NET-Q16-HDRI-AnyCPU" Version="14.8.2" />
+    <PackageReference Include="Magick.NET-Q16-HDRI-AnyCPU" Version="14.9.1" />
     <PackageReference Include="Microsoft.Build" Version="17.11.48" ExcludeAssets="runtime" PrivateAssets="all">
       <IncludeAssets>compile; build; native; contentfiles; analyzers; buildtransitive</IncludeAssets>
     </PackageReference>
diff --git a/Engine/Source/Programs/AutomationTool/AutomationUtils/AutomationUtils.Automation.csproj b/Engine/Source/Programs/AutomationTool/AutomationUtils/AutomationUtils.Automation.csproj
--- a/Engine/Source/Programs/AutomationTool/AutomationUtils/AutomationUtils.Automation.csproj
+++ b/Engine/Source/Programs/AutomationTool/AutomationUtils/AutomationUtils.Automation.csproj
@@ -64,7 +64,7 @@
     <ProjectReference Include="..\..\UnrealBuildTool\UnrealBuildTool.csproj" PrivateAssets="All"><Private>true</Private></ProjectReference>
   </ItemGroup>
   <ItemGroup>
-    <PackageReference Include="Magick.NET-Q16-HDRI-AnyCPU" Version="14.8.2" PrivateAssets="all" />
+    <PackageReference Include="Magick.NET-Q16-HDRI-AnyCPU" Version="14.9.1" PrivateAssets="all" />
     <PackageReference Include="System.Drawing.Common" Version="8.0.20" />
     <PackageReference Include="System.Net.Http" Version="4.3.4" />
 
diff --git a/Engine/Source/Programs/AutomationTool/Gauntlet/Gauntlet.Automation.csproj b/Engine/Source/Programs/AutomationTool/Gauntlet/Gauntlet.Automation.csproj
--- a/Engine/Source/Programs/AutomationTool/Gauntlet/Gauntlet.Automation.csproj
+++ b/Engine/Source/Programs/AutomationTool/Gauntlet/Gauntlet.Automation.csproj
@@ -91,7 +91,7 @@
   </ItemGroup>
   <ItemGroup>
     <PackageReference Include="Appium.WebDriver" Version="8.0.0" />
-    <PackageReference Include="Magick.NET-Q16-HDRI-AnyCPU" Version="14.8.2" PrivateAssets="all" />
+    <PackageReference Include="Magick.NET-Q16-HDRI-AnyCPU" Version="14.9.1" PrivateAssets="all" />
     <PackageReference Include="Microsoft.CSharp" Version="4.7.0" />
     <PackageReference Include="MySql.Data" Version="6.10.9" PrivateAssets="all" />
   </ItemGroup>
