Fix adb push of stripped.so. Adb push is interpreting the destination stripped.so path as a directory and failing. Instead, adb push stripped.so to destination directory. Additionally, create destination folder for libUnreal.so prior to copy from /data/local/temp and shell quote various arguments to support paths with spaces.

diff --git a/Engine/Source/Programs/AutomationTool/Android/AndroidPlatform.Automation.cs b/Engine/Source/Programs/AutomationTool/Android/AndroidPlatform.Automation.cs
--- a/Engine/Source/Programs/AutomationTool/Android/AndroidPlatform.Automation.cs
+++ b/Engine/Source/Programs/AutomationTool/Android/AndroidPlatform.Automation.cs
@@ -3378,10 +3378,12 @@ public class AndroidPlatform : Platform
 				if (!long.TryParse(RunAdbCommand(Params, DeviceName, "shell cat ./files/libUnreal.so.timestamp", null, ERunOptions.AppMustExist).Output, out long OldTimestamp) || OldTimestamp != NewTimestamp)
 				{
 					int SuccessCode = 0;
-					string PushSO = $"push -z lz4 {FinalSOPathStripped} /data/local/tmp/{FinalSOFileNameStripped}";
-					string CopySO = $"shell run-as {PackageName} cp /data/local/tmp/{FinalSOFileNameStripped} ./files/libUnreal.so";
-					string DeleteSO = $"shell rm /data/local/tmp/{FinalSOFileNameStripped}";
+					string PushSO = $"push -z lz4 \"{FinalSOPathStripped}\" /data/local/tmp/";
+					string MkdirDestDirSO = $"shell run-as {PackageName} mkdir -p ./files";
+					string CopySO = $"shell run-as {PackageName} cp \"/data/local/tmp/{FinalSOFileNameStripped}\" ./files/libUnreal.so";
+					string DeleteSO = $"shell rm \"/data/local/tmp/{FinalSOFileNameStripped}\"";
 					RunAndLogAdbCommand(Params, DeviceName, PushSO, out SuccessCode);
+					RunAndLogAdbCommand(Params, DeviceName, MkdirDestDirSO, out SuccessCode);
 					RunAndLogAdbCommand(Params, DeviceName, CopySO, out SuccessCode);
 					RunAndLogAdbCommand(Params, DeviceName, DeleteSO, out SuccessCode);
 
diff --git a/Engine/Source/Programs/UnrealBuildTool/Platform/Android/UEDeployAndroid.cs b/Engine/Source/Programs/UnrealBuildTool/Platform/Android/UEDeployAndroid.cs
--- a/Engine/Source/Programs/UnrealBuildTool/Platform/Android/UEDeployAndroid.cs
+++ b/Engine/Source/Programs/UnrealBuildTool/Platform/Android/UEDeployAndroid.cs
@@ -5385,13 +5385,13 @@ if ""%ERRORLEVEL%"" NEQ ""0"" (exit /b %ERRORLEVEL%)
 					else
 					{
 						SOPushScript += @$"
-%ADB% %DEVICE% push -z lz4 {FinalSONameStrippedRelative} /data/local/tmp/{FinalSONameStrippedRelative}
-if ""%ERRORLEVEL%"" NEQ ""0"" (%ADB% %DEVICE% push {FinalSONameStrippedRelative} /data/local/tmp/{FinalSONameStrippedRelative})
+%ADB% %DEVICE% push -z lz4 ""{FinalSONameStrippedRelative}"" /data/local/tmp/
+if ""%ERRORLEVEL%"" NEQ ""0"" (%ADB% %DEVICE% push ""{FinalSONameStrippedRelative}"" /data/local/tmp/)
 if ""%ERRORLEVEL%"" NEQ ""0"" (exit /b %ERRORLEVEL%)
 %ADB% %DEVICE% shell run-as {PackageName} mkdir -p ./files
-%ADB% %DEVICE% shell run-as {PackageName} cp /data/local/tmp/{FinalSONameStrippedRelative} ./files/libUnreal.so
+%ADB% %DEVICE% shell run-as {PackageName} cp ""/data/local/tmp/{FinalSONameStrippedRelative}"" ./files/libUnreal.so
 if ""%ERRORLEVEL%"" NEQ ""0"" (exit /b %ERRORLEVEL%)
-%ADB% %DEVICE% shell rm /data/local/tmp/{FinalSONameStrippedRelative}
+%ADB% %DEVICE% shell rm ""/data/local/tmp/{FinalSONameStrippedRelative}""
 ";
 					}
 
