AutomationTool: Prevent NullReferenceException with reparse point that is actually a deduplicated file

diff --git a/Engine/Source/Programs/AutomationTool/AutomationUtils/Utils.cs b/Engine/Source/Programs/AutomationTool/AutomationUtils/Utils.cs
--- a/Engine/Source/Programs/AutomationTool/AutomationUtils/Utils.cs
+++ b/Engine/Source/Programs/AutomationTool/AutomationUtils/Utils.cs
@@ -509,10 +509,13 @@ namespace AutomationTool
 					{
 						FileInfo SourceInfo = new FileInfo(SourceName);
 						FileInfo TargetInfo = new FileInfo(TargetName);
-						
+
 						if (SourceInfo.Attributes.HasFlag(FileAttributes.ReparsePoint))
 						{
-							SourceInfo = (FileInfo)SourceInfo.ResolveLinkTarget(true);
+							if (SourceInfo.ResolveLinkTarget(true) is FileInfo ResolvedInfo && ResolvedInfo != null)
+							{
+								SourceInfo = ResolvedInfo;
+							}
 						}
 
 						if (!bSkipSizeCheck && SourceInfo.Length != TargetInfo.Length)
