Cloud Artifacts - Added option to filter on build group

diff --git a/Engine/Source/Programs/AutomationTool/BuildGraph/Tasks/RetrieveCloudArtifactTask.cs b/Engine/Source/Programs/AutomationTool/BuildGraph/Tasks/RetrieveCloudArtifactTask.cs
--- a/Engine/Source/Programs/AutomationTool/BuildGraph/Tasks/RetrieveCloudArtifactTask.cs
+++ b/Engine/Source/Programs/AutomationTool/BuildGraph/Tasks/RetrieveCloudArtifactTask.cs
@@ -53,6 +53,12 @@ namespace AutomationTool.Tasks
 		[TaskParameter(Optional = true)]
 		public string Name { get; set; } = null!;
 
+		/// <summary>
+		/// Optional filter to require the artifact to be in build group
+		/// </summary>
+		[TaskParameter(Optional = true)]
+		public string? BuildGroup { get; set; } = null!;
+
 		/// <summary>
 		/// The artifact type. Determines the permissions and expiration policy for the artifact.
 		/// </summary>
@@ -325,6 +331,13 @@ namespace AutomationTool.Tasks
 				queryWriter.EndObject();
 			}
 
+			if (_parameters.BuildGroup != null)
+			{
+				queryWriter.BeginObject("buildgroup");
+				queryWriter.WriteString("$eq", _parameters.BuildGroup);
+				queryWriter.EndObject();
+			}
+
 			queryWriter.BeginObject("type");
 			queryWriter.WriteString("$eq", type.ToString());
 			queryWriter.EndObject();
