Cloud artifacts - do not allow to search for empty a build group as that makes no sense

diff --git a/Engine/Source/Programs/AutomationTool/BuildGraph/Tasks/RetrieveCloudArtifactTask.cs b/Engine/Source/Programs/AutomationTool/BuildGraph/Tasks/RetrieveCloudArtifactTask.cs
--- a/Engine/Source/Programs/AutomationTool/BuildGraph/Tasks/RetrieveCloudArtifactTask.cs
+++ b/Engine/Source/Programs/AutomationTool/BuildGraph/Tasks/RetrieveCloudArtifactTask.cs
@@ -331,7 +331,7 @@ namespace AutomationTool.Tasks
 				queryWriter.EndObject();
 			}
 
-			if (_parameters.BuildGroup != null)
+			if (!String.IsNullOrEmpty(_parameters.BuildGroup))
 			{
 				queryWriter.BeginObject("buildgroup");
 				queryWriter.WriteString("$eq", _parameters.BuildGroup);
