Cloud Artifacts - Treat errors from zen cli as warnings as we will generally throw a error based on the exit code. This is not great but solves issue were horde flags tasks as failed when we retry them if they had an error even if the retry eventually hid the issue

diff --git a/Engine/Source/Programs/AutomationTool/AutomationUtils/Matchers/ZenCliEventMatcher.cs b/Engine/Source/Programs/AutomationTool/AutomationUtils/Matchers/ZenCliEventMatcher.cs
--- a/Engine/Source/Programs/AutomationTool/AutomationUtils/Matchers/ZenCliEventMatcher.cs
+++ b/Engine/Source/Programs/AutomationTool/AutomationUtils/Matchers/ZenCliEventMatcher.cs
@@ -32,7 +32,7 @@ namespace AutomationUtils.Matchers
 
 				LogLevel level = match.Groups["severity"].Value switch
 				{
-					"error" => LogLevel.Error,
+					"error" => LogLevel.Warning,
 					"warning" => LogLevel.Information, // we do not want to annotate warnings from zencli as warnings as that causes too much noise - handled retries are warnings but failed retries are already errors
 					_ => LogLevel.Information,
 				};
