Fail / Interrupt (Fault) updates for FNEMatchMaking.cs tests, including changes to Gauntlet infrastructure to support the Fail/Fault distinction in reporting.  This will show Interrupts (Faults) in the logs, as well as help set the stage for Horde integration.

Further changes to enable Fault distinction in more general RPC tests are forthcoming in a later review.

diff --git a/Engine/Source/Programs/AutomationTool/Gauntlet/Framework/Base/Gauntlet.TestNode.cs b/Engine/Source/Programs/AutomationTool/Gauntlet/Framework/Base/Gauntlet.TestNode.cs
--- a/Engine/Source/Programs/AutomationTool/Gauntlet/Framework/Base/Gauntlet.TestNode.cs
+++ b/Engine/Source/Programs/AutomationTool/Gauntlet/Framework/Base/Gauntlet.TestNode.cs
@@ -25,6 +25,7 @@ namespace Gauntlet
 		Cancelled,
 		TimedOut,
 		InsufficientDevices,
+		OperationalException,	// It's not a test failure, but also not specifically a Cancel, or Timeout, or lack of available devices
 	};
 
 	/// <summary>
diff --git a/Engine/Source/Programs/AutomationTool/Gauntlet/Framework/Base/Gauntlet.TestReport.cs b/Engine/Source/Programs/AutomationTool/Gauntlet/Framework/Base/Gauntlet.TestReport.cs
--- a/Engine/Source/Programs/AutomationTool/Gauntlet/Framework/Base/Gauntlet.TestReport.cs
+++ b/Engine/Source/Programs/AutomationTool/Gauntlet/Framework/Base/Gauntlet.TestReport.cs
@@ -29,11 +29,36 @@ namespace Gauntlet
 	public enum TestStateType
 	{
 		Unknown,
+
+        /// <summary>
+        /// The test case has not (yet) been executed
+        /// </summary>
 		NotRun,
+
+        /// <summary>
+        /// The test case is currently in progress
+        /// </summary>
 		InProcess,
+
+        /// <summary>
+        /// The behavior that this test case was trying to exercise did not work as expected
+        /// </summary>
 		Fail,
+
+        /// <summary>
+        /// The behavior that this test case was trying to exercise worked as expected
+        /// </summary>
 		Success,
-		Skipped
+
+        /// <summary>
+        /// The test case was not executed, and will not be executed in this test run
+        /// </summary>
+		Skipped,
+
+        /// <summary>
+        /// The test case was stopped for some reason other than failure of the behavior being tested
+        /// </summary>
+        Interrupted,
 	}
 
 	/// <summary>
diff --git a/Engine/Source/Programs/AutomationTool/Gauntlet/Framework/Gauntlet.TestExecutor.cs b/Engine/Source/Programs/AutomationTool/Gauntlet/Framework/Gauntlet.TestExecutor.cs
--- a/Engine/Source/Programs/AutomationTool/Gauntlet/Framework/Gauntlet.TestExecutor.cs
+++ b/Engine/Source/Programs/AutomationTool/Gauntlet/Framework/Gauntlet.TestExecutor.cs
@@ -58,7 +58,8 @@ namespace Gauntlet
 				TimedOut,
 				Passed,
 				Failed,
-				Skipped
+				Skipped,
+				Interrupted
 			}
 
 			public ITestNode		TestNode;
@@ -171,9 +172,10 @@ namespace Gauntlet
 				RequiredTests = RequiredTests.OrderBy(Node => Node.Priority);
 			}
 
-			Dictionary<string, int> TestIterationsPasssed = new Dictionary<string, int>();
-			Dictionary<string, int> TestIterationsFailed = new Dictionary<string, int>();
-			Dictionary<string, int> TestIterationsPassedWithWarnings = new Dictionary<string, int>();
+			Dictionary<string, int> TestIterationsPassed = new Dictionary<string, int>();
+			Dictionary<string, int> TestIterationsInterrupted = new Dictionary<string, int>();
+            Dictionary<string, int> TestIterationsFailed = new Dictionary<string, int>();
+            Dictionary<string, int> TestIterationsPassedWithWarnings = new Dictionary<string, int>();
 
 			// Clean up local log files
 			if (Directory.Exists(ProcessUtils.LocalLogsPath))
@@ -268,6 +270,7 @@ namespace Gauntlet
 									Node.AddTestEvent(new UnrealTestEvent(EventSeverity.Error, "Test Failed to Start", new List<string> { ex.Message }));
 									PendingTests[i] = null;
 									NodeInfo.TimeSetupBegan = NodeInfo.TimeSetupEnded = NodeInfo.TimeTestEnded = DateTime.Now;
+									NodeInfo.Result = TestExecutionInfo.ExecutionResult.Interrupted;
 									CompletedTests.Add(NodeInfo);
 								}
 							}
@@ -335,6 +338,7 @@ namespace Gauntlet
 									if (Started == false)
 									{
 										TestToStart.TimeSetupEnded = TestToStart.TimeTestEnded = DateTime.Now;
+										TestToStart.Result = TestExecutionInfo.ExecutionResult.Interrupted;
 										CompletedTests.Add(TestToStart);
 										Log.Error(KnownLogEvents.Horde_BuildHealth_Ignore, "Test {Name} failed to start", TestToStart);
 									}
@@ -389,7 +393,25 @@ namespace Gauntlet
 						else
 						{
 							TestInfo.TimeTestEnded = DateTime.Now;
-							TestInfo.Result = Result == TestResult.Passed ? TestExecutionInfo.ExecutionResult.Passed : TestExecutionInfo.ExecutionResult.Failed;
+							switch (Result)
+							{
+								case TestResult.Passed:
+									{
+										TestInfo.Result = TestExecutionInfo.ExecutionResult.Passed;
+										break;
+									}
+								case TestResult.Failed:
+									{
+										TestInfo.Result = TestExecutionInfo.ExecutionResult.Failed;
+										break;
+									}
+								default:
+									{
+										// Default to "Interrupt"... if the test actually Failed, then the test case should explicitly say so
+										TestInfo.Result = TestExecutionInfo.ExecutionResult.Interrupted;
+										break;
+									}
+							}
 							CompletedTests.Add(TestInfo);
 						}
 					}
@@ -466,8 +488,9 @@ namespace Gauntlet
 				{
 					TimeSpan PassDuration = DateTime.Now - StartPassTime;
 
-					int FailedCount = 0;
-					int TestCount = CompletedTests.Count;
+					int InterruptCount = 0;
+                    int FailedCount = 0;
+                    int TestCount = CompletedTests.Count;
 
 					CompletedTests.ForEach(T =>
 					{
@@ -478,33 +501,43 @@ namespace Gauntlet
 						// status msg, kept uniform to avoid spam on notifiers (ie. don't include timestamps, etc) 
 						string Msg = string.Format("Test {0} {1}", T.TestNode, T.Result);
 
-						bool TestHadErrors = T.Result != TestExecutionInfo.ExecutionResult.Passed && T.Result != TestExecutionInfo.ExecutionResult.Skipped;
+                        bool TestHadFailures = T.Result == TestExecutionInfo.ExecutionResult.Failed;
+                        bool TestHadInterrupts = !TestHadFailures && T.Result != TestExecutionInfo.ExecutionResult.Passed && T.Result != TestExecutionInfo.ExecutionResult.Skipped;
 						bool TestHadWarnings = T.TestNode.GetWarnings().Any();
 
-						if (TestHadErrors)
+						if (TestHadFailures)
 						{
 							FailedCount++;
 						}
+                        if (TestHadInterrupts)
+                        {
+                            InterruptCount++;
+                        }
 
-						// increment counts for each test
-						if (!TestIterationsPasssed.ContainsKey(T.TestNode.Name))
+                        // increment counts for each test
+                        if (!TestIterationsPassed.ContainsKey(T.TestNode.Name))
 						{
-							TestIterationsPasssed[T.TestNode.Name] = 0;
-							TestIterationsFailed[T.TestNode.Name] = 0;
-							TestIterationsPassedWithWarnings[T.TestNode.Name] = 0;
+							TestIterationsPassed[T.TestNode.Name] = 0;
+							TestIterationsInterrupted[T.TestNode.Name] = 0;
+                            TestIterationsFailed[T.TestNode.Name] = 0;
+                            TestIterationsPassedWithWarnings[T.TestNode.Name] = 0;
 						}
 
-						if (TestHadErrors)
+						if (TestHadInterrupts)
 						{
-							TestIterationsFailed[T.TestNode.Name]++;
+							TestIterationsInterrupted[T.TestNode.Name]++;
 						}
-						else if (TestHadWarnings)
+                        else if (TestHadFailures)
+                        {
+                            TestIterationsFailed[T.TestNode.Name]++;
+                        }
+                        else if (TestHadWarnings)
 						{
 							TestIterationsPassedWithWarnings[T.TestNode.Name]++;
 						}
 						else
 						{
-							TestIterationsPasssed[T.TestNode.Name]++;
+							TestIterationsPassed[T.TestNode.Name]++;
 						}
 
 						Log.Info(Msg);
@@ -523,7 +556,7 @@ namespace Gauntlet
 					// report all tests
 					ReportMainSummary(CurrentTestPass + 1, Options.TestIterations, PassDuration, CompletedTests);
 
-					if (FailedCount > 0 && Options.StopOnError)
+					if ((FailedCount > 0 || InterruptCount > 0) && Options.StopOnError)
 					{
 						break;
 					}
@@ -531,7 +564,8 @@ namespace Gauntlet
 			} // foreach pass
 
 			int TotalTests = RequiredTests.Count() * Options.TestIterations;
-			int FailedTestCount = TestIterationsFailed.Values.Sum();
+            int InterruptTestCount = TestIterationsInterrupted.Values.Sum();
+            int FailedTestCount = TestIterationsFailed.Values.Sum();
 
 			// show details for multi passes				
 			if (Options.TestIterations > 1)
@@ -540,29 +574,33 @@ namespace Gauntlet
 
 				MB.HorizontalLine();					
 
-				MB.Paragraph(string.Format("All iterations completed. {0} of {1} executed tests completed without error", TotalTests - FailedTestCount, TotalTests));
+				MB.Paragraph(string.Format("All iterations completed. {0} of {1} executed tests completed without error", TotalTests - (FailedTestCount + InterruptTestCount), TotalTests));
 
 				List<string> Lines = new List<string>();
 				foreach (ITestNode Test in RequiredTests)
 				{
-					int IterationsWithErrors = TestIterationsFailed[Test.Name];
-					int IterationsWithWarnings = TestIterationsPassedWithWarnings[Test.Name];
-					int IterationsWithPasses = TestIterationsPasssed[Test.Name];				
+					int IterationsWithInterrupts = TestIterationsInterrupted[Test.Name];
+                    int IterationsWithFailures = TestIterationsFailed[Test.Name];
+                    int IterationsWithWarnings = TestIterationsPassedWithWarnings[Test.Name];
+					int IterationsWithPasses = TestIterationsPassed[Test.Name];				
 
-					Lines.Add(string.Format("{0}: {1} Iterations, {2} Passed, {3} Passed with Warnings, {4} Failed", 
-						Test.Name, Options.TestIterations, IterationsWithPasses, IterationsWithWarnings, IterationsWithErrors));
+					Lines.Add(string.Format("{0}: {1} Iterations, {2} Passed, {3} Passed with Warnings, {4} Failed, {5} Interrupts", 
+						Test.Name, Options.TestIterations, IterationsWithPasses, IterationsWithWarnings, IterationsWithFailures, IterationsWithInterrupts));
 				}
 
 				MB.UnorderedList(Lines);
-					
-				if (FailedTestCount > 0)
+
+                if (InterruptTestCount > 0)
+                {
+                    MB.Paragraph(string.Format("Error: {0} of {1} executed tests had Interrupts.", InterruptTestCount, TotalTests));
+                }
+                if (FailedTestCount > 0)
 				{
 					MB.Paragraph(string.Format("Error: {0} of {1} executed tests failed.", FailedTestCount, TotalTests));
 				}
-				else
-				{
-					MB.Paragraph(string.Format("{0} total tests passed.", TotalTests));
-				}
+
+				// Always report the number of passed tests
+				MB.Paragraph(string.Format("{0} total tests passed.", TotalTests - (InterruptTestCount + FailedTestCount)));
 
 				MB.HorizontalLine();
 
@@ -574,7 +612,7 @@ namespace Gauntlet
 			Globals.AbortHandlers.Remove(CancelHandler);
 			Globals.PostAbortHandlers.Remove(PostCancelHandler);
 			
-			return FailedTestCount == 0 && !IsCancelled;
+			return InterruptTestCount == 0 && FailedTestCount == 0 && !IsCancelled;
 		}
 
 		/// <summary>
diff --git a/Engine/Source/Programs/AutomationTool/Gauntlet/Framework/Utils/Gauntlet.HordeReport.cs b/Engine/Source/Programs/AutomationTool/Gauntlet/Framework/Utils/Gauntlet.HordeReport.cs
--- a/Engine/Source/Programs/AutomationTool/Gauntlet/Framework/Utils/Gauntlet.HordeReport.cs
+++ b/Engine/Source/Programs/AutomationTool/Gauntlet/Framework/Utils/Gauntlet.HordeReport.cs
@@ -1330,7 +1330,7 @@ namespace Gauntlet
 				}
 
 				/// <summary>
-				/// Set the outcome of the phase, it will udpate the test session phase counts
+				/// Set the outcome of the phase, it will update the test session phase counts
 				/// </summary>
 				/// <param name="OutCome"></param>
 				public void SetOutcome(TestPhaseOutcome OutCome)
diff --git a/Engine/Source/Programs/AutomationTool/Gauntlet/Unreal/Utils/Gauntlet.UnrealTestResults.cs b/Engine/Source/Programs/AutomationTool/Gauntlet/Unreal/Utils/Gauntlet.UnrealTestResults.cs
--- a/Engine/Source/Programs/AutomationTool/Gauntlet/Unreal/Utils/Gauntlet.UnrealTestResults.cs
+++ b/Engine/Source/Programs/AutomationTool/Gauntlet/Unreal/Utils/Gauntlet.UnrealTestResults.cs
@@ -359,7 +359,7 @@ namespace Gauntlet
 		{
 			get
 			{
-				return State == TestStateType.Fail;
+				return (State == TestStateType.Fail || State == TestStateType.Interrupted);
 			}
 		}
 		[JsonIgnore]
