Cloud Artifacts - Make zen cli matcher high priority so it doesn't get overriden by other matchers

diff --git a/Engine/Source/Programs/AutomationTool/AutomationUtils/Matchers/ZenCliEventMatcher.cs b/Engine/Source/Programs/AutomationTool/AutomationUtils/Matchers/ZenCliEventMatcher.cs
--- a/Engine/Source/Programs/AutomationTool/AutomationUtils/Matchers/ZenCliEventMatcher.cs
+++ b/Engine/Source/Programs/AutomationTool/AutomationUtils/Matchers/ZenCliEventMatcher.cs
@@ -37,7 +37,7 @@ namespace AutomationUtils.Matchers
 					_ => LogLevel.Information,
 				};
 
-				return builder.ToMatch(LogEventPriority.Low, level, KnownLogEvents.ZenCli_Event);
+				return builder.ToMatch(LogEventPriority.High, level, KnownLogEvents.ZenCli_Event);
 				
 			}
 			return null;
