Staging: ZenServer: Fix possible timeouts reading the oplog. In ReadZenCookedFilesFromZenServer, first quickly health check the server and launch if necessary, and then fetch the oplog with a timeout of 20 minutes rather than the default of 100 seconds.
#rnx

diff --git a/Engine/Source/Programs/AutomationTool/AutomationUtils/ZenUtils.cs b/Engine/Source/Programs/AutomationTool/AutomationUtils/ZenUtils.cs
--- a/Engine/Source/Programs/AutomationTool/AutomationUtils/ZenUtils.cs
+++ b/Engine/Source/Programs/AutomationTool/AutomationUtils/ZenUtils.cs
@@ -48,6 +48,56 @@ namespace AutomationUtils
 		public bool IsValid { get; private set; } = false;
 
 		public static ILogger Logger => Log.Logger;
+		public static int DefaultPort { get { return 8558; } }
+
+		public void InitializeFromHostNameAndPort(string HostNameAndPort)
+		{
+			// Leading and trailing spaces are trimmed, but spaces before or after the colon are not trimmed and are
+			// treated as part of the hostname.
+			// Valid strings:
+			// "hostname", "1.2.3.4", "[::1]"
+			// "hostname:1234", "1.2.3.4:1234", "[::1]:1234"
+			HostNameAndPort = HostNameAndPort.Trim();
+			string[] Parts = HostNameAndPort.Split(':');
+
+			string LocalHostName;
+			int LocalHostPort;
+			if (Parts.Length >= 2)
+			{
+				int PortIndex = Parts.Length - 1;
+				if (!int.TryParse(Parts[PortIndex], out LocalHostPort))
+				{
+					LocalHostName = HostNameAndPort;
+					LocalHostPort = DefaultPort;
+				}
+				else
+				{
+					LocalHostName = string.Join(':', Parts.Take(PortIndex));
+				}
+			}
+			else
+			{
+				LocalHostName = HostNameAndPort;
+				LocalHostPort = DefaultPort;
+			}
+
+			HostName = LocalHostName;
+			HostPort = LocalHostPort;
+			IsValid = true;
+		}
+
+		public void InitializeFromHostNameAndPort(string InHostName, int InHostPort)
+		{
+			HostName = InHostName;
+			HostPort = InHostPort;
+			IsValid = true;
+		}
+
+		public void InitializeFromHostName(string InHostName)
+		{
+			HostName = InHostName;
+			HostPort = DefaultPort;
+		}
 
 		public static ZenRunContext ReadFromContextFile(FileReference ContextFile)
 		{
@@ -460,6 +510,20 @@ namespace AutomationUtils
 			return Context.IsServiceRunning();
 		}
 
+		public static bool IsZenServerRunning(string HostNameAndPort)
+		{
+			ZenRunContext ZenServerContext = new ZenRunContext();
+			ZenServerContext.InitializeFromHostNameAndPort(HostNameAndPort);
+			return ZenServerContext.IsServiceRunning();
+		}
+
+		public static bool IsZenServerRunning(string HostName, int HostPort)
+		{
+			ZenRunContext ZenServerContext = new ZenRunContext();
+			ZenServerContext.InitializeFromHostNameAndPort(HostName, HostPort);
+			return ZenServerContext.IsServiceRunning();
+		}
+
 		public static JsonSerializerOptions GetDefaultJsonSerializerOptions()
 		{
 			JsonSerializerOptions Options = new();
diff --git a/Engine/Source/Programs/AutomationTool/Scripts/CopyBuildToStagingDirectory.Automation.cs b/Engine/Source/Programs/AutomationTool/Scripts/CopyBuildToStagingDirectory.Automation.cs
--- a/Engine/Source/Programs/AutomationTool/Scripts/CopyBuildToStagingDirectory.Automation.cs
+++ b/Engine/Source/Programs/AutomationTool/Scripts/CopyBuildToStagingDirectory.Automation.cs
@@ -958,21 +958,47 @@ namespace AutomationScripts
 			return HostName;
 		}
 
-		private static List<ZenCookedFile> ReadZenCookedFilesFromZenServer(ProjectParams Params, DeploymentContext SC, bool bAutoLaunch, string PackageStoreFileArgName, string PackageStoreFileArgValue)
+		private static List<ZenCookedFile> ReadZenCookedFilesFromZenServer(ProjectParams Params, DeploymentContext SC,
+			bool bAutoLaunch, string PackageStoreFileArgName, string PackageStoreFileArgValue)
 		{
 			List<ZenCookedFile> ZenCookedFiles = new List<ZenCookedFile>();
 
+			string HostNameAndPort = GetZenHostName(SC);
 			string ProjectId = SC.PackageStoreData.ZenServerStore.ProjectId;
 			string OplogId = SC.PackageStoreData.ZenServerStore.OplogId;
 
-			bool bAttemptAutoLaunchOnFailure = bAutoLaunch;
+			if (!ZenUtils.IsZenServerRunning(HostNameAndPort))
+			{
+				if (!bAutoLaunch)
+				{
+					string ExceptionMessage = string.Format(
+						"Failed sending oplog request to Zen at {0}: ZenServer is not running (or is not responding), and bAutoLaunch == false.",
+						HostNameAndPort);
+					throw new AutomationException(ExceptionMessage);
+				}
+
+				string Arguments = String.Format("IoStore -StartZenServerForStage -SponsorProcessID={0} -{1}={2}",
+					Environment.ProcessId, PackageStoreFileArgName, MakePathSafeToUseWithCommandLine(PackageStoreFileArgValue));
+				if (Params.Unattended)
+				{
+					Arguments += " -unattended";
+				}
+				RunUnrealPak(Params.RawProjectPath, Arguments, "ZenAutoLaunch");
+				if (!ZenUtils.IsZenServerRunning(HostNameAndPort))
+				{
+					string ExceptionMessage = string.Format(
+						"Failed sending oplog request to Zen at {0}: ZenServer is not running (or is not responding), even after attempting to AutoLaunch.",
+						HostNameAndPort);
+					throw new AutomationException(ExceptionMessage);
+				}
+			}
+
 			HttpResponseMessage HttpGetResult = null;
-			string HostName;
-			for (; ; )
 			{
-				HostName = GetZenHostName(SC);
 				HttpClient HttpClient = CreateZenHttpClient(SC);
-				string Uri = string.Format("http://{0}/prj/{1}/oplog/{2}/entries", HostName, ProjectId, OplogId);
+				HttpClient.Timeout = TimeSpan.FromMinutes(20);
+
+				string Uri = string.Format("http://{0}/prj/{1}/oplog/{2}/entries", HostNameAndPort, ProjectId, OplogId);
 				Uri += "?trim_by_referencedset=true";
 
 				using var Request = new HttpRequestMessage(HttpMethod.Get, Uri);
@@ -980,37 +1006,19 @@ namespace AutomationScripts
 				try
 				{
 					HttpGetResult = HttpClient.Send(Request);
-					break;
 				}
 				catch (Exception Ex)
 				{
-					if (bAttemptAutoLaunchOnFailure)
-					{
-						string Arguments = String.Format("IoStore -StartZenServerForStage -SponsorProcessID={0} -{1}={2}", Environment.ProcessId, PackageStoreFileArgName, MakePathSafeToUseWithCommandLine(PackageStoreFileArgValue));
-						if (Params.Unattended)
-						{
-							Arguments += " -unattended";
-						}
-
-						RunUnrealPak(Params.RawProjectPath, Arguments, "ZenAutoLaunch");
-						bAttemptAutoLaunchOnFailure = false;
-					}
-					else
-					{
-						string RetryMessage = bAutoLaunch
-							? "Failed even after attempting restart."
-							: "bAutoLaunch is false so did not attempt restart. Ensure that the server is running.";
-						string ExceptionMessage = String.Format("Failed sending oplog request to Zen at {0} for oplog {1}.{2}: {3}. {4}",
-								HostName, ProjectId, OplogId, Ex.Message, RetryMessage);
-						throw new AutomationException(ExceptionMessage);
-					}
+					string ExceptionMessage = String.Format("Failed sending oplog request to Zen at {0} for oplog {1}.{2}: {3}.",
+							HostNameAndPort, ProjectId, OplogId, Ex.Message);
+					throw new AutomationException(ExceptionMessage);
 				}
 			}
 			if (!HttpGetResult.IsSuccessStatusCode)
 			{
 				string ExceptionMessage = String.Format("Failed reading oplog from Zen at {0} for oplog {1}.{2}, "
 					+ "HttpClient.Send returned StatusCode {3}. Ensure that cooking was successful.",
-					HostName, ProjectId, OplogId, HttpGetResult.StatusCode);
+					HostNameAndPort, ProjectId, OplogId, HttpGetResult.StatusCode);
 				throw new AutomationException(ExceptionMessage);
 			}
 			Task<byte[]> ReadOplogTask = HttpGetResult.Content.ReadAsByteArrayAsync();
