Cloud Artifacts - Fixed retries on 502 and 504 to support the new exit code ranges from zen cli

diff --git a/Engine/Source/Programs/AutomationTool/BuildGraph/Tasks/CreateCloudArtifactTask.cs b/Engine/Source/Programs/AutomationTool/BuildGraph/Tasks/CreateCloudArtifactTask.cs
--- a/Engine/Source/Programs/AutomationTool/BuildGraph/Tasks/CreateCloudArtifactTask.cs
+++ b/Engine/Source/Programs/AutomationTool/BuildGraph/Tasks/CreateCloudArtifactTask.cs
@@ -541,7 +541,14 @@ namespace AutomationTool.Tasks
 					int[] waitDurationSeconds = [60, 300, 900];
 					int waitDuration = waitDurationSeconds[countOfRetries];
 					int exitCode = (int)e.ErrorCode;
-					if (exitCode is 502 or 504)
+					// 502 or 504 are the old exit codes that indicate errors
+					// new exit codes for retries that also work on mac and linux are 69, 92 and 93
+					// From zen: 
+					// kHttpServiceUnavailable  = 69,	// ServiceUnavailable(503)
+					// kHttpBadGateway          = 92,	// BadGateway(502)
+					// kHttpGatewayTimeout      = 93,	// GatewayTimeout(504)
+
+					if (exitCode is 502 or 504 or 69 or 92 or 93)
 					{
 						// zen returns http status code as exit codes if it had an issue
 						Logger.LogInformation("Http Status code {ExitCode} returned from zen. Waiting {WaitDuration}s before trying.", exitCode, waitDuration);
diff --git a/Engine/Source/Programs/AutomationTool/BuildGraph/Tasks/ZenExportSnapshotTask.cs b/Engine/Source/Programs/AutomationTool/BuildGraph/Tasks/ZenExportSnapshotTask.cs
--- a/Engine/Source/Programs/AutomationTool/BuildGraph/Tasks/ZenExportSnapshotTask.cs
+++ b/Engine/Source/Programs/AutomationTool/BuildGraph/Tasks/ZenExportSnapshotTask.cs
@@ -453,7 +453,8 @@ namespace AutomationTool.Tasks
 					return true;
 				}
 
-				bool isGatewayError = exitCode is 502 or 504;
+				// 502 and 504 are the old exit codes for retries the new exit codes for retries that also work on mac and linux are 69, 92 and 93 see zen cli documentation of exit codes for more details
+				bool isGatewayError = exitCode is 502 or 504 or 69 or 92 or 93;
 				int waitDuration = isGatewayError ? waitDurationWebSeconds[attempt] : 0;
 				bool willRetry = attempt < (attemptLimit - 1);
 				Logger.LogWarning("Attempt {AttemptNum} of exporting the oplog failed, {Action}...", attempt + 1, willRetry ? $"retrying after {waitDuration}s" : "abandoning");
