Fixed adb pathing of Android installer/uninstaller scripts for all platforms

#android

diff --git a/Engine/Source/Programs/AutomationTool/Android/AndroidPlatform.Automation.cs b/Engine/Source/Programs/AutomationTool/Android/AndroidPlatform.Automation.cs
--- a/Engine/Source/Programs/AutomationTool/Android/AndroidPlatform.Automation.cs
+++ b/Engine/Source/Programs/AutomationTool/Android/AndroidPlatform.Automation.cs
@@ -1549,7 +1549,7 @@ public class AndroidPlatform : Platform
 				}
 				// make a batch file that can be used to uninstall the .apk and .obb files
 				string UninstallBatchName = GetFinalBatchName(ApkName, SC, false, EBatchType.Uninstall, Target);
-				string[] UninstallBatchLines = GenerateUninstallBatchFile(bPackageDataInsideApk, PackageName, ApkName, Params, bIsPC);
+				string[] UninstallBatchLines = GenerateUninstallBatchFile(bPackageDataInsideApk, PackageName, ApkName, Params, bIsPC, Target);
 				if (bHaveAPK || bHaveUniversal)
 				{
 					File.WriteAllText(UninstallBatchName, string.Join(LineEnding, UninstallBatchLines) + LineEnding);
@@ -1729,8 +1729,18 @@ public class AndroidPlatform : Platform
 						"#!/bin/sh",
 						"cd \"`dirname \"$0\"`\"",
 						"AFS=./" + AFSExecutable,
-						"ADB=",
-						"if [ \"$ANDROID_HOME\" != \"\" ]; then ADB=$ANDROID_HOME/platform-tools/adb; else ADB=" +Environment.GetEnvironmentVariable("ANDROID_HOME") + "/platform-tools/adb; fi",
+						"ADB=adb",
+						"if ! command -v -- \"$ADB\" >/dev/null 2>&1; then",
+						"\tSTUDIO_SDK_PATH=$ANDROID_HOME",
+						"\tif [ ! -d \"$STUDIO_SDK_PATH\" ]; then",
+						"\t\tSTUDIO_SDK_PATH=" + (Target == UnrealTargetPlatform.Mac ? "~/Library/Android/sdk" : "$HOME/Android/Sdk"),
+						"\t\tif [ ! -d \"$STUDIO_SDK_PATH\" ]; then",
+						"\t\t\techo Unable to locate local Android SDK location.",
+						"\t\t\texit 1",
+						"\t\tfi",
+						"\tfi",
+						"\tADB=$STUDIO_SDK_PATH/platform-tools/adb",
+						"fi",
 						"DEVICE=",
 						"if [ \"$1\" != \"\" ]; then DEVICE=\"-s $1\"; fi",
 						"echo",
@@ -1836,11 +1846,26 @@ public class AndroidPlatform : Platform
 
 			Logger.LogInformation("Writing bat for install with {Arg0}", bPackageDataInsideApk ? "data in APK" : "separate OBB");
 			BatchLines.AddRange(new string[] {
-						"setlocal",
+						"setlocal ENABLEEXTENSIONS",
+						"setlocal ENABLEDELAYEDEXPANSION",
 						"if NOT \"%UE_SDKS_ROOT%\"==\"\" (call %UE_SDKS_ROOT%\\HostWin64\\Android\\SetupEnvironmentVars.bat)",
-						"set ANDROIDHOME=%ANDROID_HOME%",
-						"if \"%ANDROIDHOME%\"==\"\" set ANDROIDHOME="+Environment.GetEnvironmentVariable("ANDROID_HOME"),
-						"set ADB=%ANDROIDHOME%\\platform-tools\\adb.exe",
+						"set ADB=adb.exe",
+						"where.exe /q %ADB%",
+						"if /i \"%ERRORLEVEL%\" neq \"0\" (",
+						"\tset STUDIO_SDK_PATH=%ANDROID_HOME%",
+						"\tif not exist \"!STUDIO_SDK_PATH!\" (",
+						"\t\tset STUDIO_SDK_PATH=",
+						"\t\tfor /f \"tokens=2*\" %%A in ('reg.exe query \"HKLM\\SOFTWARE\\Android Studio\" /v \"SdkPath\"') do (set STUDIO_SDK_PATH=%%B)",
+						"\t\tif not exist \"!STUDIO_SDK_PATH!\" (",
+						"\t\t\tset STUDIO_SDK_PATH=%LOCALAPPDATA%\\Android\\Sdk",
+						"\t\t\tif not exist \"!STUDIO_SDK_PATH!\" (",
+						"\t\t\t\techo Unable to locate local Android SDK location.",
+						"\t\t\t\texit /b 1",
+						"\t\t\t)",
+						"\t\t)",
+						"\t)",
+						"\tset ADB=!STUDIO_SDK_PATH!\\platform-tools\\adb.exe",
+						")",
 						"set AFS=.\\" + AFSExecutable.Replace("/", "\\"),
 						"set DEVICE=",
 						"if not \"%1\"==\"\" set DEVICE=-s %1",
@@ -1900,7 +1925,7 @@ public class AndroidPlatform : Platform
 	}
 
 
-	private string[] GenerateUninstallBatchFile(bool bPackageDataInsideApk, string PackageName, string ApkName, ProjectParams Params, bool bIsPC)
+	private string[] GenerateUninstallBatchFile(bool bPackageDataInsideApk, string PackageName, string ApkName, ProjectParams Params, bool bIsPC, UnrealTargetPlatform Target)
 	{
 		string[] BatchLines = null;
 
@@ -1910,8 +1935,18 @@ public class AndroidPlatform : Platform
 			BatchLines = new string[] {
 						"#!/bin/sh",
 						"cd \"`dirname \"$0\"`\"",
-						"ADB=",
-						"if [ \"$ANDROID_HOME\" != \"\" ]; then ADB=$ANDROID_HOME/platform-tools/adb; else ADB=" +Environment.GetEnvironmentVariable("ANDROID_HOME") + "/platform-tools/adb; fi",
+						"ADB=adb",
+						"if ! command -v -- \"$ADB\" >/dev/null 2>&1; then",
+						"\tSTUDIO_SDK_PATH=$ANDROID_HOME",
+						"\tif [ ! -d \"$STUDIO_SDK_PATH\" ]; then",
+						"\t\tSTUDIO_SDK_PATH=" + (Target == UnrealTargetPlatform.Mac ? "~/Library/Android/sdk" : "$HOME/Android/Sdk"),
+						"\t\tif [ ! -d \"$STUDIO_SDK_PATH\" ]; then",
+						"\t\t\techo Unable to locate local Android SDK location.",
+						"\t\t\texit 1",
+						"\t\tfi",
+						"\tfi",
+						"\tADB=$STUDIO_SDK_PATH/platform-tools/adb",
+						"fi",
 						"DEVICE=",
 						"if [ \"$1\" != \"\" ]; then DEVICE=\"-s $1\"; fi",
 						"echo",
@@ -1932,11 +1967,26 @@ public class AndroidPlatform : Platform
 		{
 			Logger.LogInformation("Writing bat for uninstall with {Arg0}", bPackageDataInsideApk ? "data in APK" : "separate OBB");
 			BatchLines = new string[] {
-						"setlocal",
+						"setlocal ENABLEEXTENSIONS",
+						"setlocal ENABLEDELAYEDEXPANSION",
 						"if NOT \"%UE_SDKS_ROOT%\"==\"\" (call %UE_SDKS_ROOT%\\HostWin64\\Android\\SetupEnvironmentVars.bat)",
-						"set ANDROIDHOME=%ANDROID_HOME%",
-						"if \"%ANDROIDHOME%\"==\"\" set ANDROIDHOME="+Environment.GetEnvironmentVariable("ANDROID_HOME"),
-						"set ADB=%ANDROIDHOME%\\platform-tools\\adb.exe",
+						"set ADB=adb.exe",
+						"where.exe /q %ADB%",
+						"if /i \"%ERRORLEVEL%\" neq \"0\" (",
+						"\tset STUDIO_SDK_PATH=%ANDROID_HOME%",
+						"\tif not exist \"!STUDIO_SDK_PATH!\" (",
+						"\t\tset STUDIO_SDK_PATH=",
+						"\t\tfor /f \"tokens=2*\" %%A in ('reg.exe query \"HKLM\\SOFTWARE\\Android Studio\" /v \"SdkPath\"') do (set STUDIO_SDK_PATH=%%B)",
+						"\t\tif not exist \"!STUDIO_SDK_PATH!\" (",
+						"\t\t\tset STUDIO_SDK_PATH=%LOCALAPPDATA%\\Android\\Sdk",
+						"\t\t\tif not exist \"!STUDIO_SDK_PATH!\" (",
+						"\t\t\t\techo Unable to locate local Android SDK location.",
+						"\t\t\t\texit /b 1",
+						"\t\t\t)",
+						"\t\t)",
+						"\t)",
+						"\tset ADB=!STUDIO_SDK_PATH!\\platform-tools\\adb.exe",
+						")",
 						"set DEVICE=",
 						"if not \"%1\"==\"\" set DEVICE=-s %1",
 						"for /f \"delims=\" %%A in ('%ADB% %DEVICE% " + GetStorageQueryCommand(true) +"') do @set STORAGE=%%A",
