Enable test data v2 by default
+ introduce test data v1 reverse flag

#rnx

diff --git a/Engine/Build/Graph/Tasks/Inc/GauntletSettings.xml b/Engine/Build/Graph/Tasks/Inc/GauntletSettings.xml
--- a/Engine/Build/Graph/Tasks/Inc/GauntletSettings.xml
+++ b/Engine/Build/Graph/Tasks/Inc/GauntletSettings.xml
@@ -25,7 +25,7 @@
 	<Option Name="ExtraGauntletCommonArgs" DefaultValue="" Description="Additional Common arguments passed to Gauntlet command line"/>
 	<Option Name="ExtraGauntletTargetTestArgs" DefaultValue="" Description="Additional Target test arguments passed to Gauntlet command line"/>
 	<Option Name="ExtraGauntletEditorTestArgs" DefaultValue="" Description="Additional Editor test arguments passed to Gauntlet command line"/>
-	<Option Name="UseTestDataV2" DefaultValue="$(Horde)" Description="Publish report as TestData v2 for Horde"/>
+	<Option Name="UseTestDataV1" DefaultValue="false" Description="Publish report as TestData v1 for Horde"/>
 
 	<Option Name="AdbKeysPath" DefaultValue="" Description="Path to adbkeys location (for Android)"/>
 
@@ -33,7 +33,7 @@
 
 	<Property Name="GauntletCommonArgs" Value="-maxduration=$(GauntletTimeOut) -unattended -verbose -ResumeOnCriticalFailure $(ExtraGauntletCommonArgs)" />
 	<Property Name="GauntletCommonArgs" Value="$(GauntletCommonArgs) -WriteTestResultsForHorde" If="$(Horde) and !$(DisableTestResultsForHorde)"/>
-	<Property Name="GauntletCommonArgs" Value="$(GauntletCommonArgs) -UseTestDataV2" If="$(Horde) and !$(DisableTestResultsForHorde) and $(UseTestDataV2)"/>
+	<Property Name="GauntletCommonArgs" Value="$(GauntletCommonArgs) -UseTestDataV1" If="$(Horde) and !$(DisableTestResultsForHorde) and $(UseTestDataV1)"/>
 	<Property Name="GauntletCommonArgs" Value="$(GauntletCommonArgs) -branch=&quot;$(Branch)&quot; -changelist=&quot;$(Change)&quot;" />
 
 	<Property Name="GauntletDeviceArgs" Value="" />
diff --git a/Engine/Source/Programs/AutomationTool/Gauntlet/Framework/Horde/Gauntlet.Horde.cs b/Engine/Source/Programs/AutomationTool/Gauntlet/Framework/Horde/Gauntlet.Horde.cs
--- a/Engine/Source/Programs/AutomationTool/Gauntlet/Framework/Horde/Gauntlet.Horde.cs
+++ b/Engine/Source/Programs/AutomationTool/Gauntlet/Framework/Horde/Gauntlet.Horde.cs
@@ -36,6 +36,18 @@ namespace Gauntlet
 				return Environment.GetEnvironmentVariable("UE_HORDE_STEPID");
 			}
 		}
+
+		/// <summary>
+		/// Whether to generate test data v1 model when writing test results for Horde
+		/// </summary>
+		static public bool UseTestDataV1
+		{
+			get
+			{
+				return Globals.Params.ParseParam("UseTestDataV1");
+			}
+		}
+
 		public class JobStepReport
 		{
 			public string scope { get; } = "Step";
@@ -116,7 +128,7 @@ namespace Gauntlet
 				}
 
 				// detect if any testdata was generated
-				if (Globals.Params.ParseParam("UseTestDataV2"))
+				if (!UseTestDataV1)
 				{
 					DirectoryInfo TestDataOutputFolder = new DirectoryInfo(Path.Combine(CommandUtils.CmdEnv.EngineSavedFolder, "TestData"));
 					if (TestDataOutputFolder.Exists && TestDataOutputFolder.GetFiles("*.json", SearchOption.TopDirectoryOnly).Any())
diff --git a/Engine/Source/Programs/AutomationTool/Gauntlet/Unreal/Base/Gauntlet.UnrealTestNode.cs b/Engine/Source/Programs/AutomationTool/Gauntlet/Unreal/Base/Gauntlet.UnrealTestNode.cs
--- a/Engine/Source/Programs/AutomationTool/Gauntlet/Unreal/Base/Gauntlet.UnrealTestNode.cs
+++ b/Engine/Source/Programs/AutomationTool/Gauntlet/Unreal/Base/Gauntlet.UnrealTestNode.cs
@@ -1416,7 +1416,13 @@ namespace Gauntlet
 			string ArtifactPathForReport = ArtifactPath;
 			Horde.OnStepSummary += () => Horde.GenerateArtifactLink(ArtifactPathForReport);
 			BaseHordeReport HordeTestReport = null;
-			if (Globals.Params.ParseParam("UseTestDataV2"))
+			if (Horde.UseTestDataV1)
+			{
+				// Legacy
+				SimpleTestReport SimpleReport = CreateSimpleReportForHorde_legacy(Result);
+				HordeTestReport = SimpleReport;
+			}
+			else
 			{
 				AutomatedTestSessionData TestReport = new AutomatedTestSessionData(HordeReportTestKey, HordeReportTestName);
 				DateTime SessionTime = SessionStartTime.ToUniversalTime();
@@ -1463,19 +1469,6 @@ namespace Gauntlet
 				// Metadata
 				SetReportMetadata(HordeTestReport, new UnrealTestRole[] { GetCachedConfiguration().GetMainRequiredRole() });
 			}
-			if (HordeTestReport == null || Globals.Params.ParseParam("AddTestDataV1"))
-			{
-				SimpleTestReport SimpleReport = CreateSimpleReportForHorde_legacy(Result);
-				// Set or Attach
-				if (HordeTestReport == null)
-				{
-					HordeTestReport = SimpleReport;
-				}
-				else
-				{
-					HordeTestReport.AttachDependencyReport(SimpleReport.GetTestData(), GetHordeTestDataKey_legacy());
-				}
-			}
 			string HordeArtifactPath = string.IsNullOrEmpty(GetCachedConfiguration().HordeArtifactPath) ? DefaultArtifactsDir : GetCachedConfiguration().HordeArtifactPath;
 			string ArtifactPathLabel = new DirectoryInfo(ArtifactPath).Name;
 			HordeTestReport.SetOutputArtifactPath(Path.Combine(HordeArtifactPath, ArtifactPathLabel));
@@ -1622,7 +1615,7 @@ namespace Gauntlet
 
 			BaseHordeReport HordeTestPassResults = null;
 			// Handle initialization phase
-			if (Globals.Params.ParseParam("UseTestDataV2"))
+			if (!Horde.UseTestDataV1)
 			{
 				AutomatedTestSessionData SessionReport = new AutomatedTestSessionData(HordeReportTestKey, HordeReportTestName);
 				string ArtifactPathLabel = new DirectoryInfo(ArtifactPath).Name;
@@ -1689,27 +1682,20 @@ namespace Gauntlet
 					}
 					SessionReport.PopulateFromUnrealAutomatedTests(JsonTestPassResults, UnrealAutomatedTestReportPath);
 				}
-				if (HordeTestPassResults == null || Globals.Params.ParseParam("AddTestDataV1"))
+				else if (HordeTestPassResults == null)
 				{
+					// Legacy report
 					string HordeArtifactPath = GetCachedConfiguration().HordeArtifactPath;
 					AutomatedTestSessionData_legacy LegacyTestPassResults = AutomatedTestSessionData_legacy.FromUnrealAutomatedTests(
 						JsonTestPassResults, Type, Suite, UnrealAutomatedTestReportPath, HordeArtifactPath
 					);
 					// Pre Flight information
 					LegacyTestPassResults.PreFlightChange = GetCachedConfiguration().PreFlightChange;
-					// Set or Attach
-					if (HordeTestPassResults == null)
-					{
-						HordeTestPassResults = LegacyTestPassResults;
-					}
-					else
+					HordeTestPassResults = LegacyTestPassResults;
+					// Check dependencies
+					foreach (var Item in LegacyTestPassResults.GetReportDependencies())
 					{
-						SetReportMetadata(LegacyTestPassResults, RoleList);
-						HordeTestPassResults.AttachDependencyReport(LegacyTestPassResults.GetTestData(), GetHordeTestDataKey_legacy());
-						foreach (var Item in LegacyTestPassResults.GetReportDependencies())
-						{
-							HordeTestPassResults.AttachDependencyReport(Item.Value, Item.Key);
-						}
+						HordeTestPassResults.AttachDependencyReport(Item.Value, Item.Key);
 					}
 					// Make a copy of the report in the old way - until we decide the transition is over
 					UnrealEngineTestPassResults CopyTestPassResults = UnrealEngineTestPassResults.FromUnrealAutomatedTests(JsonTestPassResults, ReportURL);
@@ -1758,7 +1744,7 @@ namespace Gauntlet
 						string DeviceKey = Artifact.AppInstance.Device.Name.Replace(".", "-"); // Key must not contains dot
 						if (!string.IsNullOrEmpty(Artifact.LogPath))
 						{
-							string LogName = FileUtils.ConvertPathToUnix(Path.GetRelativePath(Path.GetFullPath(Globals.Params.ParseParam("UseTestDataV2") ? ArtifactPath : Context.Options.LogDir), Path.GetFullPath(Artifact.LogPath)));
+							string LogName = FileUtils.ConvertPathToUnix(Path.GetRelativePath(Path.GetFullPath(Horde.UseTestDataV1 ? Context.Options.LogDir : ArtifactPath), Path.GetFullPath(Artifact.LogPath)));
 							if (HordeTestPassResults is BaseHordeReport Report)
 							{
 								Report.AttachDeviceLog(DeviceKey, Artifact.LogPath, LogName);
