Improve RPC exceptions

diff --git a/Engine/Source/Programs/AutomationTool/Gauntlet/Framework/RpcFramework/Gauntlet.RpcLibrary.cs b/Engine/Source/Programs/AutomationTool/Gauntlet/Framework/RpcFramework/Gauntlet.RpcLibrary.cs
--- a/Engine/Source/Programs/AutomationTool/Gauntlet/Framework/RpcFramework/Gauntlet.RpcLibrary.cs
+++ b/Engine/Source/Programs/AutomationTool/Gauntlet/Framework/RpcFramework/Gauntlet.RpcLibrary.cs
@@ -205,28 +205,36 @@ namespace Gauntlet
 			int DelayInMilliseconds = 5 * 1000, int TimeoutInMilliseconds = 60 * 1000, bool bSuppressExceptions = false, bool bSuppressTimeoutException = false, OnBeforeTimeoutDelegate OnBeforeTimeout = null)
 		{
 			Stopwatch StopTimer = Stopwatch.StartNew();
+			int CallCount = 0;
 			while (true)
 			{
 				try
 				{
+					++CallCount;
 					// Do while RpcCall() returns true
 					if (!RpcCall())
 					{
 						return true; // Return true to caller to signal we are done
 					}
 				}
-				catch (TestException)
+				catch (Exception Err)
 				{
 					if (!bSuppressExceptions)
 					{
+						Log.Info($"{Err.GetType} encountered during call {CallCount} of {FuncName}");
 						throw;
 					}
+					else
+					{
+						Log.Verbose($"Suppressed Exception during call {CallCount} of {FuncName}: {Err}");
+					}
 				}
 
 				if (StopTimer.ElapsedMilliseconds > TimeoutInMilliseconds)
 				{
 					OnBeforeTimeout?.Invoke();
-					string TimeoutErrorMessage = $"{FuncName} timed out after {TimeoutInMilliseconds}ms";
+					string TimeoutErrorMessage = $"Hang detected during {FuncName}: call never succeeded after {CallCount} attempts in {TimeoutInMilliseconds}ms";
+
 					if (!bSuppressTimeoutException)
 					{
 						throw new TestException(TimeoutErrorMessage);
@@ -490,12 +498,12 @@ namespace Gauntlet
 			if (DateTime.UtcNow >= TimeoutEnd)
 			{
 				string NamesMissing = string.Join(", ", InTargets.Where(Target => !Target.HasSentMessage).Select(Target => Target.TargetName));
-				string ErrorMsg = string.Format("Timed out waiting for targets to send requested message. {0} targets reported ready before timeout. Missing targets: {1}", TotalMatches, NamesMissing);
+				string ErrorMsg = $"Hang detected waiting for targets to send requested '{MessageCategory}' message. {TotalMatches} targets reported ready within {TimeoutInSeconds}sec. Missing targets: {NamesMissing}";
 				Log.Error(ErrorMsg);
 				throw new TestException(ErrorMsg);
 			}
 
-			Log.Info("Message received from all RPC Targets!");
+			Log.Info($"Message '{MessageCategory}' received from all RPC Targets!");
 		}
 
 		#region Target Setup and Ready
