Disabling general matching of zen cli output in Horde but adding a specific matcher to override generic "error: " matcher in a attempt to make Horde not categorize our recovered errors as errors when uploading or downloading artifacts

diff --git a/Engine/Source/Programs/AutomationTool.Tests/LogParserTests.cs b/Engine/Source/Programs/AutomationTool.Tests/LogParserTests.cs
--- a/Engine/Source/Programs/AutomationTool.Tests/LogParserTests.cs
+++ b/Engine/Source/Programs/AutomationTool.Tests/LogParserTests.cs
@@ -1362,14 +1362,21 @@ namespace AutomationTool.Tests
 				@"    Failed putting build part: JupiterSession::PutBuildBlob: HTTP error 504 Gateway Timeout ({""error_code"": 504}) (504)",
 				@"    Failed putting build part: JupiterSession::PutBuildBlob: HTTP error 504 Gateway Timeout ({""error_code"": 504}) (504)",
 				@"    Failed putting build part: JupiterSession::PutBuildBlob: HTTP error 504 Gateway Timeout ({""error_code"": 504}) (504)",
-				@"    Failed putting build part: JupiterSession::PutBuildBlob: HTTP error 504 Gateway Timeout ({""error_code"": 504}) (504)"
+				@"    Failed putting build part: JupiterSession::PutBuildBlob: HTTP error 504 Gateway Timeout ({""error_code"": 504}) (504)",
+				@"Error: Operation failed due to a http error: Failed putting build part: JupiterSession::PutBuildBlob: HTTP error 504 Gateway Timeout ({""error_code"": 504}) (504)"
 			};
 
 			List<LogEvent> logEvents = Parse(lines);
-			Assert.AreEqual(7, logEvents.Count);
+			// generic error level matching is currently disabled
+			//Assert.AreEqual(8, logEvents.Count);
+			Assert.AreEqual(1, logEvents.Count);
 
-			// the last event is an error on line 15
-			CheckEventGroup(logEvents.Slice(logEvents.Count - 1, 1), 14, 1, LogLevel.Error, KnownLogEvents.ZenCli_Event);
+			// the last error event is on line 15
+			// generic error level matching is currently disabled
+			// CheckEventGroup(logEvents.Slice(logEvents.Count - 2, 1), 14, 1, LogLevel.Error, KnownLogEvents.ZenCli_Event);
+			
+			// the last event is an information on line 19
+			CheckEventGroup(logEvents.Slice(logEvents.Count - 1, 1), 19, 1, LogLevel.Warning, KnownLogEvents.ZenCli_Event);
 		}
 
 		static List<LogEvent> Parse(IEnumerable<string> lines)
diff --git a/Engine/Source/Programs/AutomationTool/AutomationUtils/Matchers/ZenCliEventMatcher.cs b/Engine/Source/Programs/AutomationTool/AutomationUtils/Matchers/ZenCliEventMatcher.cs
--- a/Engine/Source/Programs/AutomationTool/AutomationUtils/Matchers/ZenCliEventMatcher.cs
+++ b/Engine/Source/Programs/AutomationTool/AutomationUtils/Matchers/ZenCliEventMatcher.cs
@@ -16,9 +16,13 @@ namespace AutomationUtils.Matchers
 		static readonly Regex s_pattern = new Regex(
 			@"\[.*\]\s?\[(?<severity>warning|error|info)\]\s?(\[(?<file>\S+)\:(?<line>\d+)\])?(?<message>.*)");
 
+		static readonly Regex s_HttpErrorPattern = new Regex(
+			@"Error: Operation failed due to a http error: (?<message>.*)\((?<http_error>\d+)\)");
 		public LogEventMatch? Match(ILogCursor input)
 		{
-			if (input.TryMatch(s_pattern, out Match? match))
+			// Disabling log matching as we want to exclusivley rely on exit code to determine if the zen cli invoke failed.
+			// Ideally we could annotate these lines without impact process success state in Horde
+			/*if (input.TryMatch(s_pattern, out Match? match))
 			{
 				LogEventBuilder builder = new LogEventBuilder(input);
 				builder.Annotate(match.Groups["severity"], LogEventMarkup.Severity);
@@ -38,8 +42,22 @@ namespace AutomationUtils.Matchers
 				};
 
 				return builder.ToMatch(LogEventPriority.High, level, KnownLogEvents.ZenCli_Event);
-				
+			}*/
+
+			// Error: Operation failed due to a http error: Failed putting build part: JupiterSession::PutBuildBlob: HTTP error 504 Gateway Timeout ({"error_code": 504}) (504)
+			if (input.TryMatch(s_HttpErrorPattern, out Match? httpMatch))
+			{
+				LogEventBuilder builder = new LogEventBuilder(input);
+				string httpErrorString = httpMatch.Groups["http_error"].Value;
+				int httpError = int.Parse(httpErrorString);
+
+				// http error codes that we will retry
+				if (httpError is 502 or 504)
+				{
+					return builder.ToMatch(LogEventPriority.High, LogLevel.Warning, KnownLogEvents.ZenCli_Event);
+				}
 			}
+
 			return null;
 		}
 	}
