Fixed compiler warning in zen cli matcher for Horde solutions

diff --git a/Engine/Source/Programs/AutomationTool/AutomationUtils/Matchers/ZenCliEventMatcher.cs b/Engine/Source/Programs/AutomationTool/AutomationUtils/Matchers/ZenCliEventMatcher.cs
--- a/Engine/Source/Programs/AutomationTool/AutomationUtils/Matchers/ZenCliEventMatcher.cs
+++ b/Engine/Source/Programs/AutomationTool/AutomationUtils/Matchers/ZenCliEventMatcher.cs
@@ -48,11 +48,10 @@ namespace AutomationUtils.Matchers
 			if (input.TryMatch(s_HttpErrorPattern, out Match? httpMatch))
 			{
 				LogEventBuilder builder = new LogEventBuilder(input);
-				string httpErrorString = httpMatch.Groups["http_error"].Value;
-				int httpError = int.Parse(httpErrorString);
+				string httpError = httpMatch.Groups["http_error"].Value;
 
 				// http error codes that we will retry
-				if (httpError is 502 or 504)
+				if (httpError is "502" or "504")
 				{
 					return builder.ToMatch(LogEventPriority.High, LogLevel.Warning, KnownLogEvents.ZenCli_Event);
 				}
