Fixing ZenCliEventMatcher to catch a wider range of errors so we can downgrade them to warnings to prevent them failing job steps.

diff --git a/Engine/Source/Programs/AutomationTool.Tests/LogParserTests.cs b/Engine/Source/Programs/AutomationTool.Tests/LogParserTests.cs
--- a/Engine/Source/Programs/AutomationTool.Tests/LogParserTests.cs
+++ b/Engine/Source/Programs/AutomationTool.Tests/LogParserTests.cs
@@ -1369,7 +1369,7 @@ namespace AutomationTool.Tests
 			List<LogEvent> logEvents = Parse(lines);
 			// generic error level matching is currently disabled
 			//Assert.AreEqual(8, logEvents.Count);
-			Assert.AreEqual(1, logEvents.Count);
+			Assert.AreEqual(6, logEvents.Count);
 
 			// the last error event is on line 15
 			// generic error level matching is currently disabled
diff --git a/Engine/Source/Programs/AutomationTool/AutomationUtils/Matchers/ZenCliEventMatcher.cs b/Engine/Source/Programs/AutomationTool/AutomationUtils/Matchers/ZenCliEventMatcher.cs
--- a/Engine/Source/Programs/AutomationTool/AutomationUtils/Matchers/ZenCliEventMatcher.cs
+++ b/Engine/Source/Programs/AutomationTool/AutomationUtils/Matchers/ZenCliEventMatcher.cs
@@ -13,39 +13,18 @@ namespace AutomationUtils.Matchers
 	/// </summary>
 	class ZenCliEventMatcher : ILogEventMatcher
 	{
-		static readonly Regex s_pattern = new Regex(
-			@"\[.*\]\s?\[(?<severity>warning|error|info)\]\s?(\[(?<file>\S+)\:(?<line>\d+)\])?(?<message>.*)");
-
-		static readonly Regex s_HttpErrorPattern = new Regex(
+		static readonly Regex s_HttpOperationErrorPattern = new Regex(
 			@"Error: Operation failed due to a http error: (?<message>.*)\((?<http_error>\d+)\)");
-		public LogEventMatch? Match(ILogCursor input)
-		{
-			// Disabling log matching as we want to exclusivley rely on exit code to determine if the zen cli invoke failed.
-			// Ideally we could annotate these lines without impact process success state in Horde
-			/*if (input.TryMatch(s_pattern, out Match? match))
-			{
-				LogEventBuilder builder = new LogEventBuilder(input);
-				builder.Annotate(match.Groups["severity"], LogEventMarkup.Severity);
-				string file = match.Groups["file"].Value;
-				string line = match.Groups["line"].Value;
-				if (!string.IsNullOrEmpty(file) && !string.IsNullOrEmpty(line))
-				{
-					builder.AnnotateSourceFile(match.Groups["file"], "");
-					builder.Annotate(match.Groups["line"], LogEventMarkup.LineNumber);
-				}
-
-				LogLevel level = match.Groups["severity"].Value switch
-				{
-					"error" => LogLevel.Warning,
-					"warning" => LogLevel.Information, // we do not want to annotate warnings from zencli as warnings as that causes too much noise - handled retries are warnings but failed retries are already errors
-					_ => LogLevel.Information,
-				};
 
-				return builder.ToMatch(LogEventPriority.High, level, KnownLogEvents.ZenCli_Event);
-			}*/
+		static readonly Regex s_HttpRequestErrorPattern = new Regex(
+			@"HttpClient request failed.*?""error_code"": (?<http_error>\d+)");
 
+		public LogEventMatch? Match(ILogCursor input)
+		{
 			// Error: Operation failed due to a http error: Failed putting build part: JupiterSession::PutBuildBlob: HTTP error 504 Gateway Timeout ({"error_code": 504}) (504)
-			if (input.TryMatch(s_HttpErrorPattern, out Match? httpMatch))
+			// HttpClient request failed (session: guid): Url: url, Status: 504, Error: '' (0), Bytes: 0/19 (Up/Down), Elapsed: 1m00s, Response: '{\"error_code\": 504}', Reason: 'Gateway Time-out'
+			Match? httpMatch;
+			if (input.TryMatch(s_HttpOperationErrorPattern, out httpMatch) || input.TryMatch(s_HttpRequestErrorPattern, out httpMatch))
 			{
 				LogEventBuilder builder = new LogEventBuilder(input);
 				string httpError = httpMatch.Groups["http_error"].Value;
