Cloud Artifacts - Make retried zen events info not warning and explictly warn if the high level retries are running to often.
Also fixed issue with creating zen manifest if the directory to upload didn't exist, which hid the more descriptive error message.

diff --git a/Engine/Source/Programs/AutomationTool/AutomationUtils/Matchers/ZenCliEventMatcher.cs b/Engine/Source/Programs/AutomationTool/AutomationUtils/Matchers/ZenCliEventMatcher.cs
--- a/Engine/Source/Programs/AutomationTool/AutomationUtils/Matchers/ZenCliEventMatcher.cs
+++ b/Engine/Source/Programs/AutomationTool/AutomationUtils/Matchers/ZenCliEventMatcher.cs
@@ -32,7 +32,7 @@ namespace AutomationUtils.Matchers
 				// http error codes that we will retry
 				if (httpError is "502" or "504")
 				{
-					return builder.ToMatch(LogEventPriority.High, LogLevel.Warning, KnownLogEvents.ZenCli_Event);
+					return builder.ToMatch(LogEventPriority.High, LogLevel.Information, KnownLogEvents.ZenCli_Event);
 				}
 			}
 
diff --git a/Engine/Source/Programs/AutomationTool/BuildGraph/Tasks/CreateCloudArtifactTask.cs b/Engine/Source/Programs/AutomationTool/BuildGraph/Tasks/CreateCloudArtifactTask.cs
--- a/Engine/Source/Programs/AutomationTool/BuildGraph/Tasks/CreateCloudArtifactTask.cs
+++ b/Engine/Source/Programs/AutomationTool/BuildGraph/Tasks/CreateCloudArtifactTask.cs
@@ -464,6 +464,11 @@ namespace AutomationTool.Tasks
 			DirectoryReference baseDir = ResolveDirectory(_parameters.BaseDir);
 			List<FileReference> files = BgTaskImpl.ResolveFilespec(baseDir, _parameters.Files, tagNameToFileSet).ToList();
 
+			if (!Directory.Exists(baseDir.FullName))
+			{
+				Directory.CreateDirectory(baseDir.FullName);
+			}
+
 			FileReference manifestFile = FileReference.Combine(baseDir, "zen-manifest.txt");
 			{
 				await using StreamWriter streamWriter = new StreamWriter(manifestFile.FullName);
@@ -506,6 +511,7 @@ namespace AutomationTool.Tasks
 			string partName = "default";
 
 			const int MaxAttempts = 3;
+			const int WarnAtAttempts = MaxAttempts - 1;
 			for (int countOfRetries = 0; countOfRetries < MaxAttempts; countOfRetries++)
 			{
 				// generate new build and part ids for each attempt so that we do not continue to append to a broken build in case of a retry
@@ -560,6 +566,10 @@ namespace AutomationTool.Tasks
 						}
 						else
 						{
+							if (countOfRetries + 1 == WarnAtAttempts)
+							{
+								Logger.LogWarning("Zen returned http status code {ExitCode} after {RetryAttempts} retries. This is a indication of an infrastructure issue.", exitCode, countOfRetries);
+							}
 							Thread.Sleep(TimeSpan.FromSeconds(waitDuration));
 							continue; // retry
 						}
