RetrieveArtifactTask: Add option to control if an error is thrown when an artifact cannot be found

Add the option ErrorIfNotFound to control if an error is thrown when an artifact cannot be found. Defaults to True for backwards compatibility with existing build graphs.

#pr 13056

diff --git a/Engine/Source/Programs/AutomationTool/BuildGraph/Tasks/RetrieveArtifactTask.cs b/Engine/Source/Programs/AutomationTool/BuildGraph/Tasks/RetrieveArtifactTask.cs
--- a/Engine/Source/Programs/AutomationTool/BuildGraph/Tasks/RetrieveArtifactTask.cs
+++ b/Engine/Source/Programs/AutomationTool/BuildGraph/Tasks/RetrieveArtifactTask.cs
@@ -67,6 +67,13 @@ namespace AutomationTool.Tasks
 		[TaskParameter]
 		public string? OutputDir { get; set; }
 
+		/// <summary>
+		/// Whether or not to throw an error if the artifact was not found
+		/// </summary>
+		/// <remarks>For backwards compatibility defaults to true</remarks>
+		[TaskParameter(Optional = true)]
+		public bool ErrorIfNotFound { get; set; } = true;
+
 		/// <summary>
 		/// Append the environment variable UE_HORDE_JOBID to the artifact query
 		/// </summary>
@@ -195,7 +202,15 @@ namespace AutomationTool.Tasks
 			IArtifact? artifact = await hordeClient.Artifacts.FindAsync(streamId, minCommitId, maxCommitId, name, type, keys, 1).FirstOrDefaultAsync();
 			if (artifact == null)
 			{
-				throw new AutomationException("Unable to find any artifact matching given criteria.");
+				if (_parameters.ErrorIfNotFound)
+				{
+					throw new AutomationException("Unable to find any artifact matching given criteria.");
+				}
+				else
+				{
+					Logger.LogInformation("Unable to find any artifact matching given criteria.");
+				}
+				return;
 			}
 
 			Logger.LogInformation("Found artifact {ArtifactId}", artifact.Id);
