Update the process log file path when written to file/copy
to still reference a valid log file path

#rnx

diff --git a/Engine/Source/Programs/AutomationTool/Gauntlet/Framework/Base/Gauntlet.TargetDeviceDesktopCommon.cs b/Engine/Source/Programs/AutomationTool/Gauntlet/Framework/Base/Gauntlet.TargetDeviceDesktopCommon.cs
--- a/Engine/Source/Programs/AutomationTool/Gauntlet/Framework/Base/Gauntlet.TargetDeviceDesktopCommon.cs
+++ b/Engine/Source/Programs/AutomationTool/Gauntlet/Framework/Base/Gauntlet.TargetDeviceDesktopCommon.cs
@@ -523,6 +523,8 @@ namespace Gauntlet
 				return new LogFileReader(ProcessLogFile);
 			}
 
+			Log.Warning(KnownLogEvents.Horde_BuildHealth_Ignore, "Log file '{filepath}' is missing. The buffer will be used instead but will most likely lack the beginning of the file.", ProcessLogFile);
+
 			return ProcessLogOutput.GetReader();
 		}
 
@@ -547,10 +549,16 @@ namespace Gauntlet
 			{
 				ProcessUtils.CheckProcessLogReachedSizeLimit(new FileReference(ProcessLogFile));
 				File.Copy(ProcessLogFile, FilePath, true);
+				if (HasExited)
+				{
+					// Follow the new location
+					// SaveRoleArtifacts may clean device cache folder prior parsing the log file
+					ProcessLogFile = FilePath;
+				}
 			}
 			else
 			{
-				Log.Warning("Log file '{filepath}' is missing at the time of making a copy. The buffer will be used instead but will most likely lack the beginning of the file.", ProcessLogFile);
+				Log.Warning(KnownLogEvents.Horde_BuildHealth_Ignore, "Log file '{filepath}' is missing at the time of making a copy. The buffer will be used instead but will most likely lack the beginning of the file.", ProcessLogFile);
 				StreamWriter Writer = ProcessUtils.CreateWriterForProcessLog(FilePath, CommandLine);
 				foreach(string Line in ProcessLogOutput)
 				{
