Fix missing Phases report on failure and interruption

#rnx

diff --git a/Engine/Source/Programs/AutomationTool/Gauntlet/Unreal/Base/Gauntlet.UnrealTestNode.cs b/Engine/Source/Programs/AutomationTool/Gauntlet/Unreal/Base/Gauntlet.UnrealTestNode.cs
--- a/Engine/Source/Programs/AutomationTool/Gauntlet/Unreal/Base/Gauntlet.UnrealTestNode.cs
+++ b/Engine/Source/Programs/AutomationTool/Gauntlet/Unreal/Base/Gauntlet.UnrealTestNode.cs
@@ -159,9 +159,9 @@ namespace Gauntlet
         /// <param name="Args"></param>
         public virtual void ReportInterrupt(string Message, params object[] Args)
         {
-            Message = string.Format(Message, Args);
-            AddTestEvent(new UnrealTestEvent(EventSeverity.Fatal, Message, Enumerable.Empty<string>()));
-            if (!LogWarningsAndErrorsAfterSummary) { Log.Critical(KnownLogEvents.Gauntlet_FatalEvent, Message); }
+			Phases.Current?.AddEvent(Logging.LogLevel.Critical, Message, Args);
+            AddTestEvent(new UnrealTestEvent(EventSeverity.Fatal, Message, Args));
+            if (!LogWarningsAndErrorsAfterSummary) { Log.Critical(KnownLogEvents.Gauntlet_FatalEvent, Message, Args); }
             if (GetTestStatus() == TestStatus.Complete && GetTestResult() == TestResult.Passed)
             {
                 SetUnrealTestResult(TestResult.OperationalException);
@@ -175,9 +175,9 @@ namespace Gauntlet
         /// <param name="Args"></param>
         public virtual void ReportFailure(string Message, params object[] Args)
         {
-            Message = string.Format(Message, Args);
-            AddTestEvent(new UnrealTestEvent(EventSeverity.Error, Message, Enumerable.Empty<string>()));
-            if (!LogWarningsAndErrorsAfterSummary) { Log.Error(KnownLogEvents.Gauntlet_TestEvent, Message); }
+			Phases.Current?.AddEvent(Logging.LogLevel.Error, Message, Args);
+            AddTestEvent(new UnrealTestEvent(EventSeverity.Error, Message, Args));
+            if (!LogWarningsAndErrorsAfterSummary) { Log.Error(KnownLogEvents.Gauntlet_TestEvent, Message, Args); }
             if (GetTestStatus() == TestStatus.Complete && GetTestResult() == TestResult.Passed)
             {
                 SetUnrealTestResult(TestResult.Failed);
