[Gauntlet] Adding WdaLaunchTimeout to appium config

diff --git a/Engine/Source/Programs/AutomationTool/Gauntlet/Platform/IOS/Gauntlet.IOSAppium.cs b/Engine/Source/Programs/AutomationTool/Gauntlet/Platform/IOS/Gauntlet.IOSAppium.cs
--- a/Engine/Source/Programs/AutomationTool/Gauntlet/Platform/IOS/Gauntlet.IOSAppium.cs
+++ b/Engine/Source/Programs/AutomationTool/Gauntlet/Platform/IOS/Gauntlet.IOSAppium.cs
@@ -64,6 +64,12 @@ namespace Gauntlet
 			/// </summary>
 			public string WdaBundleId { get; set; }
 
+
+			/// <summary>
+			/// Launch Timeout for WebDriverAgent app in ms. Default: 60000ms
+			/// </summary>
+			public int WdaLaunchTimeoutMs { get; set; } = 60000;
+
 			/// <summary>
 			/// Optional - Allow you to override the location of the appium executable.
 			/// Useful if you opt not to install appium to the global npm root
@@ -83,8 +89,9 @@ namespace Gauntlet
                 Capabilities.AddAdditionalAppiumOption("autoLaunch", false);
                 Capabilities.AddAdditionalAppiumOption("usePreinstalledWDA", true);
                 Capabilities.AddAdditionalAppiumOption("updatedWDABundleId", WdaBundleId);
+				Capabilities.AddAdditionalAppiumOption("wdaLaunchTimeout", WdaLaunchTimeoutMs);
 
-                if (Log.Level == LogLevel.Verbose)
+				if (Log.Level == LogLevel.Verbose)
                 {
                     Capabilities.AddAdditionalAppiumOption("showXcodeLog", true);
                 }
