RPC message queue redirection for individual targets

diff --git a/Engine/Source/Programs/AutomationTool/Gauntlet/Framework/RpcFramework/Gauntlet.RpcExecutor.cs b/Engine/Source/Programs/AutomationTool/Gauntlet/Framework/RpcFramework/Gauntlet.RpcExecutor.cs
--- a/Engine/Source/Programs/AutomationTool/Gauntlet/Framework/RpcFramework/Gauntlet.RpcExecutor.cs
+++ b/Engine/Source/Programs/AutomationTool/Gauntlet/Framework/RpcFramework/Gauntlet.RpcExecutor.cs
@@ -1,7 +1,9 @@
 // Copyright Epic Games, Inc. All Rights Reserved.
 
+using EpicGames.Core;
 using Newtonsoft.Json;
 using System;
+using System.Collections.Concurrent;
 using System.Collections.Generic;
 using System.Diagnostics;
 using System.IO;
@@ -9,10 +11,10 @@ using System.Linq;
 using System.Net;
 using System.Net.Http;
 using System.Net.Sockets;
+using System.Runtime.CompilerServices;
 using System.Text;
 using System.Text.RegularExpressions;
 using System.Threading;
-using EpicGames.Core;
 using ThreadState = System.Threading.ThreadState;
 
 namespace Gauntlet
@@ -111,7 +113,7 @@ namespace Gauntlet
 		protected HttpListener RequestListener;
 		protected Thread HttpListenThread;
 		protected static int GauntletRPCListenPort = 23232;
-
+		protected RpcTargetRegistry ActiveTargetRegistry = null;
 
 		public static RpcExecutor Instance
 		{
@@ -135,6 +137,7 @@ namespace Gauntlet
 		}
 		~RpcExecutor()
 		{
+			ActiveTargetRegistry = null;
 			if (HttpListenThread != null && HttpListenThread.ThreadState == ThreadState.Running)
 			{
 				bHasReceivedTerminationCall = true;
@@ -142,6 +145,19 @@ namespace Gauntlet
 			}
 		}
 
+		/// <summary>
+		/// Enables a message queue for each target application in the current test, reducing async lock contention
+		/// </summary>
+		/// <param name="CurrentRegistry">RpcTargetRegistry for the currently running test</param>
+		public void SetRpcTargetRegistry(RpcTargetRegistry CurrentRegistry, [CallerFilePath] string Caller = "", [CallerLineNumber] int Line = 0)
+		{
+			if (ActiveTargetRegistry != null && ActiveTargetRegistry != CurrentRegistry)
+			{
+				Log.Warning($"{nameof(SetRpcTargetRegistry)} unexpectedly invoked from {Caller} line {Line} while another registry was already set. This may disable RPC communication in another running test.");
+			}
+			ActiveTargetRegistry = CurrentRegistry;
+		}
+
 		public static string GetListenAddressAndPort()
 		{
 			string LocalIP = "127.0.0.1";
@@ -157,12 +173,11 @@ namespace Gauntlet
 			return String.Format("{0}:{1}", LocalIP, GauntletRPCListenPort);
 		}
 
-		public Queue<GauntletIncomingRpcMessage> IncomingMessageQueue = new Queue<GauntletIncomingRpcMessage>();
-		protected Dictionary<GauntletRpcDelegateRegistryFilter, GauntletIncomingRpcMessageDelegate> MessageDelegates = new Dictionary<GauntletRpcDelegateRegistryFilter, GauntletIncomingRpcMessageDelegate>();
-		protected Dictionary<RouteMapping, GauntletHttpRouteDelegate> HttpRoutes = new Dictionary<RouteMapping, GauntletHttpRouteDelegate>();
-		protected object MessageQueueLock = new object();
+		public ConcurrentQueue<GauntletIncomingRpcMessage> IncomingMessageQueue = [];
+		protected ConcurrentDictionary<GauntletRpcDelegateRegistryFilter, GauntletIncomingRpcMessageDelegate> MessageDelegates = [];
+		protected ConcurrentDictionary<RouteMapping, GauntletHttpRouteDelegate> HttpRoutes = [];
 
-		protected object ReceivedTerminationCallLock = new object();
+		protected object ReceivedTerminationCallLock = new ();
 		private bool _bHasReceivedTerminationCall = false;
 		public bool bHasReceivedTerminationCall
 		{
@@ -218,18 +233,30 @@ namespace Gauntlet
 			{
 				FireRelevantMessageDelegates(IncomingMessage);
 
-				lock (MessageQueueLock)
+				if (null == ActiveTargetRegistry)
 				{
 					IncomingMessageQueue.Enqueue(IncomingMessage);
-					CallResponse.ResponseBody = String.Format("New message enqueued! Source: {0} \n Category: {1} \n Payload: {2} \n", IncomingMessage.SenderId, IncomingMessage.Category, IncomingMessage.Payload);
-					Log.Verbose(CallResponse.ResponseBody);
 				}
+				else
+				{
+					RpcTarget Target = ActiveTargetRegistry.GetTarget(IncomingMessage.SenderId);
+					if( null == Target )
+					{
+						Log.Warning($"Discarded incoming message with category '{IncomingMessage.Category}' : No RPC target registered for {IncomingMessage.SenderId}");
+					}
+					else
+					{
+						Target.IncomingMessageQueue.Enqueue(IncomingMessage);
+					}
+				}
+				CallResponse.ResponseBody = String.Format("New message enqueued! Source: {0} \n Category: {1} \n Payload: {2} \n", IncomingMessage.SenderId, IncomingMessage.Category, IncomingMessage.Payload);
 			}
 			else
 			{
 				CallResponse.ResponseBody = "Body could not be processed. Contents: " + MessageBody;
 				CallResponse.ResponseCode = (int)HttpStatusCode.UnprocessableEntity;
 			}
+			Log.Verbose(CallResponse.ResponseBody);
 		}
 
 		public void ReceiveListenRequest(IAsyncResult ListenResult)
@@ -334,34 +361,90 @@ namespace Gauntlet
 		}
 		public List<GauntletIncomingRpcMessage> GetNewMessages(bool bClearMessageQueue = true)
 		{
-			List<GauntletIncomingRpcMessage> NewMessages = new List<GauntletIncomingRpcMessage>();
-			if (IncomingMessageQueue.Count == 0)
+			List<GauntletIncomingRpcMessage> NewMessages = [];
+			int TargetsMessageCount = 0;
+			if(null != ActiveTargetRegistry)
+			{
+				foreach (RpcTarget Target in ActiveTargetRegistry.RpcTargets.Values)
+				{
+					TargetsMessageCount += Target.IncomingMessageQueue.Count;
+				}
+			}
+			if (IncomingMessageQueue.Count == 0 && TargetsMessageCount == 0)
 			{
 				return NewMessages;
 			}
 
-			// We need to lock the original IncomingMessageQueue to prevent modifications to the collection that would throw an exception during iteration in the foreach loop
-			lock (MessageQueueLock)
+			PopulateNewMessagesFromCentralQueue(NewMessages, bClearMessageQueue);
+			PopulateNewMessagesFromActiveTargetQueues(NewMessages, bClearMessageQueue);
+			return NewMessages;
+		}
+
+		private void PopulateNewMessagesFromQueue(ConcurrentQueue<GauntletIncomingRpcMessage> Queue, List<GauntletIncomingRpcMessage> Messages, bool bClearMessageQueue)
+		{
+			int InitialCount = Queue.Count;
+			if (InitialCount > 0)
 			{
-				foreach (GauntletIncomingRpcMessage NewMessage in IncomingMessageQueue)
+				if (bClearMessageQueue)
 				{
-					NewMessages.Add(new GauntletIncomingRpcMessage(NewMessage));
+					for(int MessageIndex = 0; MessageIndex < InitialCount; ++MessageIndex)
+					{
+						if(Queue.TryDequeue(out GauntletIncomingRpcMessage NewMessage))
+						{
+							Messages.Add(new GauntletIncomingRpcMessage(NewMessage));
+						}
+						else
+						{
+							Log.Warning("Message queue unexpectedly dequeued from another thread. RPC messages may be missing.");
+						}
+					}
+				}
+				else
+				{
+					foreach (GauntletIncomingRpcMessage NewMessage in Queue)
+					{
+						Messages.Add(new GauntletIncomingRpcMessage(NewMessage));
+					}
 				}
 			}
+		}
 
-			if (bClearMessageQueue)
+		private void PopulateNewMessagesFromCentralQueue(List<GauntletIncomingRpcMessage> Messages, bool bClearMessageQueue = true)
+		{
+			PopulateNewMessagesFromQueue(IncomingMessageQueue, Messages, bClearMessageQueue);
+		}
+
+		private void PopulateNewMessagesFromActiveTargetQueues(List<GauntletIncomingRpcMessage> Messages, bool bClearMessageQueue = true)
+		{
+			if (null != ActiveTargetRegistry)
 			{
-				ClearIncomingMessageQueue();
+				foreach (RpcTarget Target in ActiveTargetRegistry.RpcTargets.Values)
+				{
+					PopulateNewMessagesFromQueue(Target.IncomingMessageQueue, Messages, bClearMessageQueue);
+				}
 			}
+		}
 
-			return NewMessages;
+		/// <summary>
+		/// Dequeues a number messages, allowing new messages to be concurrently enqueued
+		/// </summary>
+		private void SafeClearMessages(ConcurrentQueue<GauntletIncomingRpcMessage> Queue, int MessagesToClear)
+		{
+			for (int MessageIndex = 0; MessageIndex < MessagesToClear; ++MessageIndex)
+			{
+				Queue.TryDequeue(out GauntletIncomingRpcMessage Unused);
+			}
 		}
 
 		public void ClearIncomingMessageQueue()
 		{
-			lock (MessageQueueLock)
+			SafeClearMessages(IncomingMessageQueue, IncomingMessageQueue.Count);
+			if (null != ActiveTargetRegistry)
 			{
-				IncomingMessageQueue.Clear();
+				foreach (RpcTarget Target in ActiveTargetRegistry.RpcTargets.Values)
+				{
+					SafeClearMessages(Target.IncomingMessageQueue, Target.IncomingMessageQueue.Count);
+				}
 			}
 		}
 
@@ -408,68 +491,95 @@ namespace Gauntlet
 		}
 
 		/// <summary>
-		/// Checks the IncomingMessages queue for messages matching a list of GauntletRpcDelegateRegistryFilters, which represent 
-		/// all of the messages that we are looking for. 
+		/// Process a message queue and return how many messages match a filter, optionally retaining any unmatched messages. <see cref="GauntletRpcDelegateRegistryFilter"/>
 		/// </summary>
-		/// <param name="MessageFilters">List of GauntletRpcDelegateRegistryFilters which hold the message properties we are checking
-		/// for. DesiredSenderId, DesiredCategory, and DesiredPayload can be null and will act as wild cards allowing you to match on 
-		/// specific messages or groups of messages.</param>
-		/// <param name="bSaveUnmatchedMessages">If true, messages that are in the IncomingMessages queue but do not match our filter 
-		/// will be preserved and can be if true. If false, then all new messages retrieved are purged regardless of if they were 
-		/// matched or not. Matched messages are always removed.</param>
-		/// <returns>The number of matched messages found.</returns>
-		public int CheckIncomingMessages(List<GauntletRpcDelegateRegistryFilter> MessageFilters, bool bSaveUnmatchedMessages)
+		/// <param name="MessageQueue">Queue to search for matches</param>
+		/// <param name="QueueLock">Lock object guarding the queue</param>
+		/// <param name="MessageFilters">List of Filters that match message properties</param>
+		/// <param name="bSaveUnmatchedMessages">Preserve entries in MessageQueue which are not matched. When false all entries are purged from the queue.</param>
+		/// <returns>Count of matched messages</returns>
+		private int ConsumeIncomingMessagesFromQueue(ConcurrentQueue<GauntletIncomingRpcMessage> MessageQueue, List<GauntletRpcDelegateRegistryFilter> MessageFilters, bool bSaveUnmatchedMessages)
 		{
-			if (MessageFilters == null)
-			{
-				throw new ArgumentNullException(nameof(MessageFilters));
-			}
-
-			int TotalTargets = MessageFilters.Count;
-
-			if (TotalTargets == 0)
+			ArgumentNullException.ThrowIfNull(MessageQueue);
+			ArgumentNullException.ThrowIfNull(MessageFilters);
+			if (MessageFilters.Count == 0)
 			{
-				Log.Warning("Requested Incoming Messages to be checked, but no filters were provided.");
+				Log.Warning("Requested Incoming Messages to be matched, but no filters were provided to match against.");
 				return 0;
 			}
 
-			Log.Info("Checking for new Messages from {0} Targets. Category: {1}, Contents: {2}", TotalTargets, MessageFilters[0].DesiredCategory, MessageFilters[0].DesiredPayload);
-
-			// Storage so messages that aren't matched can be re-queued in case they are being awaited elsewhere.
-			List<GauntletIncomingRpcMessage> UnmatchedMessages = new List<GauntletIncomingRpcMessage>();
-
-			List<GauntletIncomingRpcMessage> NewMessages = GetNewMessages();
-
 			int MatchesFound = 0;
-
+			List<GauntletIncomingRpcMessage> UnmatchedMessages = [];
+			List<GauntletIncomingRpcMessage> NewMessages = [];
+			PopulateNewMessagesFromQueue(MessageQueue, NewMessages, bClearMessageQueue: true);
 			foreach (var Message in NewMessages)
 			{
-				int NewMatches = GetMessageFilterMatches(Message, MessageFilters);
-
-				// Capture unmatched messages in case they need to be re-enqueued
-				if (NewMatches == 0)
+				if (GetMessageFilterMatches(Message, MessageFilters) > 0)
+				{
+					++MatchesFound;
+				}
+				else if(bSaveUnmatchedMessages)
 				{
 					UnmatchedMessages.Add(Message);
-					continue;
 				}
-
-				MatchesFound += NewMatches;
 			}
-
-			lock (MessageQueueLock)
+			
+			foreach (var Message in UnmatchedMessages)
 			{
-				IncomingMessageQueue.Clear();
+				MessageQueue.Enqueue(Message);
+			}
+			Log.Verbose("Matched {0} messages from {1} filters.", MatchesFound, MessageFilters.Count);
+			return MatchesFound;
+		}
 
-				if (bSaveUnmatchedMessages)
+		/// <summary>
+		/// Process an RpcTarget's queue and return how many messages match a filter, optionally retaining any unmatched messages. <see cref="GauntletRpcDelegateRegistryFilter"/>
+		/// </summary>
+		/// <param name="Target">Queue to search for matches</param>
+		/// <param name="MessageFilters">List of Filters that match message properties</param>
+		/// <param name="bSaveUnmatchedMessages">Preserve entries in MessageQueue which are not matched. When false all entries are purged from the queue.</param>
+		/// <returns>Count of matched messages</returns>
+		public int ConsumeIncomingMessagesFromTarget(RpcTarget Target, List<GauntletRpcDelegateRegistryFilter> MessageFilters, bool bSaveUnmatchedMessages)
+		{
+			return ConsumeIncomingMessagesFromQueue(Target.IncomingMessageQueue, MessageFilters, bSaveUnmatchedMessages);
+		}
+
+		/// <summary>
+		/// Process all active queues and return how many messages match a filter, optionally retaining any unmatched messages. <see cref="GauntletRpcDelegateRegistryFilter"/>
+		/// </summary>
+		/// <param name="MessageFilters">List of Filters that match message properties</param>
+		/// <param name="bSaveUnmatchedMessages">Preserve entries in MessageQueue which are not matched. When false all entries are purged from the queue.</param>
+		/// <returns>Count of matched messages</returns>
+		public int ConsumeIncomingMessagesFromAllQueues(List<GauntletRpcDelegateRegistryFilter> MessageFilters, bool bSaveUnmatchedMessages)
+		{
+			int MatchesFound = ConsumeIncomingMessagesFromQueue(IncomingMessageQueue, MessageFilters, bSaveUnmatchedMessages);
+
+			if (null != ActiveTargetRegistry)
+			{
+				foreach (RpcTarget Target in ActiveTargetRegistry.RpcTargets.Values)
 				{
-					foreach (var Message in UnmatchedMessages)
-					{
-						IncomingMessageQueue.Enqueue(Message);
-					}
+					MatchesFound += ConsumeIncomingMessagesFromTarget(Target, MessageFilters, bSaveUnmatchedMessages);
 				}
 			}
+			return MatchesFound;
+		}
 
-			Log.Info("Found {0} messages from {1} targets.", MatchesFound, TotalTargets);
+		/// <summary>
+		/// Checks the IncomingMessages queue for messages matching a list of GauntletRpcDelegateRegistryFilters, which represent 
+		/// all of the messages that we are looking for. 
+		/// </summary>
+		/// <param name="MessageFilters">List of GauntletRpcDelegateRegistryFilters which hold the message properties we are checking
+		/// for. DesiredSenderId, DesiredCategory, and DesiredPayload can be null and will act as wild cards allowing you to match on 
+		/// specific messages or groups of messages.</param>
+		/// <param name="bSaveUnmatchedMessages">If true, messages that are in the IncomingMessages queue but do not match our filter 
+		/// will be preserved and can be later matched. If false, then all new messages retrieved are purged regardless of if they were 
+		/// matched or not. Matched messages are always removed.</param>
+		/// <returns>The number of matched messages found.</returns>
+		public int CheckIncomingMessages(List<GauntletRpcDelegateRegistryFilter> MessageFilters, bool bSaveUnmatchedMessages)
+		{
+			Log.Info("Checking for new Messages from {0} Targets. Category: {1}, Contents: {2}", MessageFilters.Count, MessageFilters[0].DesiredCategory, MessageFilters[0].DesiredPayload);
+			int MatchesFound = ConsumeIncomingMessagesFromAllQueues(MessageFilters, bSaveUnmatchedMessages);
+			Log.Info("Found {0} messages from {1} targets.", MatchesFound, MessageFilters.Count);
 			return MatchesFound;
 		}
 
@@ -481,9 +591,10 @@ namespace Gauntlet
 		/// for. DesiredSenderId, DesiredCategory, and DesiredPayload can be null and will act as wild cards allowing you to match on 
 		/// specific messages or groups of messages.</param>
 		/// <param name="bSaveUnmatchedMessages">If true, messages that are in the IncomingMessages queue but do not match our filter 
-		/// will be preserved and can be if true. If false, then all new messages retrieved are purged regardless of if they were 
+		/// will be preserved and can be later matched. If false, then all new messages retrieved are purged regardless of if they were 
 		/// matched or not. Matched messages are always removed.</param>
 		/// <returns>The number of matched messages found.</returns>
+		[Obsolete($"Instead use {nameof(ConsumeIncomingMessagesFromTarget)} on the target message queue(s)", false)]
 		public int CheckIncomingMessages(GauntletRpcDelegateRegistryFilter MessageFilter, bool bSaveUnmatchedMessages)
 		{
 			return CheckIncomingMessages(new List<GauntletRpcDelegateRegistryFilter>(new GauntletRpcDelegateRegistryFilter[] { MessageFilter }), bSaveUnmatchedMessages);
@@ -497,9 +608,10 @@ namespace Gauntlet
 		/// <param name="TargetMessagePayload">All or part of the message payload to be checked. Is case sensitive. Can be null/empty 
 		/// to match for any payload</param>
 		/// <param name="bSaveUnmatchedMessages">If true, messages that are in the IncomingMessages queue but do not match our filter 
-		/// will be preserved and can be if true. If false, then all new messages retrieved are purged regardless of if they were 
+		/// will be preserved and can be later matched. If false, then all new messages retrieved are purged regardless of if they were 
 		/// matched or not. Matched messages are always removed.</param>
-		/// <returns>The number of matched messages found.</returns>		
+		/// <returns>The number of matched messages found.</returns>
+		[Obsolete($"Instead use {nameof(ConsumeIncomingMessagesFromTarget)} on the target message queue(s)", false)]
 		public int CheckIncomingMessages(string TargetSenderId, string TargetMessageCategory, string TargetMessagePayload, bool bSaveUnmatchedMessages)
 		{
 			GauntletRpcDelegateRegistryFilter NewFilter = new GauntletRpcDelegateRegistryFilter(TargetSenderId, TargetMessageCategory, TargetMessagePayload, false);
@@ -513,7 +625,7 @@ namespace Gauntlet
 		/// <param name="TargetMessageCategory">Category of message to be checked.</param>
 		/// <param name="TargetMessagePayload">All or part of the contens of the message payload to be checked. Is case sensitive. </param>
 		/// <param name="bSaveUnmatchedMessages">If true, messages that are in the IncomingMessages queue but do not match our filter 
-		/// will be preserved and can be if true. If false, then all new messages retrieved are purged regardless of if they were 
+		/// will be preserved and can be later matched. If false, then all new messages retrieved are purged regardless of if they were 
 		/// matched or not. Matched messages are always removed.</param>
 		/// <returns>The number of matched messages found.</returns>
 		public int CheckIncomingMessages(List<RpcTarget> InTargets, string TargetMessageCategory, string TargetMessagePayload, bool bSaveUnmatchedMessages)
@@ -536,9 +648,10 @@ namespace Gauntlet
 		/// <param name="TargetMessageCategory">Category of message to be checked.</param>
 		/// <param name="TargetMessagePayload">All or part of the contens of the message payload to be checked. Is case sensitive. </param>
 		/// <param name="bSaveUnmatchedMessages">If true, messages that are in the IncomingMessages queue but do not match our filter 
-		/// will be preserved and can be if true. If false, then all new messages retrieved are purged regardless of if they were 
+		/// will be preserved and can be later matched. If false, then all new messages retrieved are purged regardless of if they were 
 		/// matched or not. Matched messages are always removed.</param>
 		/// <returns>The number of matched messages found.</returns>
+		[Obsolete($"Instead use {nameof(ConsumeIncomingMessagesFromTarget)} on the target message queue(s)", false)]
 		public int CheckIncomingMessages(RpcTarget InTarget, string TargetMessageCategory, string TargetMessagePayload, bool bSaveUnmatchedMessages)
 		{
 			GauntletRpcDelegateRegistryFilter NewFilter = new GauntletRpcDelegateRegistryFilter(InTarget.TargetName, TargetMessageCategory, TargetMessagePayload, false);
@@ -547,6 +660,7 @@ namespace Gauntlet
 
 		/// <summary>
 		/// Basic filter class that allows users to filter what incoming messages they would like to fire their delegate functions on.
+		/// Null values act as wild cards.
 		/// </summary>
 		public class GauntletRpcDelegateRegistryFilter
 		{
@@ -633,8 +747,15 @@ namespace Gauntlet
 				Log.Warning("Cannot register Delegate with no filter at all.");
 				return;
 			}
-			Log.Info("New delegate registered: " + SenderFilter + " | " + CategoryFilter + " | " + PayloadFilter);
-			MessageDelegates.Add(new GauntletRpcDelegateRegistryFilter(SenderFilter, CategoryFilter, PayloadFilter, bDeregisterAfterTrigger), FunctionToRun);
+			
+			if(MessageDelegates.TryAdd(new GauntletRpcDelegateRegistryFilter(SenderFilter, CategoryFilter, PayloadFilter, bDeregisterAfterTrigger), FunctionToRun))
+			{
+				Log.Info("New delegate registered: " + SenderFilter + " | " + CategoryFilter + " | " + PayloadFilter);
+			}
+			else
+			{
+				Log.Error("Unexpectedly failed to register new delegate with duplicate filter");
+			}
 		}
 
 		/// <summary>
@@ -654,9 +775,15 @@ namespace Gauntlet
 				{
 					if (FunctionToRun == MessageDelegates[MessageFilter])
 					{
-						MessageDelegates.Remove(MessageFilter);
-						bDeregisteredDelegate = true;
-						Log.Verbose("Delegate removed for filter: Sender: " + SenderFilter + " | Category: " + CategoryFilter + " | Payload: " + PayloadFilter);
+						if(MessageDelegates.TryRemove(MessageFilter, out var Unused))
+						{
+							bDeregisteredDelegate = true;
+							Log.Verbose("Delegate removed for filter: Sender: " + SenderFilter + " | Category: " + CategoryFilter + " | Payload: " + PayloadFilter);
+						}
+						else
+						{
+							Log.Verbose("Delegate unexpectedly already removed");
+						}
 					}
 				}
 			}
@@ -676,7 +803,7 @@ namespace Gauntlet
 				MessageDelegates[MessageFilter].Invoke(NewMessage);
 				if (MessageFilter.bDeregisterAfterFirstTrigger)
 				{
-					MessageDelegates.Remove(MessageFilter);
+					MessageDelegates.TryRemove(MessageFilter, out var Unused);
 				}
 			}
 		}
@@ -733,16 +860,18 @@ namespace Gauntlet
 			else
 			{
 				Log.Info(string.Format("Registering endpoint {0} | {1} with new delegate.", Path.Path, Path.Verb));
-				HttpRoutes.Add(Path, RouteDelegate);
+				if(!HttpRoutes.TryAdd(Path, RouteDelegate))
+				{
+					Log.Error($"Unexpectedly encountered existing RouteMapping for {Path}");
+				}
 			}
 		}
 
 		public void DeregisterHttpRoute(RouteMapping Path)
 		{
-			if (HttpRoutes.ContainsKey(Path))
+			if (HttpRoutes.ContainsKey(Path) && HttpRoutes.TryRemove(Path, out var Unused))
 			{
-				Log.Info(string.Format("Deregistering endpoint {0}", Path));
-				HttpRoutes.Remove(Path);
+				Log.Info(string.Format("Deregistered endpoint {0}", Path));
 			}
 			else
 			{
diff --git a/Engine/Source/Programs/AutomationTool/Gauntlet/Framework/RpcFramework/Gauntlet.RpcTargetRegistry.cs b/Engine/Source/Programs/AutomationTool/Gauntlet/Framework/RpcFramework/Gauntlet.RpcTargetRegistry.cs
--- a/Engine/Source/Programs/AutomationTool/Gauntlet/Framework/RpcFramework/Gauntlet.RpcTargetRegistry.cs
+++ b/Engine/Source/Programs/AutomationTool/Gauntlet/Framework/RpcFramework/Gauntlet.RpcTargetRegistry.cs
@@ -11,6 +11,7 @@ using System.Net;
 using System.Net.Http;
 using System.Net.NetworkInformation;
 using UnrealBuildTool;
+using static Gauntlet.RpcExecutor;
 
 namespace Gauntlet
 {
@@ -30,7 +31,7 @@ namespace Gauntlet
 		public bool HasSentMessage = false;
 		public bool IsDeferStart = false;
 		protected bool bIsInvalidated = false;
-
+		public ConcurrentQueue<GauntletIncomingRpcMessage> IncomingMessageQueue = new ();
 
 		public RpcTarget(string InIpAddress, int InPort, string InName, UnrealTargetRole InType = UnrealTargetRole.Client)
 		{
