Fixed C# P4 BatchedCommand helper calculating the wrong length and overflowing the 32,699 character limit

It wasn't taking into account the length of the base Command or GlobalOptions, which are also added to the arguments when running p4.exe.

diff --git a/Engine/Source/Programs/AutomationTool/AutomationUtils/P4Utils.cs b/Engine/Source/Programs/AutomationTool/AutomationUtils/P4Utils.cs
--- a/Engine/Source/Programs/AutomationTool/AutomationUtils/P4Utils.cs
+++ b/Engine/Source/Programs/AutomationTool/AutomationUtils/P4Utils.cs
@@ -2168,7 +2168,7 @@ namespace AutomationTool
 		public  void BatchedCommand(string Command, List<string> Files, bool AllowSpew = true)
 		{
 			// using the limit of ProcessStartInfo.Arguments
-			const int MaxCommandLineLength = 32699;
+			int MaxCommandLineLength = 32699 - Command.Length - GlobalOptions.Length - 2;
 
 			StringBuilder CommandLine = new StringBuilder();
 			for (int Idx = 0; Idx < Files.Count; Idx++)
