[Gauntlet] Adds an option for all roles to use the same device cache, useful for some multi-role editor tests.
This feature is opt-in, and will activate when -UniformCache is present on the commandline.

diff --git a/Engine/Source/Programs/AutomationTool/Gauntlet/Framework/Devices/Gauntlet.DevicePool.cs b/Engine/Source/Programs/AutomationTool/Gauntlet/Framework/Devices/Gauntlet.DevicePool.cs
--- a/Engine/Source/Programs/AutomationTool/Gauntlet/Framework/Devices/Gauntlet.DevicePool.cs
+++ b/Engine/Source/Programs/AutomationTool/Gauntlet/Framework/Devices/Gauntlet.DevicePool.cs
@@ -372,6 +372,11 @@ namespace Gauntlet
 		/// </summary>
 		private string LocalTempDir;
 
+		/// <summary>
+		/// Uniform cache for desktop roles. Only used if -UniformCache is enabled
+		/// </summary>
+		private string UniformCache;
+
 		/// <summary>
 		/// Protected constructor - code should use DevicePool.Instance
 		/// </summary>
@@ -1271,12 +1276,22 @@ namespace Gauntlet
 		/// </summary>
 		private string GetCleanCachePath(DeviceDefinition InDeviceDefiniton)
 		{
-			string SanitizedDeviceName = FileUtils.SanitizeFilename(InDeviceDefiniton.Name);
+			if (!string.IsNullOrEmpty(UniformCache))
+			{
+				return UniformCache;
+			}
+
 			// Give the desktop platform a temp folder with its name under the device cache
+			string SanitizedDeviceName = FileUtils.SanitizeFilename(InDeviceDefiniton.Name);
 			string DeviceCache = Path.Combine(LocalTempDir, "DeviceCache");
 			string PlatformCache = Path.Combine(DeviceCache, InDeviceDefiniton.Platform.ToString());
 			string ClientCache = Path.Combine(PlatformCache, SanitizedDeviceName);
 
+			if (Globals.Params.ParseParam("UniformCache"))
+			{
+				UniformCache = ClientCache;
+			}
+
 			// On Desktops, builds are installed in the device cache.
 			// When using device reservation blocks, we don't want to fully clean the cache and lose previously installed builds.
 			// If bRetainBuilds evaluates to true, it means we are in the second step or beyond in a device reservation block.
