[Gauntlet] Ensure uniform cache is set before assigning the value as it can be changed if the cache can't be cleaned

diff --git a/Engine/Source/Programs/AutomationTool/Gauntlet/Framework/Devices/Gauntlet.DevicePool.cs b/Engine/Source/Programs/AutomationTool/Gauntlet/Framework/Devices/Gauntlet.DevicePool.cs
--- a/Engine/Source/Programs/AutomationTool/Gauntlet/Framework/Devices/Gauntlet.DevicePool.cs
+++ b/Engine/Source/Programs/AutomationTool/Gauntlet/Framework/Devices/Gauntlet.DevicePool.cs
@@ -373,7 +373,7 @@ namespace Gauntlet
 		private string LocalTempDir;
 
 		/// <summary>
-		/// Uniform cache for desktop roles. Only used if -UniformCache is enabled
+		/// Uniform cache, useful for some instances of multi-role editor tests. Only used if -UniformCache is enabled
 		/// </summary>
 		private string UniformCache;
 
@@ -1287,11 +1287,6 @@ namespace Gauntlet
 			string PlatformCache = Path.Combine(DeviceCache, InDeviceDefiniton.Platform.ToString());
 			string ClientCache = Path.Combine(PlatformCache, SanitizedDeviceName);
 
-			if (Globals.Params.ParseParam("UniformCache"))
-			{
-				UniformCache = ClientCache;
-			}
-
 			// On Desktops, builds are installed in the device cache.
 			// When using device reservation blocks, we don't want to fully clean the cache and lose previously installed builds.
 			// If bRetainBuilds evaluates to true, it means we are in the second step or beyond in a device reservation block.
@@ -1355,6 +1350,11 @@ namespace Gauntlet
 				Directory.CreateDirectory(ClientCache);
 			}
 
+			if (Globals.Params.ParseParam("UniformCache"))
+			{
+				UniformCache = ClientCache;
+			}
+
 			return ClientCache;
 		}
 
