BundleUtils honors CustomConfig

#rnx

diff --git a/Engine/Source/Programs/AutomationTool/AutomationUtils/BundleUtils.cs b/Engine/Source/Programs/AutomationTool/AutomationUtils/BundleUtils.cs
--- a/Engine/Source/Programs/AutomationTool/AutomationUtils/BundleUtils.cs
+++ b/Engine/Source/Programs/AutomationTool/AutomationUtils/BundleUtils.cs
@@ -28,19 +28,19 @@ namespace AutomationUtils.Automation
 			public string ExecFileName { get; set; } // TODO: We never used this.  Clean this up.
 		}
 
-		public static void LoadBundleConfig(DirectoryReference ProjectDir, UnrealTargetPlatform Platform, out IReadOnlyDictionary<string, BundleSettings> Bundles)
+		public static void LoadBundleConfig(DirectoryReference ProjectDir, UnrealTargetPlatform Platform, string CustomConfig, out IReadOnlyDictionary<string, BundleSettings> Bundles)
 		{
-			LoadBundleConfig<BundleSettings>(ProjectDir, Platform, out Bundles, delegate (BundleSettings Settings, ConfigHierarchy BundleConfig, string Section) { });
+			LoadBundleConfig<BundleSettings>(ProjectDir, Platform, CustomConfig, out Bundles, delegate (BundleSettings Settings, ConfigHierarchy BundleConfig, string Section) { });
 		}
 
-		public static void LoadBundleConfig<TPlatformBundleSettings>(DirectoryReference ProjectDir, UnrealTargetPlatform Platform, 
+		public static void LoadBundleConfig<TPlatformBundleSettings>(DirectoryReference ProjectDir, UnrealTargetPlatform Platform, string CustomConfig, 
 			out IReadOnlyDictionary<string, TPlatformBundleSettings> Bundles, 
 			Action<TPlatformBundleSettings, ConfigHierarchy, string> GetPlatformSettings) 
 			where TPlatformBundleSettings : BundleSettings, new()
 		{
 			var Results = new List<TPlatformBundleSettings>();
 
-			ConfigHierarchy BundleConfig = ConfigCache.ReadHierarchy(ConfigHierarchyType.InstallBundle, ProjectDir, Platform);
+			ConfigHierarchy BundleConfig = ConfigCache.ReadHierarchy(ConfigHierarchyType.InstallBundle, ProjectDir, Platform, CustomConfig);
 
 			const string BundleDefinitionPrefix = "InstallBundleDefinition ";
 
@@ -207,9 +207,9 @@ namespace AutomationUtils.Automation
 			return FoundDependencies.Values;
 		}
 
-		public static bool HasPlatformBundleSource(DirectoryReference ProjectDir, UnrealTargetPlatform Platform)
+		public static bool HasPlatformBundleSource(DirectoryReference ProjectDir, UnrealTargetPlatform Platform, string CustomConfig)
 		{
-			ConfigHierarchy BundleConfig = ConfigCache.ReadHierarchy(ConfigHierarchyType.InstallBundle, ProjectDir, Platform);
+			ConfigHierarchy BundleConfig = ConfigCache.ReadHierarchy(ConfigHierarchyType.InstallBundle, ProjectDir, Platform, CustomConfig);
 			return 
 				BundleConfig.GetArray("InstallBundleManager.BundleSources", "DefaultBundleSources", out List<string> InstallBundleSources) && 
 				InstallBundleSources.Contains("Platform");
