Implement extra metadata to test report output for test audit in Gauntlet
* add descriptive metadata to the Test Data output
** TestTeam
** TestIntent
** TestHarness
** TestProfile (ie: Performance, functional, etc.)
* add test recipe to the Test Data output
** Version
** NodeClass
** Params (Distilled Parameters of used to run the test)
** TimeoutMinutes

#rnx

diff --git a/Engine/Source/Programs/AutomationTool/Gauntlet/Framework/Base/Gauntlet.TestReport.cs b/Engine/Source/Programs/AutomationTool/Gauntlet/Framework/Base/Gauntlet.TestReport.cs
--- a/Engine/Source/Programs/AutomationTool/Gauntlet/Framework/Base/Gauntlet.TestReport.cs
+++ b/Engine/Source/Programs/AutomationTool/Gauntlet/Framework/Base/Gauntlet.TestReport.cs
@@ -359,4 +359,92 @@ namespace Gauntlet
 			}
 		}
 	}
+
+#nullable enable
+	namespace Report
+	{
+		/// <summary>
+		/// Interface to retrieve test audit metadata
+		/// </summary>
+		public interface ITestAuditMetadata
+		{
+			/* Unique identifier */
+			public string TestKey { get; }
+
+			/* Test name for display */
+			public string TestName { get; }
+
+			/* TestNode class */
+			public string TestNodeClass { get; }
+
+			/* Test version */
+			public Version TestVersion { get; }
+
+			/* Duration in minutes after which the test run has timeout */
+			public float TestTimeoutMinutes { get; }
+
+			/* Command line used to run the test */
+			public string? TestCommandLine { get; }
+
+			/* Test harness */
+			public TestHarness? TestHarness { get; }
+
+			/* Test result profile */
+			public TestProfile? TestProfile { get; }
+
+			/* Team to which the test was written for */
+			public TestTeam? TestTeam { get; }
+
+			/* The purpose of the test */
+			public string? TestIntent { get; }
+
+			/* Tags associated with the test */
+			public string[]? TestTags { get; }
+
+			/* Params relevant to run this test */
+			public Dictionary<string, object>? TestParams { get; }
+		}
+
+		/// <summary>
+		/// Extend able Property value for report
+		/// </summary>
+		public partial class ReportProperty
+		{
+			private string Key;
+			public ReportProperty(string InKey)
+			{
+				Key = InKey;
+			}
+
+			public override string ToString()
+			{
+				return Key;
+			}
+		}
+
+		/// <summary>
+		/// Test profile (unit, functional, performance, etc)
+		/// </summary>
+		public partial class TestProfile : ReportProperty
+		{
+			public TestProfile(string Key) : base(Key) { }
+		}
+
+
+		/// <summary>
+		/// Test harness
+		/// </summary>
+		public partial class TestHarness : ReportProperty
+		{
+			public TestHarness(string Key) : base(Key) { }
+		}
+
+		/// <summary>
+		/// Test team
+		/// </summary>
+		public partial class TestTeam : ReportProperty
+		{
+			public TestTeam(string Key) : base(Key) { }
+		}
+	}
 }
\ No newline at end of file
diff --git a/Engine/Source/Programs/AutomationTool/Gauntlet/Framework/Gauntlet.Utils.cs b/Engine/Source/Programs/AutomationTool/Gauntlet/Framework/Gauntlet.Utils.cs
--- a/Engine/Source/Programs/AutomationTool/Gauntlet/Framework/Gauntlet.Utils.cs
+++ b/Engine/Source/Programs/AutomationTool/Gauntlet/Framework/Gauntlet.Utils.cs
@@ -33,6 +33,8 @@ namespace Gauntlet
 			set { InnerParams = value; }
 		}
 
+		public static string BuildCommand = String.Empty;
+
 		public static readonly DateTime ProgramStartTime = DateTime.UtcNow;
 
 		static string InnerTempDir;
diff --git a/Engine/Source/Programs/AutomationTool/Gauntlet/Framework/Utils/Gauntlet.HordeReport.cs b/Engine/Source/Programs/AutomationTool/Gauntlet/Framework/Utils/Gauntlet.HordeReport.cs
--- a/Engine/Source/Programs/AutomationTool/Gauntlet/Framework/Utils/Gauntlet.HordeReport.cs
+++ b/Engine/Source/Programs/AutomationTool/Gauntlet/Framework/Utils/Gauntlet.HordeReport.cs
@@ -13,6 +13,7 @@ using EpicGames.Core;
 using Logging = Microsoft.Extensions.Logging;
 using System.Globalization;
 using static Gauntlet.HordeReport.TestDataCollection;
+using Gauntlet.Report;
 
 namespace Gauntlet
 {
@@ -1174,6 +1175,8 @@ namespace Gauntlet
 		}
 		#endregion
 
+#nullable enable
+
 		/// <summary>
 		/// Contains information about a test session (v2)
 		/// </summary>
@@ -1181,14 +1184,15 @@ namespace Gauntlet
 		{
 			public override string Type => "Automated Test Session";
 
-			public override int Version => 2;
+			public const int DEFAULT_VERSION = 2;
+			public override int Version => DEFAULT_VERSION;
 
 			/// <summary>
 			/// Unique key that identify the test and group the different sessions
 			/// </summary>
 			public string Key { get; set; }
 
-			public override string GetTestDataKey(string BaseKey = null)
+			public override string GetTestDataKey(string? BaseKey = null)
 			{
 				return Key; // Ignore BaseKey
 			}
@@ -1216,7 +1220,7 @@ namespace Gauntlet
 			/// <summary>
 			/// Cursor to the current phase
 			/// </summary>
-			private TestPhase CurrentPhase { get; set; }
+			private TestPhase? CurrentPhase { get; set; }
 
 			/// <summary>
 			/// Static constructor
@@ -1227,11 +1231,31 @@ namespace Gauntlet
 				LogValueFormatter.RegisterFormatter(typeof(TestLogValue), new TestLogValueFormatter());
 			}
 
-			public AutomatedTestSessionData(string InKey, string InName = null) : base()
+			public AutomatedTestSessionData(string InKey, string? InName = null) : base()
 			{
 				Key = InKey;
 				Data = new TestSession(string.IsNullOrEmpty(InName) ? InKey : InName, this);
-				CurrentPhase = null;
+			}
+
+			public AutomatedTestSessionData(ITestAuditMetadata InAudit) : this(InAudit.TestKey, InAudit.TestName)
+			{
+				Data.summary.commandline = InAudit.TestCommandLine;
+
+				AddTestRecipe(InAudit.TestNodeClass, InAudit.TestVersion.ToString(), InAudit.TestTimeoutMinutes, InAudit.TestParams);
+
+				if (InAudit.TestIntent != null
+					|| InAudit.TestHarness != null
+					|| InAudit.TestProfile != null
+					|| InAudit.TestTeam != null)
+				{
+					AddTestDescription(InAudit.TestTeam?.ToString(), InAudit.TestIntent, InAudit.TestHarness?.ToString(), InAudit.TestProfile?.ToString());
+				}
+
+				string[]? InTags = InAudit.TestTags;
+				if (InTags != null)
+				{
+					AddTestTags(InTags);
+				}
 			}
 
 			/// <summary>
@@ -1242,7 +1266,7 @@ namespace Gauntlet
 				public string name { get; set; }
 				[JsonIgnore]
 				public string key { get; private set; }
-				public string appInstanceLogPath { get; set; }
+				public string? appInstanceLogPath { get; set; }
 				public Dictionary<string, string> metadata { get; set; }
 
 				public TestDevice(string InName, string InKey)
@@ -1259,8 +1283,9 @@ namespace Gauntlet
 			public class TestSessionSummary
 			{
 				public string testName { get; set; }
-				public string dateTime { get; set; }
+				public string? dateTime { get; set; }
 				public float timeElapseSec { get; set; }
+				public string? commandline {  get; set; }
 				public int phasesTotalCount { get; set; }
 				public int phasesSucceededCount { get; set; }
 				public int phasesUndefinedCount { get; set; }
@@ -1272,6 +1297,43 @@ namespace Gauntlet
 				}
 			}
 
+			/// <summary>
+			/// Description of the test
+			/// </summary>
+			public class TestDescription
+			{
+				public string? team { get; set; }
+				public string? profile { get; set; }
+				public string? harness { get; set; }
+				public string? intent { get; set; }
+
+				public TestDescription(string? InTeam = null, string? InProfile = null, string? InHarness = null, string? InIntent = null)
+				{
+					team = InTeam;
+					profile = InProfile;
+					harness = InHarness;
+					intent = InIntent;
+				}
+			}
+
+			/// <summary>
+			/// Recipe to run the test
+			/// </summary>
+			public class TestRecipe
+			{
+				public string nodeClass { get; set; }
+				public string version { get; set; }
+				public float timeoutMinutes { get; set; }
+				[JsonPropertyName("params")]
+				public Dictionary<string, object>? @params { get; set; }
+				public TestRecipe(string InNodeClass, string InVersion, float InTimeoutMinutes = 0)
+				{
+					nodeClass = InNodeClass;
+					version = InVersion;
+					timeoutMinutes = InTimeoutMinutes;
+				}
+			}
+
 			[JsonConverter(typeof(JsonTryParseEnumConverter))]
 			public enum TestPhaseOutcome
 			{
@@ -1290,27 +1352,25 @@ namespace Gauntlet
 			{
 				public string key { get; set; }
 				public string name { get; set; }
-				public string dateTime { get; set; }
+				public string? dateTime { get; set; }
 				public float timeElapseSec { get; set; }
 				public TestPhaseOutcome outcome { get; set; }
-				public string errorFingerprint { get; set; }
+				public string? errorFingerprint { get; set; }
 				public bool? hasWarning { get; set; }
-				public HashSet<string> tags { get; set; }
+				public HashSet<string>? tags { get; set; }
 				public HashSet<string> deviceKeys { get; set; }
-				public string eventStreamPath { get; set; }
+				public string? eventStreamPath { get; set; }
 
-				private TestSession session { get; set; }
+				private TestSession? session { get; set; }
 
 				private TestEventStream eventStream { get; set; }
 
-				public TestPhase(string InName, string InKey = null, TestSession InSession = null)
+				public TestPhase(string InName, string? InKey = null, TestSession? InSession = null)
 				{
 					key = string.IsNullOrEmpty(InKey) ? GenerateKey(InName) : ValidateKey(InKey);
 					name = InName;
 					deviceKeys = new HashSet<string>();
 					outcome = TestPhaseOutcome.Unknown;
-					tags = null;
-					errorFingerprint = null;
 					eventStream = new TestEventStream(session?.report, this);
 					if (InSession != null)
 					{
@@ -1369,7 +1429,10 @@ namespace Gauntlet
 				public void AttachToSession(TestSession InSession)
 				{
 					session = InSession;
-					eventStream.AttachToReport(session.report);
+					if (session.report != null)
+					{
+						eventStream.AttachToReport(session.report);
+					}
 					session.phases.Add(this);
 					session.summary.phasesTotalCount += 1;
 					UpdatePhasesCount();
@@ -1455,21 +1518,22 @@ namespace Gauntlet
 				public TestSessionSummary summary { get; set; }
 				public List<TestPhase> phases { get; set; }
 				public Dictionary<string, TestDevice> devices { get; set; }
-				public int version { get { return report.Version; } }
-				public Dictionary<string, string> metadata { get { return report.Metadata; } }
-				public HashSet<string> tags { get; set; }
+				public int version { get { return report?.Version ?? DEFAULT_VERSION /* default to test data v2 report version */; } }
+				public Dictionary<string, string> metadata { get { return report?.Metadata ?? new(); } }
+				public HashSet<string>? tags { get; set; }
+				public TestDescription? description { get; set; }
+				public TestRecipe? recipe { get; set; }
 
 				[JsonIgnore]
 				// back pointer
-				public AutomatedTestSessionData report { get; protected set; }
+				public AutomatedTestSessionData? report { get; protected set; }
 
-				public TestSession(string InName, AutomatedTestSessionData InReport = null)
+				public TestSession(string InName, AutomatedTestSessionData? InReport = null)
 				{
 					summary = new TestSessionSummary(InName);
 					phases = new List<TestPhase>();
 					devices = new Dictionary<string, TestDevice>();
 					report = InReport;
-					tags = null;
 				}
 			}
 
@@ -1490,7 +1554,7 @@ namespace Gauntlet
 			/// <param name="InName"></param>
 			/// <param name="InKey"></param>
 			/// <returns></returns>
-			public virtual TestDevice AddDevice(string InName, string InKey = null)
+			public virtual TestDevice AddDevice(string InName, string? InKey = null)
 			{
 				string Key = ValidateKey(InKey ?? InName);
 				if (Data.devices.ContainsKey(Key))
@@ -1509,7 +1573,7 @@ namespace Gauntlet
 			/// </summary>
 			/// <param name="Key"></param>
 			/// <returns></returns>
-			public virtual TestDevice GetDevice(string Key)
+			public virtual TestDevice? GetDevice(string Key)
 			{
 				return Data.devices.GetValueOrDefault(Key);
 			}
@@ -1545,7 +1609,7 @@ namespace Gauntlet
 			/// <param name="InName"></param>
 			/// <param name="InKey"></param>
 			/// <returns></returns>
-			public virtual TestPhase AddPhase(string InName, string InKey = null)
+			public virtual TestPhase AddPhase(string InName, string? InKey = null)
 			{
 				TestPhase Phase = new TestPhase(InName, InKey, Data);
 				CurrentPhase = Phase;
@@ -1557,7 +1621,7 @@ namespace Gauntlet
 			/// </summary>
 			/// <param name="Key"></param>
 			/// <returns></returns>
-			public virtual TestPhase GetPhase(string Key)
+			public virtual TestPhase? GetPhase(string Key)
 			{
 				return Data.phases.FirstOrDefault(P => P.key == Key);
 			}
@@ -1569,7 +1633,7 @@ namespace Gauntlet
 			/// <exception cref="AutomationException"></exception>
 			public virtual void SetCurrentPhase(string Key)
 			{
-				TestPhase Phase = GetPhase(Key);
+				TestPhase? Phase = GetPhase(Key);
 				if (Phase == null)
 				{
 					throw new AutomationException("Could not find Phase '{0}'", Key);
@@ -1600,7 +1664,7 @@ namespace Gauntlet
 			/// Get the current session phase 
 			/// </summary>
 			/// <returns></returns>
-			public virtual TestPhase GetCurrentPhase()
+			public virtual TestPhase? GetCurrentPhase()
 			{
 				return CurrentPhase;
 			}
@@ -1618,6 +1682,31 @@ namespace Gauntlet
 				Data.tags.UnionWith(NewTags);
 			}
 
+			/// <summary>
+			/// Add test description
+			/// </summary>
+			/// <param name="Team"></param>
+			/// <param name="Intent"></param>
+			/// <param name="Harness"></param>
+			/// <param name="Profile"></param>
+			public virtual void AddTestDescription(string? Team, string? Intent, string? Harness, string? Profile)
+			{
+				Data.description = new TestDescription(Team, Profile, Harness, Intent);
+			}
+
+			/// <summary>
+			/// Add recipe to run the same test
+			/// </summary>
+			/// <param name="NodeClass"></param>
+			/// <param name="Version"></param>
+			/// <param name="TimeoutMinutes"></param>
+			/// <param name="Params"></param>
+			public virtual void AddTestRecipe(string NodeClass, string Version, float TimeoutMinutes = 0, Dictionary<string, object>? Params = null)
+			{
+				Data.recipe = new TestRecipe(NodeClass, Version, TimeoutMinutes);
+				Data.recipe.@params = Params;
+			}
+
 			/// <summary>
 			/// Represent the event stream of a phase
 			/// </summary>
@@ -1630,11 +1719,11 @@ namespace Gauntlet
 				// Redact hex strings
 				static Regex Regex_hexstring = new Regex(@"(?:0X|#)[0-9A-F]+", RegexOptions.Compiled);
 
-				private AutomatedTestSessionData report { get; set; }
+				private AutomatedTestSessionData? report { get; set; }
 				private TestPhase phase { get; set; }
 				private List<LogEvent> stream { get; set; }
 
-				public TestEventStream(AutomatedTestSessionData InReport, TestPhase InPhase)
+				public TestEventStream(AutomatedTestSessionData? InReport, TestPhase InPhase)
 				{
 					report = InReport;
 					phase = InPhase;
@@ -1711,14 +1800,14 @@ namespace Gauntlet
 				/// <param name="Args"></param>
 				public virtual void AddEvent(DateTime Time, Logging.LogLevel Level, string Message, params object[] Args)
 				{
-					string Format = null;
-					Dictionary<string, object> Properties = null;
+					string? Format = null;
+					Dictionary<string, object?>? Properties = null;
 					if (Args.Any())
 					{
 						Format = Message;
-						Properties = new Dictionary<string, object>();
-						MessageTemplate.ParsePropertyValues(Format, Args, Properties);
-						Message = MessageTemplate.Render(Format, Properties);
+						Properties = new Dictionary<string, object?>();
+						MessageTemplate.ParsePropertyValues(Format, Args, Properties!);
+						Message = MessageTemplate.Render(Format, Properties!);
 					}
 					Logging.EventId EventId = Level == Logging.LogLevel.Critical ? KnownLogEvents.Gauntlet_FatalEvent : KnownLogEvents.Gauntlet_TestEvent;
 					LogEvent Event = new LogEvent(Time, Level, EventId, Message, Format, Properties, null);
@@ -1798,7 +1887,7 @@ namespace Gauntlet
 				public virtual void AddURLLink(DateTime Time, Logging.LogLevel Level, string Name, string URL)
 				{
 					string Token = TestLogValueTypes.url.ToLower();
-					var Properties = new Dictionary<string, object>() { { Token, TestLogValue.CreateURLLink(Name, URL) } };
+					var Properties = new Dictionary<string, object?>() { { Token, TestLogValue.CreateURLLink(Name, URL) } };
 					string Format = $"{{{Token}}}";
 					string LogMessage = MessageTemplate.Render(Format, Properties);
 					LogEvent Event = new LogEvent(Time, Level, KnownLogEvents.Gauntlet_TestEvent, LogMessage, Format, Properties, null);
@@ -1817,7 +1906,7 @@ namespace Gauntlet
 				{
 					if (report == null)
 					{
-						AddWarning("Phase not attached to a report. Failed to attached Artifact {Name}.", Artifact.ToString());
+						AddWarning("Phase not attached to a report. Failed to attached Artifact {Name}.", Artifact.ToString()!);
 					}
 					else
 					{
@@ -1894,7 +1983,7 @@ namespace Gauntlet
 			/// <param name="FilePath"></param>
 			/// <param name="Name"></param>
 			/// <returns></returns>
-			protected virtual bool AttachArtifact(ref string FilePath, string Name = null)
+			protected virtual bool AttachArtifact(ref string FilePath, string? Name = null)
 			{
 				if (AttachArtifact(FilePath, Name))
 				{
@@ -1911,9 +2000,9 @@ namespace Gauntlet
 			/// <param name="FilePath"></param>
 			/// <param name="Name"></param>
 			/// <returns></returns>
-			protected virtual bool AttachArtifactOrWarn(ref string FilePath, string Name = null)
+			protected virtual bool AttachArtifactOrWarn(ref string? FilePath, string? Name = null)
 			{
-				if (!AttachArtifact(ref FilePath, Name))
+				if (FilePath != null && !AttachArtifact(ref FilePath, Name))
 				{
 					AddWarning("Failed to attached Artifact {FilePath}.", Path.GetRelativePath(Globals.UnrealRootDir, FilePath));
 					FilePath = null;
@@ -2003,7 +2092,7 @@ namespace Gauntlet
 			public interface ILogArtifact
 			{
 				public bool AttachTo(AutomatedTestSessionData Report);
-				Dictionary<string, object> ToProperties();
+				Dictionary<string, object?> ToProperties();
 			}
 
 			/// <summary>
@@ -2014,10 +2103,10 @@ namespace Gauntlet
 				public string Type { get; protected set; }
 				public string Token { get; protected set; }
 				public string Name { get; protected set; }
-				public Dictionary<string, string> Files { get; protected set; }
+				public Dictionary<string, string?> Files { get; protected set; }
 				public string RootPath { get; protected set; }
 
-				public ArtifactFiles(string InType, string InName, Dictionary<string, string> InFiles, string InRootPath)
+				public ArtifactFiles(string InType, string InName, Dictionary<string, string?> InFiles, string InRootPath)
 				{
 					Type = InType;
 					Token = Type.ToLower().Replace(" ", "_");
@@ -2026,17 +2115,17 @@ namespace Gauntlet
 					RootPath = InRootPath;
 				}
 
-				public ArtifactFiles(string InName, Dictionary<string, string> InFiles, string InRootPath)
+				public ArtifactFiles(string InName, Dictionary<string, string?> InFiles, string InRootPath)
 					: this(TestLogValueTypes.artifacts, InName, InFiles, InRootPath) { }
 
-				private bool AttachFileTo(ref string FilePath, AutomatedTestSessionData Report)
+				private bool AttachFileTo(ref string? FilePath, AutomatedTestSessionData Report)
 				{
 					if (string.IsNullOrEmpty(FilePath))
 					{
 						// Ignore this file
 						return true;
 					}
-					string Name = null;
+					string? Name = null;
 					if (!Path.IsPathFullyQualified(FilePath))
 					{
 						Name = FilePath;
@@ -2051,7 +2140,7 @@ namespace Gauntlet
 					bool AllAttached = true;
 					foreach (string Key in Files.Keys.ToList())
 					{
-						string FilePath = Files[Key];
+						string? FilePath = Files[Key];
 						if (!AttachFileTo(ref FilePath, Report))
 						{
 							AllAttached = AllAttached && false;
@@ -2062,7 +2151,7 @@ namespace Gauntlet
 					return AllAttached;
 				}
 
-				public Dictionary<string, object> ToProperties()
+				public Dictionary<string, object?> ToProperties()
 				{
 					Dictionary<Utf8String, object> Prop = new Dictionary<Utf8String, object>();
 					foreach(var Item in Files)
@@ -2073,7 +2162,7 @@ namespace Gauntlet
 						}
 					}
 					TestLogValue Value = new TestLogValue(Type, Name, Prop);
-					Dictionary<string, object> EventProperties = new Dictionary<string, object> { { Token, Value } };
+					Dictionary<string, object?> EventProperties = new Dictionary<string, object?> { { Token, Value } };
 
 					return EventProperties;
 				}
@@ -2084,9 +2173,9 @@ namespace Gauntlet
 			/// </summary>
 			public class ImageComparisonFiles : ArtifactFiles, ILogArtifact
 			{
-				public ImageComparisonFiles(string InName, Dictionary<string, string> InFiles, string InRootPath, bool AddMetadataFiles = false) : base(TestLogValueTypes.imageCompare, InName, InFiles, InRootPath)
+				public ImageComparisonFiles(string InName, Dictionary<string, string?> InFiles, string InRootPath, bool AddMetadataFiles = false) : base(TestLogValueTypes.imageCompare, InName, InFiles, InRootPath)
 				{
-					string FilePath = Files.GetValueOrDefault(TestLogValueTypes.unapproved);
+					string? FilePath = Files.GetValueOrDefault(TestLogValueTypes.unapproved);
 					if (AddMetadataFiles && !string.IsNullOrEmpty(FilePath))
 					{
 						Files.Add(TestLogValueTypes.unapprovedMetadata, Path.ChangeExtension(FilePath, ".json"));
@@ -2099,7 +2188,7 @@ namespace Gauntlet
 					FilePath = Files.GetValueOrDefault(TestLogValueTypes.difference);
 					if (AddMetadataFiles && !string.IsNullOrEmpty(FilePath))
 					{
-						Files.Add(TestLogValueTypes.differenceReport, Path.Combine(Path.GetDirectoryName(FilePath), "Report.json"));
+						Files.Add(TestLogValueTypes.differenceReport, Path.Combine(Path.GetDirectoryName(FilePath)!, "Report.json"));
 					}
 				}
 			}
@@ -2110,7 +2199,7 @@ namespace Gauntlet
 			public class EmbeddedHTMLFile : ArtifactFiles, ILogArtifact
 			{
 				public EmbeddedHTMLFile(string InName, string InFile, string InRootPath)
-					: base(TestLogValueTypes.embeddedHtml, InName, new Dictionary<string, string>() { { "path", InFile } }, InRootPath) { }
+					: base(TestLogValueTypes.embeddedHtml, InName, new Dictionary<string, string?>() { { "path", InFile } }, InRootPath) { }
 			}
 
 			private static Logging.LogLevel EventTypeToLogLevel(EventType Type)
@@ -2222,14 +2311,14 @@ namespace Gauntlet
 			/// <param name="InPhaseKey"></param>
 			/// <returns></returns>
 			/// <exception cref="AutomationException"></exception>
-			public virtual TestEventStream GetEventStream(string InPhaseKey = null)
+			public virtual TestEventStream GetEventStream(string? InPhaseKey = null)
 			{
 				if (InPhaseKey == null && CurrentPhase == null)
 				{
 					AddPhase("Main");
 				}
-				string PhaseKey = string.IsNullOrEmpty(InPhaseKey) ? CurrentPhase.key : InPhaseKey;
-				TestPhase Phase = GetPhase(PhaseKey);
+				string PhaseKey = string.IsNullOrEmpty(InPhaseKey) ? CurrentPhase!.key : InPhaseKey;
+				TestPhase? Phase = GetPhase(PhaseKey);
 				if (Phase == null)
 				{
 					throw new AutomationException("No Phase '{0}' set in Test '{1}'", PhaseKey, Key);
@@ -2249,7 +2338,7 @@ namespace Gauntlet
 				Log.Verbose("Writing Phase Event Streams at {0}", OutputDir);
 				foreach (TestPhase Phase in Data.phases)
 				{
-					string OutputFilePath = null;
+					string? OutputFilePath = null;
 					TestEventStream Stream = Phase.GetStream();
 					// skip if the event stream is empty
 					if (Stream.HasAnyEvent())
@@ -2343,7 +2432,7 @@ namespace Gauntlet
 							// If Artifact values is not null nor a bunch on 0, then we have a file attachment.
 							if (!string.IsNullOrEmpty(Artifact) && Artifact.Substring(0, 4) != "0000")
 							{
-								UnrealAutomationArtifact ArtifactEntry = InTestResult.Artifacts.Find(A => A.Id == Artifact);
+								UnrealAutomationArtifact? ArtifactEntry = InTestResult.Artifacts.Find(A => A.Id == Artifact);
 								if (ArtifactEntry != null)
 								{
 									DateTime EventTime = ConvertUETimeStringToDateTime(InTestEntry.Timestamp);
@@ -2389,17 +2478,17 @@ namespace Gauntlet
 			}
 			public class DataItem
 			{
-				public string key { get; set; }
-				public object data { get; set; }
+				public string? key { get; set; }
+				public object? data { get; set; }
 			}
 			public TestDataCollection()
 			{
 				items = new List<DataItem>();
 			}
 
-			public DataItem AddNewTestReport(object InData, string InKey = null)
+			public DataItem AddNewTestReport(object InData, string? InKey = null)
 			{
-				Dictionary<string, object> ExtraItems = null;
+				Dictionary<string, object>? ExtraItems = null;
 				if (InData is BaseHordeReport InHordeReport)
 				{
 					InHordeReport.FinalizeReport();
@@ -2448,7 +2537,7 @@ namespace Gauntlet
 			/// <param name="bIncrementNameIfFileExists"></param>
 			public void WriteToJson(string OutputTestDataFilePath, bool bIncrementNameIfFileExists = false)
 			{
-				string OutputTestDataDir = Path.GetDirectoryName(OutputTestDataFilePath);
+				string OutputTestDataDir = Path.GetDirectoryName(OutputTestDataFilePath)!;
 				if (!Directory.Exists(OutputTestDataDir))
 				{
 					Directory.CreateDirectory(OutputTestDataDir);
diff --git a/Engine/Source/Programs/AutomationTool/Gauntlet/Unreal/Automation/UE.Automation.cs b/Engine/Source/Programs/AutomationTool/Gauntlet/Unreal/Automation/UE.Automation.cs
--- a/Engine/Source/Programs/AutomationTool/Gauntlet/Unreal/Automation/UE.Automation.cs
+++ b/Engine/Source/Programs/AutomationTool/Gauntlet/Unreal/Automation/UE.Automation.cs
@@ -9,6 +9,7 @@ using System.Linq;
 using UnrealBuildTool;
 using EpicGames.Core;
 using Log = Gauntlet.Log;
+using Gauntlet.Report;
 
 namespace UE
 {
@@ -498,29 +499,58 @@ namespace UE
 		}
 
 		/// <summary>
-		/// Override the Suite by RunTest if it is a Group
+		/// Override the Suite by GroupName if RunTest is a Group
 		/// </summary>
 		public override string Suite
 		{
 			get
 			{
-				string BaseSuite = base.Suite;
+				if (GroupName != null)
+				{
+					return GroupName;
+				}
+
 				if (GetCachedConfiguration() is AutomationTestConfig AutomationConfig)
 				{
 					if (!string.IsNullOrEmpty(AutomationConfig.RunTest))
 					{
 						// try to get a usefull identification token
-						int StartPos = AutomationConfig.RunTest.ToLower().StartsWith("group:") ? 6 : AutomationConfig.RunTest.LastIndexOf('.') + 1;
-						BaseSuite = AutomationConfig.RunTest.Substring(StartPos);
-						if (string.IsNullOrEmpty(BaseSuite))
+						int StartPos = AutomationConfig.RunTest.LastIndexOf('.') + 1;
+						string TestName = AutomationConfig.RunTest.Substring(StartPos);
+						if (!string.IsNullOrEmpty(TestName))
+						{
+							return TestName;
+						}
+						// get the token before if the first one is empty
+						else if (StartPos > 1)
 						{
-							// get the token before if the first one is empty
-							BaseSuite = AutomationConfig.RunTest.Substring(0, StartPos - 1).Substring(AutomationConfig.RunTest.LastIndexOf('.', StartPos - 2) + 1);
+							return AutomationConfig.RunTest.Substring(0, StartPos - 1).Substring(AutomationConfig.RunTest.LastIndexOf('.', StartPos - 2) + 1);
 						}
 					}
 				}
 
-				return BaseSuite;
+				return base.Suite;
+			}
+		}
+
+		/// <summary>
+		/// The test group name if any from the command line arguments
+		/// </summary>
+		protected virtual string GroupName
+		{
+			get
+			{
+				// try to get the group name from the test config by parsing the Runtest command line argument.
+				if (GetCachedConfiguration() is AutomationTestConfig AutomationConfig)
+				{
+					if (!string.IsNullOrEmpty(AutomationConfig.RunTest) && AutomationConfig.RunTest.StartsWith("group:", StringComparison.OrdinalIgnoreCase))
+					{
+						int StartPos = ("group:").Length;
+						return AutomationConfig.RunTest.Substring(StartPos);
+					}
+				}
+
+				return null;
 			}
 		}
 
@@ -557,6 +587,16 @@ namespace UE
 			}
 		}
 
+		/* Provide test metadata context for test reporting */
+		/// <inheritdoc/>
+		protected override TestHarness TestHarness => TestHarness.UEFunctionalTest;
+		/// <inheritdoc/>
+		protected override TestTeam TestTeam => GroupName != null ? new TestTeam(GroupName) : null;
+		/// <inheritdoc/>
+		protected override TestProfile TestProfile => TestProfile.Functional;
+		/// <inheritdoc/>
+		protected override string TestIntent => GroupName != null ? $"Run functional tests from {Context.BuildInfo.ProjectName} {GroupName} suite" : null;
+
 		/// <summary>
 		/// Called when a test is starting
 		/// </summary>
diff --git a/Engine/Source/Programs/AutomationTool/Gauntlet/Unreal/Base/Gauntlet.UnrealTestNode.cs b/Engine/Source/Programs/AutomationTool/Gauntlet/Unreal/Base/Gauntlet.UnrealTestNode.cs
--- a/Engine/Source/Programs/AutomationTool/Gauntlet/Unreal/Base/Gauntlet.UnrealTestNode.cs
+++ b/Engine/Source/Programs/AutomationTool/Gauntlet/Unreal/Base/Gauntlet.UnrealTestNode.cs
@@ -16,6 +16,7 @@ using static Gauntlet.UnrealSessionInstance;
 using Gauntlet.Utils;
 using static Gauntlet.HordeReport.AutomatedTestSessionData;
 using Gauntlet.TestTracking;
+using Gauntlet.Report;
 
 namespace Gauntlet
 {
@@ -37,6 +38,68 @@ namespace Gauntlet
 			PromoteWarnings = 2,            // Promote warnings from Unreal instances to regular test warnings.	(By default only ensures are warnings)
 		}
 
+		/// <summary>
+		/// Test Metadata for audit
+		/// </summary>
+		public ITestAuditMetadata TestMetadata;
+
+#nullable enable
+		/// <summary>
+		/// Unique key to identify and cross reference test across test sessions 
+		/// </summary>
+		protected virtual string ReportTestKey
+		{
+			get
+			{
+				return Type;
+			}
+		}
+		protected virtual string HordeReportTestKey => ReportTestKey;
+
+		/// <summary>
+		/// The test name used in report
+		/// </summary>
+		protected virtual string ReportTestName
+		{
+			get
+			{
+				string ReportName = Name;
+				if (ReportName.Split('.').Length > 1)
+				{
+					ReportName = ReportName.Split('.').Last();
+				}
+
+				return ReportName;
+			}
+		}
+		protected virtual string HordeReportTestName => ReportTestName;
+
+		/// <summary>
+		/// The harness used to execute the test
+		/// </summary>
+		protected virtual TestHarness? TestHarness => null;
+
+		/// <summary>
+		/// The test profile (unit test, functional, performance)
+		/// </summary>
+		protected virtual TestProfile? TestProfile => null;
+
+		/// <summary>
+		/// The team for which the test is for
+		/// </summary>
+		protected virtual TestTeam? TestTeam => null;
+
+		/// <summary>
+		/// Short description of the test intent
+		/// </summary>
+		protected virtual string? TestIntent => null;
+
+		/// <summary>
+		/// The tags associated with this test
+		/// </summary>
+		protected virtual string[]? TestTags => null;
+#nullable disable
+
 		/// <summary>
 		/// Returns an identifier for this test
 		/// </summary>
@@ -427,6 +490,24 @@ namespace Gauntlet
 			PopulateCommandlineInfo();
 			// We format warnings ourselves so don't show these
 			LogWarningsAndErrorsAfterSummary = false;
+
+			TestMetadata = new UnrealTestMetadata
+			(
+				() => HordeReportTestKey,
+				() => HordeReportTestName,
+				() => GetType().FullName,
+				() => TestVersion,
+				() => MaxDuration / 60,
+				() => GetRunLocalCommand(Globals.BuildCommand),
+				() => TestHarness,
+				() => TestProfile,
+				() => TestTeam,
+				() => TestIntent,
+				() => TestTags,
+				() => GetDistilledTestParams()
+						.ToLookup(A => A.Key, A => A.Value) /* use a lookup here to avoid duplicate param key */
+						.ToDictionary(A => A.Key, A => A.Last())
+			);
 		}
 
 		~UnrealTestNode()
@@ -960,9 +1041,9 @@ namespace Gauntlet
 			});
 		}
 
-		public override string GetRunLocalCommand(string LaunchingBuildCommand)
+		protected virtual IEnumerable<KeyValuePair<string, object>> GetDistilledTestParams()
 		{
-			string[] ArgsToNotDisplay =
+			HashSet<string> ArgsToNotDisplay =
 			new List<string>{
 					"test",
 					"tests",
@@ -991,46 +1072,41 @@ namespace Gauntlet
 					"VerifyLogin",
 					"cleardevices",
 					"reboot",
+					"RequireBoot",
 					"fullclean",
 					"BuildName",
 					"PerfReportServer",
-			}.Select(I => I.ToLower()).ToArray();
+					"sandbox",
+			}.Select(I => I.ToLower()).ToHashSet();
 
-			bool ShouldArgBeDisplayed(string InArg)
+			return Context.TestParams.AllArguments.Select(A =>
 			{
-				InArg = InArg.Split("=", 2).First().ToLower();
-				return !ArgsToNotDisplay.Any(str => InArg == str);
-			}
+				var Arg = A.Split('=', 2);
+				return new KeyValuePair<string, object>(Arg[0], Arg.Length > 1 ? (Arg[1].Equals("false", StringComparison.OrdinalIgnoreCase)? false : Arg[1]) : true);
+			}).Where(A => !ArgsToNotDisplay.Contains(A.Key.ToLower()));
+		}
+
+		public override string GetRunLocalCommand(string LaunchingBuildCommand)
+		{
 
-			string WrapParameterInQuotes(string InString)
+			string WrapParameterInQuotes(string ParamStrName, string ParamStrValue)
 			{
-				var positionOfEqual = InString.IndexOf("=");
-				if (positionOfEqual > 0)
+				if (ParamStrValue.Contains(' '))
 				{
-					var ParamStrName = InString.Substring(0, positionOfEqual);
-					var ParamStrValue = InString.Substring(positionOfEqual + 1);
-					int n;
-					if (int.TryParse(ParamStrValue, out n))
-					{
-						return InString;
-					}
-					else
-					{
-						return string.Format("{0}=\"{1}\"", ParamStrName, ParamStrValue);
-					}
+					return string.Format("{0}=\"{1}\"", ParamStrName, ParamStrValue);
 				}
 				else
 				{
-					return InString;
+					return string.Format("{0}={1}", ParamStrName, ParamStrValue);
 				}
 			}
 
-			var CleanedArgs = Context.TestParams.AllArguments.Where(arg => ShouldArgBeDisplayed(arg));
-			CleanedArgs = CleanedArgs.Select(str => WrapParameterInQuotes(str));
+			var CleanedArgs = GetDistilledTestParams().Select(Arg => Arg.Value.Equals(true)? Arg.Key : WrapParameterInQuotes(Arg.Key, Arg.Value.ToString()));
 
 			string StringOfCleanedArguments = string.Format("-{0}", string.Join(" -", CleanedArgs));
 			string CommandToRunLocally =
 				string.Format("RunUAT {0} -Test=\"{1}\" {2}", LaunchingBuildCommand, GetType(), StringOfCleanedArguments);
+
 			return CommandToRunLocally;
 		}
 
@@ -1398,17 +1474,6 @@ namespace Gauntlet
 			return null;
 		}
 
-		/// <summary>
-		/// Unique key to identify and cross reference test name across test sessions 
-		/// </summary>
-		protected virtual string HordeReportTestKey
-		{
-			get
-			{
-				return Type;
-			}
-		}
-
 		private string GetHordeTestDataKey_legacy()
 		{
 			if (string.IsNullOrEmpty(GetCachedConfiguration().HordeTestDataKey))
@@ -1419,23 +1484,6 @@ namespace Gauntlet
 			return GetCachedConfiguration().HordeTestDataKey;
 		}
 
-		/// <summary>
-		/// Optional override for the test report name, useful when a single test node has different testing modes
-		/// </summary>
-		protected virtual string HordeReportTestName
-		{
-			get
-			{
-				string ReportName = Name;
-				if (ReportName.Split('.').Length > 1)
-				{
-					ReportName = ReportName.Split('.').Last();
-				}
-
-				return ReportName;
-			}
-		}
-
 		/// <summary>
 		/// Generate a Simple Test Report from the results of this test
 		/// </summary>
@@ -1459,7 +1507,7 @@ namespace Gauntlet
 			}
 			else
 			{
-				AutomatedTestSessionData TestReport = new AutomatedTestSessionData(HordeReportTestKey, HordeReportTestName);
+				AutomatedTestSessionData TestReport = new AutomatedTestSessionData(TestMetadata);
 				DateTime SessionTime = SessionStartTime.ToUniversalTime();
 				float TimeElapse = (float)(DateTime.Now - SessionStartTime).TotalSeconds;
 				TestReport.SetSessionTiming(SessionTime, TimeElapse);
@@ -1652,7 +1700,7 @@ namespace Gauntlet
 			// Handle initialization phase
 			if (!Horde.UseTestDataV1)
 			{
-				AutomatedTestSessionData SessionReport = new AutomatedTestSessionData(HordeReportTestKey, HordeReportTestName);
+				AutomatedTestSessionData SessionReport = new AutomatedTestSessionData(TestMetadata);
 				string ArtifactPathLabel = new DirectoryInfo(ArtifactPath).Name;
 				string HordeArtifactPath = Path.Combine(GetCachedConfiguration().HordeArtifactPath, ArtifactPathLabel);
 				SessionReport.SetOutputArtifactPath(HordeArtifactPath);
diff --git a/Engine/Source/Programs/AutomationTool/Gauntlet/Unreal/RunUnreal.cs b/Engine/Source/Programs/AutomationTool/Gauntlet/Unreal/RunUnreal.cs
--- a/Engine/Source/Programs/AutomationTool/Gauntlet/Unreal/RunUnreal.cs
+++ b/Engine/Source/Programs/AutomationTool/Gauntlet/Unreal/RunUnreal.cs
@@ -82,6 +82,7 @@ namespace Gauntlet
 		public override ExitCode Execute()
 		{
 			Globals.Params = new Gauntlet.Params(this.Params);
+			Globals.BuildCommand = ToString().Split(".").Last();
 
 			UnrealTestOptions ContextOptions = new UnrealTestOptions();
 
diff --git a/Engine/Source/Programs/AutomationTool/Gauntlet/Unreal/Utils/Gauntlet.UnrealTestAudit.cs b/Engine/Source/Programs/AutomationTool/Gauntlet/Unreal/Utils/Gauntlet.UnrealTestAudit.cs
new file mode 100644
--- /dev/null
+++ b/Engine/Source/Programs/AutomationTool/Gauntlet/Unreal/Utils/Gauntlet.UnrealTestAudit.cs
@@ -0,0 +1,133 @@
+// Copyright Epic Games, Inc. All Rights Reserved.
+
+using System;
+using System.Text.RegularExpressions;
+using System.Collections.Generic;
+using Gauntlet.Report;
+using System.Linq;
+
+#nullable enable
+
+namespace Gauntlet
+{
+	namespace Report
+	{
+		/// <summary>
+		/// Test profile (unit, functional, performance, etc)
+		/// </summary>
+		public partial class TestProfile
+		{
+			public static readonly TestProfile Unit = new TestProfile("Unit");
+			public static readonly TestProfile Functional = new TestProfile("Functional");
+			public static readonly TestProfile Performance = new TestProfile("Performance");
+			public static readonly TestProfile Smoke = new TestProfile("Smoke");
+		}
+
+		/// <summary>
+		/// Test harness
+		/// </summary>
+		public partial class TestHarness
+		{
+			public static readonly TestHarness UEFunctionalTest = new TestHarness("UE Functional Test");
+			public static readonly TestHarness Gauntlet = new TestHarness("Gauntlet");
+			public static readonly TestHarness GauntletRPC = new TestHarness("Gauntlet RPC");
+			public static readonly TestHarness LLT = new TestHarness("Low Level Test");
+		}
+
+		/// <summary>
+		/// Test team
+		/// </summary>
+		public partial class TestTeam
+		{
+			public static readonly TestTeam UE = new TestTeam("UE");
+		}
+	}
+
+	/// <summary>
+	/// Unreal Test metadata component for Audit
+	/// </summary>
+	public class UnrealTestMetadata : ITestAuditMetadata
+	{
+		public UnrealTestMetadata(string Key, string Name, string NodeClass, Version Version, float Timeout, string? CommandLine = null, TestHarness? Harness = null, TestProfile? Profile = null, TestTeam? Team = null, string? Intent = null, string[]? Tags = null, Dictionary<string, object>? Params = null)
+		{
+			_key = () => Key;
+			_name = () => Name;
+			_class = () => NodeClass;
+			_version =() => Version;
+			_timeout = () => Timeout;
+			_commandline = () => CommandLine;
+			_harness = () => Harness;
+			_profile = () => Profile;
+			_team = () => Team;
+			_intent = () => Intent;
+			_tags = () => Tags;
+			_params = () => Params;
+		}
+		public UnrealTestMetadata(Func<string> Key, Func<string> Name, Func<string> NodeClass, Func<Version> Version, Func<float> Timeout, Func<string?>? CommandLine = null, Func<TestHarness?>? Harness = null, Func<TestProfile?>? Profile = null, Func<TestTeam?>? Team = null, Func<string?>? Intent = null, Func<string[]?>? Tags = null, Func<Dictionary<string, object>>? Params = null)
+		{
+			_key = Key;
+			_name = Name;
+			_class = NodeClass;
+			_version = Version;
+			_timeout = Timeout;
+			_commandline = CommandLine;
+			_harness = Harness;
+			_profile = Profile;
+			_team = Team;
+			_intent = Intent;
+			_tags = Tags;
+			_params = Params;
+		}
+
+		Func<string> _key;
+		Func<string> _name;
+		Func<string> _class;
+		Func<Version> _version;
+		Func<float> _timeout;
+		Func<string?>? _commandline;
+		Func<TestHarness?>? _harness;
+		Func<TestProfile?>? _profile;
+		Func<TestTeam?>? _team;
+		Func<string?>? _intent;
+		Func<string[]?>? _tags;
+		Func<Dictionary<string, object>?>? _params;
+
+		static Regex AbsoluteFilePathPattern = new Regex(@"^(?:[/\\]+|[a-z]:[/\\]|https?://)", RegexOptions.Compiled | RegexOptions.IgnoreCase);
+
+		/* Unique identifier */
+		public string TestKey => _key();
+
+		/* Test name for display */
+		public string TestName => _name();
+
+		/* TestNode class */
+		public string TestNodeClass => _class();
+
+		/* Test version */
+		public Version TestVersion => _version();
+
+		/* Duration in minutes after which the test run has timeout */
+		public float TestTimeoutMinutes => _timeout();
+
+		/* Command line used to run the test */
+		public string? TestCommandLine => _commandline?.Invoke();
+
+		/* Test harness */
+		public TestHarness? TestHarness => _harness?.Invoke();
+
+		/* Test result profile */
+		public TestProfile? TestProfile => _profile?.Invoke();
+
+		/* Team to which the test was written for */
+		public TestTeam? TestTeam => _team?.Invoke();
+
+		/* The purpose of the test */
+		public string? TestIntent => _intent?.Invoke();
+
+		/* Tags associated with the test */
+		public string[]? TestTags => _tags?.Invoke();
+
+		/* Params relevant to run this test */
+		public Dictionary<string, object>? TestParams => _params?.Invoke()?.Where(P => P.Value is string ? !AbsoluteFilePathPattern.IsMatch((string)P.Value) : true).ToDictionary();
+	}
+}
