Add support for describing the results for a horde step in a small, readable way to show in step view, plus the ability to create an "Invalidating Event" during setup that will be bubbled up to horde even though the test node has not yet started.

diff --git a/Engine/Source/Programs/AutomationTool/Gauntlet/Framework/Gauntlet.TestExecutor.cs b/Engine/Source/Programs/AutomationTool/Gauntlet/Framework/Gauntlet.TestExecutor.cs
--- a/Engine/Source/Programs/AutomationTool/Gauntlet/Framework/Gauntlet.TestExecutor.cs
+++ b/Engine/Source/Programs/AutomationTool/Gauntlet/Framework/Gauntlet.TestExecutor.cs
@@ -270,6 +270,7 @@ namespace Gauntlet
 									Node.AddTestEvent(new UnrealTestEvent(EventSeverity.Error, "Test Failed to Start", new List<string> { ex.Message }));
 									PendingTests[i] = null;
 									NodeInfo.TimeSetupBegan = NodeInfo.TimeSetupEnded = NodeInfo.TimeTestEnded = DateTime.Now;
+									Horde.ReportInvalidatingIssue("FailedToStart", "The test encountered a failure on startup.");
 									NodeInfo.Result = TestExecutionInfo.ExecutionResult.Interrupted;
 									CompletedTests.Add(NodeInfo);
 								}
@@ -308,6 +309,7 @@ namespace Gauntlet
 										PendingTests[i] = null;
 										NodeInfo.TimeSetupBegan = NodeInfo.TimeSetupEnded = NodeInfo.TimeTestEnded = DateTime.Now;
 										NodeInfo.Result = TestExecutionInfo.ExecutionResult.NotStarted;
+										Horde.ReportInvalidatingIssue("NotEnoughDevices", "The test was unable to secure enough devices to run.");
 										CompletedTests.Add(NodeInfo);
 									}
 								}
diff --git a/Engine/Source/Programs/AutomationTool/Gauntlet/Framework/Horde/Gauntlet.Horde.cs b/Engine/Source/Programs/AutomationTool/Gauntlet/Framework/Horde/Gauntlet.Horde.cs
--- a/Engine/Source/Programs/AutomationTool/Gauntlet/Framework/Horde/Gauntlet.Horde.cs
+++ b/Engine/Source/Programs/AutomationTool/Gauntlet/Framework/Horde/Gauntlet.Horde.cs
@@ -1,6 +1,7 @@
 // Copyright Epic Games, Inc. All Rights Reserved.
 
 using AutomationTool;
+using EpicGames.Horde.Jobs;
 using System;
 using System.Collections.Generic;
 using System.IO;
@@ -37,6 +38,60 @@ namespace Gauntlet
 			}
 		}
 
+		public enum JobStepMessageLevels
+		{
+			None,
+			Info,
+			Warning,
+			Error
+		}
+
+		public class JobStepMessage
+		{ 
+			public string Message { get; set; } = string.Empty;
+			public string Severity { get; set; } = JobStepMessageLevels.None.ToString();
+		}
+
+		private static JobStepMessage CurStepMessage = new JobStepMessage();
+		/// <summary>
+		/// Update the step message (short message written in-line in step view)
+		/// These messages should be generic so we can aggregate multiple instances of it.
+		/// Good message example "TestTimedOut", or "IslandNotFound", or "LoginFailed"
+		/// Bad message example "This test timed out after 3 hours", "7 of 8 islands found"
+		/// </summary>
+		/// <param name="inMessage"></param>
+		/// <param name="inSeverity"></param>
+		public static void UpdateStepMessage(string inMessage, JobStepMessageLevels inSeverity)
+		{
+			CurStepMessage.Message = inMessage;
+			CurStepMessage.Severity = inSeverity.ToString();
+		}
+
+		/// <summary>
+		/// Labels a test as being invalid, and provides reason why to the summary report, with a concise message to also be shown in step view. This should ONLY
+		/// be used right before an exception is thrown. If it is being used outside of a currently running test (as in, the test is actively ticking), you want to set
+		/// bGenerateSummary to true
+		/// </summary>
+		/// <param name="IssueType">Describes the type of issue (ex. NotEnoughDevices) - should be general so it can be shown in one line</param>
+		/// <param name="IssueDescription">In-depth description of the issue. If the issue type was NotEnoughDevices, for example, this is where you would say what sort of device
+		/// you were looking for.</param>
+		/// <param name="bGenerateSummary"> Whether or not we should immediately generate the summary log. If this is being called from outside of a test node, set to yes.</param>
+		public static void ReportInvalidatingIssue(string IssueType, string IssueDescription, bool bGenerateSummary = true)
+		{
+			string InvalidatingIssueString = $"**Invalidating issue found:** {IssueType}\n";
+			OnStepSummary += () => InvalidatingIssueString;
+			if (!string.IsNullOrEmpty(IssueDescription))
+			{
+				OnStepSummary += () => $"More details: {IssueDescription}";
+			}
+
+			UpdateStepMessage(IssueType, JobStepMessageLevels.Error);
+			if (bGenerateSummary)
+			{
+				GenerateSummary();
+			}
+		}
+
 		/// <summary>
 		/// Whether to generate test data v1 model when writing test results for Horde
 		/// </summary>
@@ -54,8 +109,10 @@ namespace Gauntlet
 			public string name { get; set; }
 			public string placement { get; } = "Summary";
 			public string fileName { get; set; }
-		}
+			public JobStepMessage stepMessage { get; set; }
 
+		}
+		static bool bHasGeneratedStepReport = false;
 
 		public delegate string SummaryLineDelegate();
 		/// <summary>
@@ -109,6 +166,12 @@ namespace Gauntlet
 					return;
 				}
 
+				if (bHasGeneratedStepReport)
+				{
+					Log.Info("Trying to generate a step report when we have already generated one. ");
+					return;
+				}
+
 				StringBuilder Markdown = new StringBuilder();
 
 				if (OnStepSummary != null)
@@ -118,7 +181,6 @@ namespace Gauntlet
 						Markdown.AppendLine(callback());
 					}
 				}
-
 				if (OnStepSummaryFinalization != null)
 				{
 					foreach (SummaryLineDelegate callback in OnStepSummaryFinalization.GetInvocationList())
@@ -148,9 +210,10 @@ namespace Gauntlet
 
 					File.WriteAllText(Path.Combine(LogFolder, MarkdownFilename), Markdown.ToString());
 
-					JobStepReport report = new JobStepReport() { name = "Gauntlet Step Details", fileName = MarkdownFilename };
+					JobStepReport report = new JobStepReport() { name = "Gauntlet Step Details", fileName = MarkdownFilename, stepMessage = CurStepMessage };
 					File.WriteAllText(Path.Combine(LogFolder, "GauntletStepDetails.report.json"), JsonSerializer.Serialize(report));
 				}
+				bHasGeneratedStepReport = true;
 			}
 			catch (Exception Ex)
 			{
diff --git a/Engine/Source/Programs/AutomationTool/Gauntlet/Unreal/Base/Gauntlet.UnrealSession.cs b/Engine/Source/Programs/AutomationTool/Gauntlet/Unreal/Base/Gauntlet.UnrealSession.cs
--- a/Engine/Source/Programs/AutomationTool/Gauntlet/Unreal/Base/Gauntlet.UnrealSession.cs
+++ b/Engine/Source/Programs/AutomationTool/Gauntlet/Unreal/Base/Gauntlet.UnrealSession.cs
@@ -693,6 +693,7 @@ namespace Gauntlet
 			if (!CheckRolesArePossible(ref ValidationIssues))
 			{
 				ValidationIssues.ForEach(S => Log.Error(KnownLogEvents.Gauntlet_BuildDropEvent, S));
+				Horde.ReportInvalidatingIssue("BuildNotFound", "The test was unable to find the necessary executables.");
 				throw new AutomationException("One or more issues occurred when validating build {0} against requested roles", InSource.BuildName);
 			}
 
@@ -856,8 +857,8 @@ namespace Gauntlet
 					{
 						return null;
 					}
-
 					// If device reservation fails, the device pool cannot support this launch.
+					Horde.ReportInvalidatingIssue("NotEnoughDevices", "The test was unable to secure enough devices to run.");
 					throw new AutomationException("Failed to acquire all devices for launch. See above for details.");
 				}
 
@@ -1469,6 +1470,7 @@ namespace Gauntlet
 
 					if (SessionRetries == 0)
 					{
+						Horde.ReportInvalidatingIssue("ApplicationInstallFailed", "Unable to install application for session.");
 						throw new AutomationException("Unable to install application for session.");
 					}
 
diff --git a/Engine/Source/Programs/AutomationTool/Gauntlet/Unreal/RunUnreal.cs b/Engine/Source/Programs/AutomationTool/Gauntlet/Unreal/RunUnreal.cs
--- a/Engine/Source/Programs/AutomationTool/Gauntlet/Unreal/RunUnreal.cs
+++ b/Engine/Source/Programs/AutomationTool/Gauntlet/Unreal/RunUnreal.cs
@@ -116,6 +116,13 @@ namespace Gauntlet
 				Gauntlet.Log.Level = Gauntlet.LogLevel.VeryVerbose;
 			}
 
+			// Use this flag if iterating in horde to get a test immediately declared invalid.
+			if (ParseParam("forcestepfail"))
+			{
+				string ErrorMessage = "-forcestepfail specified on commandline.";
+				Horde.ReportInvalidatingIssue("TestCategory", ErrorMessage);
+				throw new AutomationException(ErrorMessage);
+			}
 			if (ParseParam("log"))
 			{
 				if (!Directory.Exists(ContextOptions.LogDir))
@@ -134,12 +141,16 @@ namespace Gauntlet
 
 			if (string.IsNullOrEmpty(ContextOptions.Build))
 			{
-				throw new AutomationException("No build specified. Use -build=p:\\path\\to\\build");
+				string ErrorMessage = "No build specified. Use -build=p:\\path\\to\\build";
+				Horde.ReportInvalidatingIssue("BuildNotSpecified", ErrorMessage);
+				throw new AutomationException(ErrorMessage);
 			}
 
 			if (typeof(UnrealBuildSource).IsAssignableFrom(ContextOptions.BuildSourceType) == false)
 			{
-				throw new AutomationException("Provided BuildSource type does not inherit from UnrealBuildSource");
+				string ErrorMessage = "Provided BuildSource type does not inherit from UnrealBuildSource";
+				Horde.ReportInvalidatingIssue("InvalidBuildSource", ErrorMessage);
+				throw new AutomationException(ErrorMessage);
 			}
 
 			// make -test=none implicit if no test is supplied
@@ -423,7 +434,9 @@ namespace Gauntlet
 				EPerfSpec PerfSpec;
 				if (!Enum.TryParse<EPerfSpec>(PerfSpecArg, true, out PerfSpec))
 				{
-					throw new AutomationException("Unable to convert perfspec '{0}' into an EPerfSpec", PerfSpec);
+					string ErrorMessage = string.Format("Unable to convert perfspec '{0}' into an EPerfSpec", PerfSpec);
+					Horde.ReportInvalidatingIssue("InvalidPerfSpec", ErrorMessage);
+					throw new AutomationException(ErrorMessage);
 				}
 
 				// parse hardware model
