Fix failure to serialize null object in Horde test report

#rnx

diff --git a/Engine/Source/Programs/AutomationTool/Gauntlet/Framework/Utils/Gauntlet.HordeReport.cs b/Engine/Source/Programs/AutomationTool/Gauntlet/Framework/Utils/Gauntlet.HordeReport.cs
--- a/Engine/Source/Programs/AutomationTool/Gauntlet/Framework/Utils/Gauntlet.HordeReport.cs
+++ b/Engine/Source/Programs/AutomationTool/Gauntlet/Framework/Utils/Gauntlet.HordeReport.cs
@@ -2040,9 +2040,9 @@ namespace Gauntlet
 			{
 				public Utf8String Type { get; protected set; }
 				public string Text { get; protected set; }
-				public Dictionary<Utf8String, object> Properties { get; protected set; }
+				public Dictionary<Utf8String, object?> Properties { get; protected set; }
 
-				public TestLogValue(string InType, string InText, Dictionary<Utf8String, object> InProperties)
+				public TestLogValue(string InType, string InText, Dictionary<Utf8String, object?> InProperties)
 				{
 					Type = new Utf8String(InType);
 					Text = InText;
@@ -2059,7 +2059,7 @@ namespace Gauntlet
 				/// </summary>
 				public static TestLogValue CreateURLLink(string Text, string URL)
 				{
-					return new TestLogValue(TestLogValueTypes.url, Text, new Dictionary<Utf8String, object>() { { new Utf8String("href"), URL } });
+					return new TestLogValue(TestLogValueTypes.url, Text, new Dictionary<Utf8String, object?>() { { new Utf8String("href"), URL } });
 				}
 			}
 
@@ -2076,10 +2076,10 @@ namespace Gauntlet
 					writer.WriteString(TestLogValueTypes.Text, entry.ToString());
 					if (entry.Properties != null)
 					{
-						foreach (KeyValuePair<Utf8String, object> pair in entry.Properties)
+						foreach (KeyValuePair<Utf8String, object?> pair in entry.Properties)
 						{
 							writer.WritePropertyName(pair.Key);
-							writer.WriteStringValue(pair.Value.ToString());
+							writer.WriteStringValue(pair.Value?.ToString());
 						}
 					}
 					writer.WriteEndObject();
@@ -2153,7 +2153,7 @@ namespace Gauntlet
 
 				public Dictionary<string, object?> ToProperties()
 				{
-					Dictionary<Utf8String, object> Prop = new Dictionary<Utf8String, object>();
+					Dictionary<Utf8String, object?> Prop = new Dictionary<Utf8String, object?>();
 					foreach(var Item in Files)
 					{
 						if (!string.IsNullOrEmpty(Item.Value))
