[Gauntlet] Fixes typo in TargetDeviceIOS.CleanArtifacts that was preventing project artifacts from being properly cleaned

diff --git a/Engine/Source/Programs/AutomationTool/Gauntlet/Platform/IOS/Gauntlet.TargetDeviceIOS.cs b/Engine/Source/Programs/AutomationTool/Gauntlet/Platform/IOS/Gauntlet.TargetDeviceIOS.cs
--- a/Engine/Source/Programs/AutomationTool/Gauntlet/Platform/IOS/Gauntlet.TargetDeviceIOS.cs
+++ b/Engine/Source/Programs/AutomationTool/Gauntlet/Platform/IOS/Gauntlet.TargetDeviceIOS.cs
@@ -461,7 +461,7 @@ namespace Gauntlet
 			}
 
 			// The same applies here for the --list
-			if (Result.Output.Contains(ListProjectSaved))
+			if (Result.Output.Contains(ProjectSavedDirectory))
 			{
 				Result = ExecuteIOSDeployCommand(CleanProject);
 				Result.WaitForExit();
