Modifying the ProcessCloudContent build step to allow content to be directed to multiple artifacts specified through the install bundle config.
#FailedValidations Preflight

diff --git a/Engine/Source/Programs/AutomationTool/AutomationUtils/BundleUtils.cs b/Engine/Source/Programs/AutomationTool/AutomationUtils/BundleUtils.cs
--- a/Engine/Source/Programs/AutomationTool/AutomationUtils/BundleUtils.cs
+++ b/Engine/Source/Programs/AutomationTool/AutomationUtils/BundleUtils.cs
@@ -26,6 +26,7 @@ namespace AutomationUtils.Automation
 			public bool	UseChunkDBs { get; set; }
 			public bool UseDetailedInstallSizes { get; set; } // default is true
 			public string ExecFileName { get; set; } // TODO: We never used this.  Clean this up.
+			public string CloudAppNameOverride { get; set; }
 		}
 
 		public static void LoadBundleConfig(DirectoryReference ProjectDir, UnrealTargetPlatform Platform, string CustomConfig, out IReadOnlyDictionary<string, BundleSettings> Bundles)
@@ -141,6 +142,16 @@ namespace AutomationUtils.Automation
 					}
 
 				}
+				{
+					if (BundleConfig.GetString(SectionName, "CloudAppNameOverride", out string CloudAppNameOverride))
+					{
+						Bundle.CloudAppNameOverride = CloudAppNameOverride;
+					}
+					else 
+					{
+						Bundle.CloudAppNameOverride = "";
+					}
+				}
 				GetPlatformSettings(Bundle, BundleConfig, BundleDefinitionPrefix + Bundle.Name);
 
 				Results.Add(Bundle);
@@ -175,6 +186,17 @@ namespace AutomationUtils.Automation
 			return null;
 		}
 
+		public static IEnumerable<string> GetCloudAppsForBundles<TPlatformBundleSettings>(IReadOnlyDictionary<string, TPlatformBundleSettings> InstallBundles) where TPlatformBundleSettings : BundleSettings
+		{
+			// Gather all bundle clouda app name overrides
+			var CloudAppNames = new HashSet<string>();
+			foreach (var Bundle in InstallBundles.Values)
+			{
+				CloudAppNames.Add(Bundle.CloudAppNameOverride);
+			}
+			return CloudAppNames;
+		}
+
 		public static IEnumerable<TPlatformBundleSettings> GetBundleDependencies<TPlatformBundleSettings>(
 			TPlatformBundleSettings Bundle, IReadOnlyDictionary<string, TPlatformBundleSettings> InstallBundles) where TPlatformBundleSettings : BundleSettings
 		{
