Add -buildmachine commandline argument when running DerivedDataTool.

diff --git a/Engine/Source/Programs/AutomationTool/BuildGraph/Tasks/DDCReplayLoadTask.cs b/Engine/Source/Programs/AutomationTool/BuildGraph/Tasks/DDCReplayLoadTask.cs
--- a/Engine/Source/Programs/AutomationTool/BuildGraph/Tasks/DDCReplayLoadTask.cs
+++ b/Engine/Source/Programs/AutomationTool/BuildGraph/Tasks/DDCReplayLoadTask.cs
@@ -94,9 +94,10 @@ namespace AutomationTool.Tasks
 			if (_parameters.DDCReplayFile.StartsWith("#", StringComparison.OrdinalIgnoreCase))
 			{
 				HashSet<FileReference> files = ResolveFilespec(Unreal.RootDirectory, _parameters.DDCReplayFile, tagNameToFileSet);
+				string commonArgs = String.Format("{0}{1}", _parameters.Arguments, CommandUtils.IsBuildMachine ? " -buildmachine" : "");
 				foreach (FileReference file in files)
 				{
-					string taskArguments = $"{_parameters.Arguments} -DDC-ReplayLoad={file}";
+					string taskArguments = $"{commonArgs} -DDC-ReplayLoad={file}";
 					await ExecuteAsync(_parameters.Exe, taskArguments, _parameters.WorkingDir, envVars: ParseEnvVars(_parameters.Environment, _parameters.EnvironmentFile), logOutput: _parameters.LogOutput, errorLevel: _parameters.ErrorLevel);
 				}
 			}
