Make step message fields nullable, do not serialize them if they are null

diff --git a/Engine/Source/Programs/AutomationTool/Gauntlet/Framework/Horde/Gauntlet.Horde.cs b/Engine/Source/Programs/AutomationTool/Gauntlet/Framework/Horde/Gauntlet.Horde.cs
--- a/Engine/Source/Programs/AutomationTool/Gauntlet/Framework/Horde/Gauntlet.Horde.cs
+++ b/Engine/Source/Programs/AutomationTool/Gauntlet/Framework/Horde/Gauntlet.Horde.cs
@@ -8,6 +8,7 @@ using System.IO;
 using System.Linq;
 using System.Text;
 using System.Text.Json;
+using System.Text.Json.Serialization;
 
 namespace Gauntlet
 {
@@ -46,13 +47,12 @@ namespace Gauntlet
 			Error
 		}
 
-		public class JobStepMessage
-		{ 
-			public string Message { get; set; } = string.Empty;
-			public string Severity { get; set; } = JobStepMessageLevels.None.ToString();
-		}
 
-		private static JobStepMessage CurStepMessage = new JobStepMessage();
+#nullable enable
+		private static string? curStepMessage { get; set; } = null;
+		private static string? curStepSeverity { get; set; } = null;
+#nullable disable
+
 		/// <summary>
 		/// Update the step message (short message written in-line in step view)
 		/// These messages should be generic so we can aggregate multiple instances of it.
@@ -63,8 +63,8 @@ namespace Gauntlet
 		/// <param name="inSeverity"></param>
 		public static void UpdateStepMessage(string inMessage, JobStepMessageLevels inSeverity)
 		{
-			CurStepMessage.Message = inMessage;
-			CurStepMessage.Severity = inSeverity.ToString();
+			curStepMessage = inMessage;
+			curStepSeverity = inSeverity.ToString();
 		}
 
 		/// <summary>
@@ -109,7 +109,11 @@ namespace Gauntlet
 			public string name { get; set; }
 			public string placement { get; } = "Summary";
 			public string fileName { get; set; }
-			public JobStepMessage stepMessage { get; set; }
+#nullable enable
+			public string? message { get; set; }
+			public string? severity { get; set; }
+
+#nullable disable
 
 		}
 		static bool bHasGeneratedStepReport = false;
@@ -209,9 +213,12 @@ namespace Gauntlet
 					string MarkdownFilename = "GauntletStepDetails.md";
 
 					File.WriteAllText(Path.Combine(LogFolder, MarkdownFilename), Markdown.ToString());
-
-					JobStepReport report = new JobStepReport() { name = "Gauntlet Step Details", fileName = MarkdownFilename, stepMessage = CurStepMessage };
-					File.WriteAllText(Path.Combine(LogFolder, "GauntletStepDetails.report.json"), JsonSerializer.Serialize(report));
+					JsonSerializerOptions options = new()
+					{
+						DefaultIgnoreCondition = JsonIgnoreCondition.WhenWritingNull
+					};
+					JobStepReport report = new JobStepReport() { name = "Gauntlet Step Details", fileName = MarkdownFilename, message = curStepMessage, severity = curStepSeverity };
+					File.WriteAllText(Path.Combine(LogFolder, "GauntletStepDetails.report.json"), JsonSerializer.Serialize(report, options));
 				}
 				bHasGeneratedStepReport = true;
 			}
