Adding test step messages to UE.BootTest for use in Horde iteration.

diff --git a/Engine/Source/Programs/AutomationTool/Gauntlet/Unreal/Automation/UE.BootTest.cs b/Engine/Source/Programs/AutomationTool/Gauntlet/Unreal/Automation/UE.BootTest.cs
--- a/Engine/Source/Programs/AutomationTool/Gauntlet/Unreal/Automation/UE.BootTest.cs
+++ b/Engine/Source/Programs/AutomationTool/Gauntlet/Unreal/Automation/UE.BootTest.cs
@@ -67,7 +67,7 @@ namespace UE
 			{
 				kTimeOutDuration = logIdleTimeout;
 			}
-
+			Gauntlet.Horde.UpdateStepMessage("DefaultStepMessage", Gauntlet.Horde.JobStepMessageLevels.Info);
 			return Config;
 		}
 
@@ -165,6 +165,7 @@ namespace UE
 			{
 				if (!string.IsNullOrEmpty(GetCompletionString()) && !DidDetectLaunch)
 				{
+					Gauntlet.Horde.UpdateStepMessage("LaunchNotCompleted", Gauntlet.Horde.JobStepMessageLevels.Error);
 					ReportError("Failed to detect completion of launch");
 				}
 				else
@@ -177,12 +178,14 @@ namespace UE
 						if (!File.Exists(RoleResult.Artifacts.LogPath))
 						{
 							MissingLogs = true;
+							Gauntlet.Horde.UpdateStepMessage("LogsNotFound", Gauntlet.Horde.JobStepMessageLevels.Error);
 							ReportError("No log files found for {0}. Were they not retrieved from the device?", RoleResult.Artifacts.SessionRole);
 						}
 					}
 
 					if (!MissingLogs)
 					{
+						Gauntlet.Horde.UpdateStepMessage("CreateReportSucceeded", Gauntlet.Horde.JobStepMessageLevels.Info);
 						Log.Info("Found valid log artifacts for test");
 					}
 				}
