UAT: Implement/resurrect the "-CompileOnly" command line option documented in the usage info.
Enable SubmitTool to invoke UAT once with "-CompileOnly" to compile all out-of-date script modules before running multiple instances in parallel with the command line option `-NoCompileUAT` and the environment variable "uebp_UATMutexNoWait=1".

diff --git a/Engine/Source/Programs/AutomationTool/AutomationUtils/Automation.cs b/Engine/Source/Programs/AutomationTool/AutomationUtils/Automation.cs
--- a/Engine/Source/Programs/AutomationTool/AutomationUtils/Automation.cs
+++ b/Engine/Source/Programs/AutomationTool/AutomationUtils/Automation.cs
@@ -96,10 +96,17 @@ namespace AutomationTool
 				ShouldKillProcesses = !GlobalCommandLine.NoKill;
 				Logger.LogDebug("ShouldKillProcesses={ShouldKillProcesses}", ShouldKillProcesses);
 
-				if (AutomationToolCommandLine.CommandsToExecute.Count == 0 && GlobalCommandLine.Help)
+				if (AutomationToolCommandLine.CommandsToExecute.Count == 0)
 				{
-					DisplayHelp(AutomationToolCommandLine.GlobalParameters);
-					return ExitCode.Success;
+					if (GlobalCommandLine.Help)
+					{
+						DisplayHelp(AutomationToolCommandLine.GlobalParameters);
+						return ExitCode.Success;
+					}
+					if (GlobalCommandLine.CompileOnly)
+					{
+						return ExitCode.Success;
+					}
 				}
 
 				// Disable AutoSDKs if specified on the command line
@@ -339,7 +346,7 @@ namespace AutomationTool
 			HelpUtils.PrintHelp("Automation Help:",
 @"Executes scripted commands
 
-AutomationTool.exe [-verbose] [-compileonly] [-p4] Command0 [-Arg0 -Arg1 -Arg2 ...] Command1 [-Arg0 -Arg1 ...] Command2 [-Arg0 ...] Commandn ... [EnvVar0=MyValue0 ... EnvVarn=MyValuen]",
+AutomationTool.exe [-verbose] [-Help|-List|-CompileOnly] [-p4] Command0 [-Arg0 -Arg1 -Arg2 ...] Command1 [-Arg0 -Arg1 ...] Command2 [-Arg0 ...] CommandN ... [EnvVar0=MyValue0 ... EnvVarN=MyValueN]",
 				ParamDict.ToList());
 			CommandUtils.LogHelp(typeof(Automation));
 		}
@@ -419,6 +426,7 @@ AutomationTool.exe [-verbose] [-compileonly] [-p4] Command0 [-Arg0 -Arg1 -Arg2 .
 			Help = AutomationToolCommandLine.IsSetGlobal("-Help");
 			NoAutoSDK = AutomationToolCommandLine.IsSetGlobal("-NoAutoSDK");
 			List = AutomationToolCommandLine.IsSetGlobal("-List");
+			CompileOnly = AutomationToolCommandLine.IsSetGlobal("-CompileOnly");
 			TelemetryPath = (string)AutomationToolCommandLine.GetValueUnchecked("-Telemetry");
 			Verbose = AutomationToolCommandLine.IsSetGlobal("-Verbose");
 			AllowStdOutLogVerbosity = AutomationToolCommandLine.IsSetGlobal("-AllowStdOutLogVerbosity");
@@ -475,6 +483,13 @@ AutomationTool.exe [-verbose] [-compileonly] [-p4] Command0 [-Arg0 -Arg1 -Arg2 .
 			private set { bList = value; }
 		}
 		
+		private static bool? bCompileOnly;
+		public static bool CompileOnly
+		{
+			get => bCompileOnly.Value; 
+			private set { bCompileOnly = value; }
+		}
+		
 		private static bool? bVerbose;
 		public static bool Verbose
 		{
diff --git a/Engine/Source/Programs/AutomationTool/Program.cs b/Engine/Source/Programs/AutomationTool/Program.cs
--- a/Engine/Source/Programs/AutomationTool/Program.cs
+++ b/Engine/Source/Programs/AutomationTool/Program.cs
@@ -286,6 +286,7 @@ namespace AutomationToolDriver
 					{"-UTF8Output", ""},
 					{"-AllowStdOutLogVerbosity", ""},
 					{"-NoAutoSDK", ""},
+					{"-CompileOnly", "Run without a command to compile all out-of-date script modules."},
 					{"-Compile", "Force all script modules to be compiled"},
 					{"-NoCompile", "Do not attempt to compile any script modules - attempts to run with whatever is up to date" },
 					{"-IgnoreBuildRecords", "Ignore build records (Intermediate/ScriptModule/ProjectName.json) files when determining if script modules are up to date" },
@@ -322,7 +323,11 @@ namespace AutomationToolDriver
 			}
 
 			// Validate
-			var Result = AutomationToolCommandLine.CommandsToExecute.Count > 0 || AutomationToolCommandLine.IsSetGlobal("-Help") || AutomationToolCommandLine.IsSetGlobal("-List");
+			var Result =
+				AutomationToolCommandLine.CommandsToExecute.Count > 0 ||
+				AutomationToolCommandLine.IsSetGlobal("-Help") ||
+				AutomationToolCommandLine.IsSetGlobal("-List") ||
+				AutomationToolCommandLine.IsSetGlobal("-CompileOnly");
 			if (AutomationToolCommandLine.CommandsToExecute.Count > 0)
 			{
 				Logger.LogDebug("Found {NumScripts} scripts to execute:", AutomationToolCommandLine.CommandsToExecute.Count);
