LaunchMultiServer: Add '-DirectClientCount=##' to specify the number of clients to run when using the '-Client=Direct' feature.

The default is still to run a single client per launched server, but this allows you to run fewer clients if you'd like.

diff --git a/Engine/Source/Programs/AutomationTool/Scripts/LaunchMultiServer.Automation.cs b/Engine/Source/Programs/AutomationTool/Scripts/LaunchMultiServer.Automation.cs
--- a/Engine/Source/Programs/AutomationTool/Scripts/LaunchMultiServer.Automation.cs
+++ b/Engine/Source/Programs/AutomationTool/Scripts/LaunchMultiServer.Automation.cs
@@ -254,6 +254,9 @@ Example running 4 clients locally connecting to single proxy server w/ a non-loc
 		[CommandLine("-Client", ListSeparator = ','), Help("Client=<Type>", "Specify the client type to connect with (proxy or direct or proxy,direct)")]
 		public List<string> ClientType = new List<string>();
 
+		[CommandLine, Help("DirectClientCount=##", "Specify how many clients to launch in client=direct mode.  If 0 (default), will launch one for each server.")]
+		public int DirectClientCount = 0;
+
 		[CommandLine, Help("ProxyServerCount=##", "Specify how many proxy servers to launch in client=proxy mode")]
 		public int ProxyServerCount = 1;
 
@@ -350,8 +353,9 @@ Example running 4 clients locally connecting to single proxy server w/ a non-loc
 			{
 				if (ClientValue.ToLower() == "direct")
 				{
-					Logger.LogInformation("Starting direct client instances connecting to MultiServer instances");
-					StartProcessesForClients(ProcessLaunchType.DirectClient, AllGameServerAddresses, CommonArgs);
+                    int MaxClientCount = (DirectClientCount == 0) ? AllGameServerAddresses.Length : DirectClientCount;
+                    Logger.LogInformation($"Starting {MaxClientCount} direct client instances connecting to MultiServer instances");
+                    StartProcessesForClients(ProcessLaunchType.DirectClient, AllGameServerAddresses, CommonArgs, 1, MaxClientCount);
 				}
 				else if (ClientValue.ToLower() == "proxy")
 				{
@@ -482,16 +486,18 @@ Example running 4 clients locally connecting to single proxy server w/ a non-loc
 			return LocalGameServerAddresses;
 		}
 
-		/// <summary>
-		/// Start the Client processes
-		/// </summary>
-		/// <param name="ClientType">What type of Client we're launching which will help set the titles properly</param>
-		/// <param name="GameServerAddresses">The game servers (or proxies) that a client will connect to.  The addresses may contain a port range (e.g. IpAddress:StartPort-EndPort).</param>
-		/// <param name="AdditionalClientArguments">What additional parameters we launch the client with</param>
-		/// <param name="ClientsPerServer">The number of clients to launch per-GameServerAddress entry</param>
-		virtual protected void StartProcessesForClients(ProcessLaunchType ClientType, string[] GameServerAddresses, string AdditionalClientArguments, int ClientsPerServer = 1)
+        /// <summary>
+        /// Start the Client processes
+        /// </summary>
+        /// <param name="ClientType">What type of Client we're launching which will help set the titles properly</param>
+        /// <param name="GameServerAddresses">The game servers (or proxies) that a client will connect to.  The addresses may contain a port range (e.g. IpAddress:StartPort-EndPort).</param>
+        /// <param name="AdditionalClientArguments">What additional parameters we launch the client with</param>
+        /// <param name="ClientsPerServer">The number of clients to launch per-GameServerAddress entry</param>
+        /// <param name="MaxClientCount">The maximum number of clients to launch.  Zero for no limit (one for each server address or proxy server port)</param>
+        virtual protected void StartProcessesForClients(ProcessLaunchType ClientType, string[] GameServerAddresses, string AdditionalClientArguments, int ClientsPerServer = 1, int MaxClientCount = 0)
 		{
-			for (int Idx= 0; Idx < GameServerAddresses.Length; ++Idx)
+            int LaunchedClientCount = 0;
+            for (int Idx= 0; Idx < GameServerAddresses.Length; ++Idx)
 			{
 				string GameServerAddressWithPortRange = GameServerAddresses[Idx];
 				string GameServerBaseAddress;
@@ -519,9 +525,20 @@ Example running 4 clients locally connecting to single proxy server w/ a non-loc
 						string InstanceDetails = (ClientsPerServer > 1) ? $"(Srv {Idx + 1})" : string.Empty;
 						string AppTitle = $"{ClientType} {ClientId} {InstanceDetails}";
 						ExecuteApp(ClientType, AppTitle, InstanceArgs);
+
+                        ++LaunchedClientCount;
+                        if (MaxClientCount != 0 && LaunchedClientCount >= MaxClientCount)
+                        {
+                            break;
+                        }
 					}
-				}
-			}
+                
+                    if (MaxClientCount != 0 && LaunchedClientCount >= MaxClientCount)
+                    {
+                        break;
+                    }
+                }
+            }
 		}
 
 		virtual protected void StartProcessesForProxyServer(string[] LocalProxyServerAddresses, string[] AllGameServerAddresses, string AdditionalProxyArguments)
