Fix Framework Gauntlet self tests

#rnx

diff --git a/Engine/Source/Programs/AutomationTool/Gauntlet/SelfTest/Framework/Gauntlet.SelfTest.OrderOfOpsTest.cs b/Engine/Source/Programs/AutomationTool/Gauntlet/SelfTest/Framework/Gauntlet.SelfTest.OrderOfOpsTest.cs
--- a/Engine/Source/Programs/AutomationTool/Gauntlet/SelfTest/Framework/Gauntlet.SelfTest.OrderOfOpsTest.cs
+++ b/Engine/Source/Programs/AutomationTool/Gauntlet/SelfTest/Framework/Gauntlet.SelfTest.OrderOfOpsTest.cs
@@ -18,9 +18,10 @@ namespace Gauntlet.SelfTest
 		TickTest,
 		Status,
 		StopTest,
-		Result,
-		Summary,
+		GetResult,
+		SetFinalResult,
 		CleanupTest,
+		Summary
 	};
 
 	/// <summary>
@@ -62,6 +63,8 @@ namespace Gauntlet.SelfTest
 
 		protected void SetNewStage(TestStages NewStage)
 		{
+			Log.Info(" >> Executing {Stage}", NewStage.ToString());
+
 			int CurrentValue = (int)CurrentStage;
 			int NewValue = (int)NewStage;
 
@@ -81,12 +84,12 @@ namespace Gauntlet.SelfTest
 
 		public TestResult GetTestResult()
 		{
-			// this is the last stage and can hit multiple times.
+			// this is a stage that can hit multiple times.
 			// we can just return Passed because SetNewStage will 
 			// throw an exception if something is wrong here
-			if (CurrentStage != TestStages.Result)
+			if (CurrentStage != TestStages.GetResult)
 			{
-				SetNewStage(TestStages.Result);
+				SetNewStage(TestStages.GetResult);
 			}
 			return TestResult.Passed;
 		}
@@ -96,12 +99,13 @@ namespace Gauntlet.SelfTest
 		/// </summary>
 		public void SetTestResult(TestResult testResult)
 		{
+			SetNewStage(TestStages.SetFinalResult);
 		}
 
 		public string GetTestSummary()
 		{
 			SetNewStage(TestStages.Summary);
-			return "";
+			return "Summary";
 		}
 
 		public virtual void AddTestEvent(UnrealTestEvent InEvent)
diff --git a/Engine/Source/Programs/AutomationTool/Gauntlet/SelfTest/Framework/Gauntlet.SelfTest.TestAutoParam.cs b/Engine/Source/Programs/AutomationTool/Gauntlet/SelfTest/Framework/Gauntlet.SelfTest.TestAutoParam.cs
--- a/Engine/Source/Programs/AutomationTool/Gauntlet/SelfTest/Framework/Gauntlet.SelfTest.TestAutoParam.cs
+++ b/Engine/Source/Programs/AutomationTool/Gauntlet/SelfTest/Framework/Gauntlet.SelfTest.TestAutoParam.cs
@@ -29,10 +29,10 @@ namespace Gauntlet.SelfTest
 		[AutoParam("StringDefault")]
 		public string StringParam;
 
-		[AutoParamWithNames("StringDefault", "StringParam")]
+		[AutoParamWithNames("StringDefault", OptionNames: "StringParam")]
 		public string StringProperty { get; set; }
 
-		[AutoParamWithNames("InternalDefault", "ExternalName", "OtherExternalName")]
+		[AutoParamWithNames("InternalDefault", OptionNames: ["ExternalName", "OtherExternalName"])]
 		public string InternalName;
 
 		[AutoParam(42)]
diff --git a/Engine/Source/Programs/AutomationTool/Gauntlet/SelfTest/Unreal/Gauntlet.SelfTest.TestUnrealBuildParsing.cs b/Engine/Source/Programs/AutomationTool/Gauntlet/SelfTest/Unreal/Gauntlet.SelfTest.TestUnrealBuildParsing.cs
--- a/Engine/Source/Programs/AutomationTool/Gauntlet/SelfTest/Unreal/Gauntlet.SelfTest.TestUnrealBuildParsing.cs
+++ b/Engine/Source/Programs/AutomationTool/Gauntlet/SelfTest/Unreal/Gauntlet.SelfTest.TestUnrealBuildParsing.cs
@@ -48,7 +48,7 @@ namespace Gauntlet.SelfTest
 				// Test a content project
 				new BuildData("ElementalDemo", "UnrealGame.exe", UnrealTargetConfiguration.Development, UnrealTargetPlatform.Win64, UnrealTargetRole.Client, true ),
 				new BuildData("ElementalDemo", "UnrealGame-Win64-Test.exe", UnrealTargetConfiguration.Test, UnrealTargetPlatform.Win64, UnrealTargetRole.Client, true ),
-				new BuildData("ElementalDemo", "UnrealGame-Win64-Shippinge.exe", UnrealTargetConfiguration.Shipping, UnrealTargetPlatform.Win64, UnrealTargetRole.Client, true ),
+				new BuildData("ElementalDemo", "UnrealGame-Win64-Shipping.exe", UnrealTargetConfiguration.Shipping, UnrealTargetPlatform.Win64, UnrealTargetRole.Client, true ),
 
 				// Test a regular monolithic project
 				new BuildData("ActionRPG", "ActionRPG.exe", UnrealTargetConfiguration.Development, UnrealTargetPlatform.Win64, UnrealTargetRole.Client ),
