Change StepMessage in BootTest for hordetesting

diff --git a/Engine/Source/Programs/AutomationTool/Gauntlet/Unreal/Automation/UE.BootTest.cs b/Engine/Source/Programs/AutomationTool/Gauntlet/Unreal/Automation/UE.BootTest.cs
--- a/Engine/Source/Programs/AutomationTool/Gauntlet/Unreal/Automation/UE.BootTest.cs
+++ b/Engine/Source/Programs/AutomationTool/Gauntlet/Unreal/Automation/UE.BootTest.cs
@@ -67,7 +67,7 @@ namespace UE
 			{
 				kTimeOutDuration = logIdleTimeout;
 			}
-			Gauntlet.Horde.UpdateStepMessage("DefaultStepMessage", Gauntlet.Horde.JobStepMessageLevels.Info);
+			Gauntlet.Horde.UpdateStepMessage("DidNotFinish", Gauntlet.Horde.JobStepMessageLevels.Error);
 			return Config;
 		}
 
