Extending timeout for WebDriverAgent install.
DeviceException in the case of an WDA install issue.

diff --git a/Engine/Source/Programs/AutomationTool/Gauntlet/Platform/IOS/Gauntlet.IOSAppium.cs b/Engine/Source/Programs/AutomationTool/Gauntlet/Platform/IOS/Gauntlet.IOSAppium.cs
--- a/Engine/Source/Programs/AutomationTool/Gauntlet/Platform/IOS/Gauntlet.IOSAppium.cs
+++ b/Engine/Source/Programs/AutomationTool/Gauntlet/Platform/IOS/Gauntlet.IOSAppium.cs
@@ -247,7 +247,7 @@ namespace Gauntlet
             }
             catch (Exception Ex)
             {
-                throw new AutomationException(Ex, "Failed to start appium instance for {0}: {1}", UUID, Ex.Message);
+                throw new DeviceException("Failed to start appium instance for {0}: {1}", UUID, Ex.Message);
             }
         }
 
