Fixed error message in LaunchMultiServer automation script

- When running with "-Client=Direct", fixed error text abourt "Invalid Game Server Address specified (null) should be in the format of IpAddress:Port or IpAddress:StartPort-EndPort"
- Now counts the correct number of addresses and sets the address array length correctly

diff --git a/Engine/Source/Programs/AutomationTool/Scripts/LaunchMultiServer.Automation.cs b/Engine/Source/Programs/AutomationTool/Scripts/LaunchMultiServer.Automation.cs
--- a/Engine/Source/Programs/AutomationTool/Scripts/LaunchMultiServer.Automation.cs
+++ b/Engine/Source/Programs/AutomationTool/Scripts/LaunchMultiServer.Automation.cs
@@ -344,18 +344,39 @@ Example running 4 clients locally connecting to single proxy server w/ a non-loc
 			string[] LocalGameServerAddresses = StartProcessesForDedicatedServers(NumLocalServers, NonLocalServerAddresses, CommonArgs);
 
 			// Generate all of the GameServerAddresses by merging the local & non-local servers
-			string[] AllGameServerAddresses = new string[TotalNumServers];
-			LocalGameServerAddresses.CopyTo(AllGameServerAddresses, 0);
-			NonLocalServerAddresses.CopyTo(AllGameServerAddresses, LocalGameServerAddresses.Length);
+			string[] AllGameServerAddresses = LocalGameServerAddresses.Concat(NonLocalServerAddresses).ToArray();
 
 			// Now start-up the clients
 			foreach (string ClientValue in ClientType)
 			{
 				if (ClientValue.ToLower() == "direct")
 				{
-                    int MaxClientCount = (DirectClientCount == 0) ? AllGameServerAddresses.Length : DirectClientCount;
-                    Logger.LogInformation($"Starting {MaxClientCount} direct client instances connecting to MultiServer instances");
-                    StartProcessesForClients(ProcessLaunchType.DirectClient, AllGameServerAddresses, CommonArgs, 1, MaxClientCount);
+					int MaxDirectClientCount = 0;
+					foreach (var GameServerAddressWithPortRange in AllGameServerAddresses)
+					{
+						string GameServerBaseAddress;
+						int StartPort, EndPort;
+						bool bValidAddress = ParseAddressAndPortRange(out GameServerBaseAddress, out StartPort, out EndPort, GameServerAddressWithPortRange);
+						if (bValidAddress)
+						{
+							MaxDirectClientCount += (EndPort - StartPort) + 1;
+						}
+					}
+
+					if (DirectClientCount != 0)
+					{
+						if (DirectClientCount > MaxDirectClientCount)
+						{
+							Logger.LogWarning($"Specified DirectClientCount=({DirectClientCount}) but only a maximum of {MaxDirectClientCount} addresses to connect to.");
+							DirectClientCount = MaxDirectClientCount;
+						}
+					}
+					else
+					{
+						DirectClientCount = MaxDirectClientCount;
+					}
+                    Logger.LogInformation($"Starting {DirectClientCount} direct client instances connecting to MultiServer instances");
+                    StartProcessesForClients(ProcessLaunchType.DirectClient, AllGameServerAddresses, CommonArgs, 1, DirectClientCount);
 				}
 				else if (ClientValue.ToLower() == "proxy")
 				{
