Gauntlet support for not setting up mcp parameters in the test code.

diff --git a/Engine/Source/Programs/AutomationTool/Gauntlet/Unreal/Game/EpicGame.TestConfig.cs b/Engine/Source/Programs/AutomationTool/Gauntlet/Unreal/Game/EpicGame.TestConfig.cs
--- a/Engine/Source/Programs/AutomationTool/Gauntlet/Unreal/Game/EpicGame.TestConfig.cs
+++ b/Engine/Source/Programs/AutomationTool/Gauntlet/Unreal/Game/EpicGame.TestConfig.cs
@@ -112,6 +112,8 @@ namespace EpicGame
 		[AutoParam]
 		public string EditorTraces { get; set; }
 
+		// Allows a more derived test config to take over mcp setup.
+		protected bool SkipMcpSetup = false;
 
 		// incrementing value to ensure we can assign unique values to ports etc
 		static private int NumberOfConfigsCreated = 0;
@@ -154,7 +156,7 @@ namespace EpicGame
 		{
 			base.ApplyToConfig(AppConfig, ConfigRole, OtherRoles);
 
-			if (ConfigRole.RoleType.IsClient() || ConfigRole.RoleType.IsServer() || RequiresLogin)
+			if (!SkipMcpSetup && (ConfigRole.RoleType.IsClient() || ConfigRole.RoleType.IsServer() || RequiresLogin))
 			{
 				string McpString = "";
 				bool bIsBuildMachine = CommandUtils.IsBuildMachine; 
