Supporting multiple test roles using a shared device directory.

diff --git a/Engine/Source/Programs/AutomationTool/Gauntlet/Framework/Base/Gauntlet.TargetDeviceDesktopCommon.cs b/Engine/Source/Programs/AutomationTool/Gauntlet/Framework/Base/Gauntlet.TargetDeviceDesktopCommon.cs
--- a/Engine/Source/Programs/AutomationTool/Gauntlet/Framework/Base/Gauntlet.TargetDeviceDesktopCommon.cs
+++ b/Engine/Source/Programs/AutomationTool/Gauntlet/Framework/Base/Gauntlet.TargetDeviceDesktopCommon.cs
@@ -449,7 +449,7 @@ namespace Gauntlet
 
 					if (string.IsNullOrEmpty(LogFile))
 					{
-						LogFile = Path.Combine(LogFolder, $"{ProjectName}.log");
+						LogFile = Path.Combine(LogFolder, $"{ProjectName}_{AppConfig.ProcessType.ToString()}.log");
 						CommandArguments += string.Format(" -abslog=\"{0}\"", LogFile);
 					}
 
diff --git a/Engine/Source/Programs/AutomationTool/Gauntlet/Unreal/Base/Gauntlet.UnrealSession.cs b/Engine/Source/Programs/AutomationTool/Gauntlet/Unreal/Base/Gauntlet.UnrealSession.cs
--- a/Engine/Source/Programs/AutomationTool/Gauntlet/Unreal/Base/Gauntlet.UnrealSession.cs
+++ b/Engine/Source/Programs/AutomationTool/Gauntlet/Unreal/Base/Gauntlet.UnrealSession.cs
@@ -1035,6 +1035,12 @@ namespace Gauntlet
 		/// <param name="DestinationArtifactPath"></param>
 		/// <returns></returns>
 		public UnrealRoleArtifacts SaveRoleArtifacts(UnrealTestContext InContext, UnrealSessionInstance.RoleInstance InRunningRole, string DestinationArtifactPath)
+		{
+			string DummyOutputString = null;
+			return SaveRoleArtifacts(InContext, InRunningRole, DestinationArtifactPath, ref DummyOutputString);
+		}
+
+		public UnrealRoleArtifacts SaveRoleArtifacts(UnrealTestContext InContext, UnrealSessionInstance.RoleInstance InRunningRole, string DestinationArtifactPath, ref string DirectoryToCleanUp)
 		{
 			DirectoryInfo SourceDirectory = new DirectoryInfo(InRunningRole.AppInstance.ArtifactPath);
 			DirectoryInfo DestinationDirectory = new DirectoryInfo(DestinationArtifactPath);
@@ -1155,33 +1161,41 @@ namespace Gauntlet
 			// delete this path.
 			if (SkipArchivingAssets && !bRetainArtifacts && SourceDirectory.Exists)
 			{
-				// Account for any read-only files.
-				void SetAttributesNormal(DirectoryInfo Directory)
+				if (DirectoryToCleanUp != null)
+				{
+					DirectoryToCleanUp = SourceDirectory.FullName;
+				}
+				else
 				{
-					foreach (FileInfo File in Directory.GetFiles())
+					// Account for any read-only files.
+					void SetAttributesNormal(DirectoryInfo Directory)
 					{
-						File.Attributes = FileAttributes.Normal;
+						foreach (FileInfo File in Directory.GetFiles())
+						{
+							File.Attributes = FileAttributes.Normal;
+						}
+						foreach (DirectoryInfo SubDirectory in Directory.GetDirectories())
+						{
+							SetAttributesNormal(SubDirectory);
+						}
 					}
-					foreach (DirectoryInfo SubDirectory in Directory.GetDirectories())
+					;
+					try
 					{
-						SetAttributesNormal(SubDirectory);
+						SetAttributesNormal(SourceDirectory);
+					}
+					catch
+					{
+						Log.Info("Could not remove the read-only attribute from {SourceDirectory}.", SourceDirectory);
+					}
+					try
+					{
+						SourceDirectory.Delete(true);
+					}
+					catch (Exception Exception)
+					{
+						Log.Info("Encountered an exception when deleting source artifacts at {SourceDirectory}. Artifacts will remain on the device.\n\tError message: {Exception}", SourceDirectory, Exception.Message);
 					}
-				};
-				try
-				{
-					SetAttributesNormal(SourceDirectory);
-				}
-				catch
-				{
-					Log.Info("Could not remove the read-only attribute from {SourceDirectory}.", SourceDirectory);
-				}
-				try
-				{
-					SourceDirectory.Delete(true);
-				}
-				catch (Exception Exception)
-				{
-					Log.Info("Encountered an exception when deleting source artifacts at {SourceDirectory}. Artifacts will remain on the device.\n\tError message: {Exception}", SourceDirectory, Exception.Message);
 				}
 			}
 
@@ -1230,6 +1244,7 @@ namespace Gauntlet
 			Dictionary<UnrealTargetRole, int> RoleCounts = new Dictionary<UnrealTargetRole, int>();
 			List<UnrealRoleArtifacts> AllArtifacts = new List<UnrealRoleArtifacts>();
 
+			HashSet<string> DirectoriesToDelete = new();
 			foreach (UnrealSessionInstance.RoleInstance App in TestInstance.RunningRoles)
 			{
 				string RoleName = (App.Role.IsDummy() ? "Dummy" : "") + App.Role.RoleType.ToString();
@@ -1271,7 +1286,13 @@ namespace Gauntlet
 					IDeviceUsageReporter.RecordStart(device.Name, device.Platform, IDeviceUsageReporter.EventType.SavingArtifacts);
 					try
 					{
-						Artifacts = SaveRoleArtifacts(Context, App, DestPath);
+						string DirectoryToDelete = "";
+						Artifacts = SaveRoleArtifacts(Context, App, DestPath, ref DirectoryToDelete);
+						if (!string.IsNullOrEmpty(DirectoryToDelete))
+						{
+							DirectoriesToDelete.Add(DirectoryToDelete);
+						}
+
 					}
 					catch (Exception SaveArtifactsException)
 					{
@@ -1283,7 +1304,12 @@ namespace Gauntlet
 						{
 							Log.Info("Rebooting device and retrying save role artifacts once.");
 							App.AppInstance.Device.Reboot();
+							string DirectoryToDelete = "";
 							Artifacts = SaveRoleArtifacts(Context, App, DestPath);
+							if (!string.IsNullOrEmpty(DirectoryToDelete))
+							{
+								DirectoriesToDelete.Add(DirectoryToDelete);
+							}
 						}
 						else
 						{
@@ -1305,6 +1331,41 @@ namespace Gauntlet
 				}
 			}
 
+			// Account for any read-only files.
+			void SetAttributesNormal(DirectoryInfo Directory)
+			{
+				foreach (FileInfo File in Directory.GetFiles())
+				{
+					File.Attributes = FileAttributes.Normal;
+				}
+				foreach (DirectoryInfo SubDirectory in Directory.GetDirectories())
+				{
+					SetAttributesNormal(SubDirectory);
+				}
+			}
+
+			foreach (string DirectoryToDelete in DirectoriesToDelete)
+			{
+				DirectoryInfo Directory = new DirectoryInfo(DirectoryToDelete);
+				try
+				{
+					SetAttributesNormal(Directory);
+				}
+				catch
+				{
+					Log.Info("Could not remove the read-only attribute from {SourceDirectory}.", Directory);
+				}
+				try
+				{
+					Directory.Delete(true);
+				}
+				catch (Exception Exception)
+				{
+					Log.Info("Encountered an exception when deleting source artifacts at {SourceDirectory}. Artifacts will remain on the device.\n\tError message: {Exception}", Directory, Exception.Message);
+				}
+			}
+
+
 			return AllArtifacts;
 		}
 
