Make first registered client/editor/server caching default behavior in RpcTargetRegistry.

diff --git a/Engine/Source/Programs/AutomationTool/Gauntlet/Framework/RpcFramework/Gauntlet.RpcTargetRegistry.cs b/Engine/Source/Programs/AutomationTool/Gauntlet/Framework/RpcFramework/Gauntlet.RpcTargetRegistry.cs
--- a/Engine/Source/Programs/AutomationTool/Gauntlet/Framework/RpcFramework/Gauntlet.RpcTargetRegistry.cs
+++ b/Engine/Source/Programs/AutomationTool/Gauntlet/Framework/RpcFramework/Gauntlet.RpcTargetRegistry.cs
@@ -259,6 +259,7 @@ namespace Gauntlet
 			if (GetTarget(InTarget.IpAddress, InTarget.Port) == null)
 			{
 				RpcTargets.TryAdd(InTarget.TargetName, InTarget);
+				CacheFirstRpcTargetByType(InTarget.TargetName, InTarget.TargetType);
 			}
 		}
 
@@ -402,40 +403,48 @@ namespace Gauntlet
 		#endregion
 
 		#region Target Config Helpers
+
+		private bool InitialTargetRoleIsSet(UnrealTargetRole TargetRoleType) => TargetRoleType switch
+		{
+			UnrealTargetRole.Client		  => ClientTarget is not null,
+			UnrealTargetRole.EditorGame	  => ClientTarget is not null,
+			UnrealTargetRole.Server		  => ServerTarget is not null,
+			UnrealTargetRole.EditorServer => ServerTarget is not null,
+			UnrealTargetRole.Editor		  => EditorTarget is not null,
+			UnrealTargetRole.CookedEditor => EditorTarget is not null,
+			_ => false,
+		};
+		
 		/// <summary>
 		/// Caches the first retrieved Client, Editor or Server to a class variables for simplified access.
 		/// </summary>
 		/// <param name="RpcTargetName">TargetName of the RpcTarget we are trying to cache</param>
-		/// <param name="TargetType"></param>
+		/// <param name="TargetType">The <see cref="UnrealTargetRole"/> being cached.</param>
 		private void CacheFirstRpcTargetByType(string RpcTargetName, UnrealTargetRole TargetType)
 		{
-			RpcTarget RpcTargetToBeCached = GetTarget(RpcTargetName);
-			switch (TargetType)
+			if (!InitialTargetRoleIsSet(TargetType))
 			{
-				case UnrealTargetRole.Client:
-				case UnrealTargetRole.EditorGame:
-					if (ClientTarget == null)
-					{
+				RpcTarget RpcTargetToBeCached = GetTarget(RpcTargetName);
+				switch (TargetType)
+				{
+					case UnrealTargetRole.Client:
+					case UnrealTargetRole.EditorGame:
 						ClientTarget = RpcTargetToBeCached;
-					}
-					break;
-				case UnrealTargetRole.Editor:
-				case UnrealTargetRole.CookedEditor:
-					if (EditorTarget == null)
-					{
+						break;
+					case UnrealTargetRole.Editor:
+					case UnrealTargetRole.CookedEditor:
 						EditorTarget = RpcTargetToBeCached;
-					}
-					break;
-				case UnrealTargetRole.Server:
-				case UnrealTargetRole.EditorServer:
-					if (ServerTarget == null)
-					{
+						break;
+					case UnrealTargetRole.Server:
+					case UnrealTargetRole.EditorServer:
 						ServerTarget = RpcTargetToBeCached;
-					}
-					break;
-				default:
-					Log.Error("Trying to initialize RPC Target without a RpcTargetType!");
-					break;
+						break;
+					case UnrealTargetRole.Unknown:
+					case UnrealTargetRole.Host:
+					default:
+						Log.Error($"No handling for setting {TargetType} client/editor/server target property on {nameof(RpcTargetRegistry)}.");
+						break;
+				}
 			}
 		}
 
