[Linux] re-enable -nocef when running unattended to unblock horde ASan tests

diff --git a/Engine/Source/Programs/AutomationTool/Gauntlet/Unreal/Base/Gauntlet.UnrealTestConfiguration.cs b/Engine/Source/Programs/AutomationTool/Gauntlet/Unreal/Base/Gauntlet.UnrealTestConfiguration.cs
--- a/Engine/Source/Programs/AutomationTool/Gauntlet/Unreal/Base/Gauntlet.UnrealTestConfiguration.cs
+++ b/Engine/Source/Programs/AutomationTool/Gauntlet/Unreal/Base/Gauntlet.UnrealTestConfiguration.cs
@@ -1053,7 +1053,7 @@ namespace Gauntlet
 				// Disable cef as it seems to want to talk to an X11 server so unlikely its even working
 				if (AppConfig.Platform == UnrealTargetPlatform.Linux)
 				{
-					AppConfig.CommandLine += " -renderoffscreen";
+					AppConfig.CommandLine += " -renderoffscreen -nocef";
 				}
 			}
 
