Minor BugHawk fix, wrong variable checked for correct range. In most use cases, both are positive so no actual behavioural difference

#rnx

diff --git a/Engine/Source/Programs/AutomationTool/Scripts/CopyBuildToStagingDirectory.Automation.cs b/Engine/Source/Programs/AutomationTool/Scripts/CopyBuildToStagingDirectory.Automation.cs
--- a/Engine/Source/Programs/AutomationTool/Scripts/CopyBuildToStagingDirectory.Automation.cs
+++ b/Engine/Source/Programs/AutomationTool/Scripts/CopyBuildToStagingDirectory.Automation.cs
@@ -4496,7 +4496,7 @@ namespace AutomationScripts
 				int CompressionMinPercentSaved = 0;
 				if (PlatformGameConfig.GetInt32("/Script/UnrealEd.ProjectPackagingSettings", "PackageCompressionMinPercentSaved", out CompressionMinPercentSaved))
 				{
-					if (CompressionMinBytesSaved > 0)
+					if (CompressionMinPercentSaved > 0)
 					{
 						AdditionalArgs += String.Format(" -compressionMinPercentSaved={0}", CompressionMinPercentSaved);
 					}
