Sanitize error - brackets

diff --git a/Engine/Source/Programs/AutomationTool/Gauntlet/Unreal/Base/Gauntlet.UnrealTestNode.cs b/Engine/Source/Programs/AutomationTool/Gauntlet/Unreal/Base/Gauntlet.UnrealTestNode.cs
--- a/Engine/Source/Programs/AutomationTool/Gauntlet/Unreal/Base/Gauntlet.UnrealTestNode.cs
+++ b/Engine/Source/Programs/AutomationTool/Gauntlet/Unreal/Base/Gauntlet.UnrealTestNode.cs
@@ -2414,6 +2414,7 @@ namespace Gauntlet
 			if (Event.IsSanReport && CommandUtils.IsBuildMachine
 				&& SanitizerEventMatcher.AddSanitizerSummaryProperties(ref CallstackString, Properties))
 			{
+				CallstackString = CallstackString.Replace("{", "{{").Replace("}", "}}");
 				Message = $" * Fatal Error: {{Summary}}\n{CallstackString}";
 				EventType = SanitizerEventMatcher.ConvertSanitizerNameToEventId(Properties.GetValueOrDefault("SanitizerName")?.ToString());
 			}
