Disable issue tracking on ad-hoc stack analysis

#rnx

diff --git a/Engine/Source/Programs/AutomationTool/Gauntlet/Platform/Windows/Gauntlet.TargetDeviceWindows.cs b/Engine/Source/Programs/AutomationTool/Gauntlet/Platform/Windows/Gauntlet.TargetDeviceWindows.cs
--- a/Engine/Source/Programs/AutomationTool/Gauntlet/Platform/Windows/Gauntlet.TargetDeviceWindows.cs
+++ b/Engine/Source/Programs/AutomationTool/Gauntlet/Platform/Windows/Gauntlet.TargetDeviceWindows.cs
@@ -440,11 +440,11 @@ namespace Gauntlet
 						}
 						if (Stack.Any())
 						{
-							Log.Error(KnownLogEvents.Gauntlet_TestEvent, "Stack analysis produced with 'cdb {Args}':\n{StackTrace}", Args, string.Join('\n', Stack));
+							Log.Error(KnownLogEvents.Horde_BuildHealth_Ignore, "Stack analysis produced with 'cdb {Args}':\n{StackTrace}", Args, string.Join('\n', Stack));
 						}
 						else
 						{
-							Log.Warning("Could not find stack analysis from cdb command.");
+							Log.Warning(KnownLogEvents.Horde_BuildHealth_Ignore, "Could not find stack analysis from cdb command.");
 						}
 					}
 				}
