LaunchMultiServer:  When specifying -notimeouts, we should also pass that as an argument to the Clients to ensure they do not close the connection during debug sessions.

#rnx

diff --git a/Engine/Source/Programs/AutomationTool/Scripts/LaunchMultiServer.Automation.cs b/Engine/Source/Programs/AutomationTool/Scripts/LaunchMultiServer.Automation.cs
--- a/Engine/Source/Programs/AutomationTool/Scripts/LaunchMultiServer.Automation.cs
+++ b/Engine/Source/Programs/AutomationTool/Scripts/LaunchMultiServer.Automation.cs
@@ -534,6 +534,7 @@ Example running 4 clients locally connecting to single proxy server w/ a non-loc
 				for (int Port = StartPort; Port <= EndPort; ++Port)
 				{
 					string ClientArgs = $" -game {AdditionalClientArguments}";
+					ClientArgs += NoTimeouts ? " -notimeouts" : "";
 					ClientArgs += $" {GameServerBaseAddress}:{Port}";
 					ClientArgs += $" -log={ClientType}-{{ClientId}}.log -windowed -SaveWinPos={{ClientId}}";
 
