Make TestFlight upload check configurable with CheckCount and CheckDelay

diff --git a/Engine/Source/Programs/AutomationTool/Apple/ExportAppleBuild.Automation.cs b/Engine/Source/Programs/AutomationTool/Apple/ExportAppleBuild.Automation.cs
--- a/Engine/Source/Programs/AutomationTool/Apple/ExportAppleBuild.Automation.cs
+++ b/Engine/Source/Programs/AutomationTool/Apple/ExportAppleBuild.Automation.cs
@@ -24,6 +24,7 @@ using System.Text.Json;
 using System.Text.RegularExpressions;
 using System.Threading.Tasks;
 using System.Net.Http.Json;
+using Org.BouncyCastle.Asn1;
 
 class ExportAppleBuild : BuildCommand
 {
@@ -47,6 +48,8 @@ class ExportAppleBuild : BuildCommand
 		string VersionConflictSigningIdentity = ParseParamValue("VersionConflictSigningIdentity");
 		string VersionConflictIncrement = ParseParamValue("VersionConflictIncrement");
 		string ExportAuthOverride = ParseParamValue("ExportAuthOverride");
+		string CheckDelay = ParseParamValue("CheckDelay");
+		string CheckCount = ParseParamValue("CheckCount");
 
 		DirectoryReference XcArchive;
 		if (!string.IsNullOrEmpty(LatestXcArchiveParam))
@@ -143,6 +146,25 @@ class ExportAppleBuild : BuildCommand
 
 		if (!string.IsNullOrEmpty(WhatsNew))
 		{
+			int Delay = 30;
+			if (!string.IsNullOrEmpty(CheckDelay))
+			{
+				int Result = 0;
+				if (Int32.TryParse(CheckDelay, out Result))
+				{
+					Delay = Int32.Max(Int32.Min(Result, 300), 15);
+				}
+			}
+			int Tries = 30;
+			if (!string.IsNullOrEmpty(CheckCount))
+			{
+				int Result = 0;
+				if (Int32.TryParse(CheckCount, out Result))
+				{
+					Tries = Int32.Max(Int32.Min(Result, 100), 5);
+				}
+			}
+
 			// get some info from the .xcarchive that we just uploaded
 			string BundleId, AppVersion, BundleVersion, CommandLine;
 			GetInfoFromXcArchive(XcArchive, out BundleId, out AppVersion, out BundleVersion, out CommandLine);
@@ -156,7 +178,7 @@ class ExportAppleBuild : BuildCommand
 			WhatsNew = WhatsNew.Replace("@uecommandline", CommandLine);
 			WhatsNew = WhatsNew.Replace("@commandline", CommandLine);
 
-			UpdateWhatsNew(JWTToken, BundleId, AppVersion, BundleVersion, WhatsNew);
+			UpdateWhatsNew(JWTToken, BundleId, AppVersion, BundleVersion, WhatsNew, Delay, Tries);
 		}
 
 		return (ExitCode)0;
@@ -479,13 +501,11 @@ class ExportAppleBuild : BuildCommand
 		return true;
 	}
 
-	bool UpdateWhatsNew(string Token, string BundleId, string PreReleaseVersion, string BundleVersion, string Text)
+	bool UpdateWhatsNew(string Token, string BundleId, string PreReleaseVersion, string BundleVersion, string Text, int Delay = 30, int Tries = 30)
 	{
 		using (HttpClient = CreateClient(Token))
 		{
-			// wait 30 second between queries for the version
-			int Delay = 30;
-			int Tries = 15 * 60 / Delay;
+			// wait Delay second between queries for the version - made this configurable as part of the call
 
 			string BuildId = null;
 			while (BuildId == null && Tries-- > 0)
@@ -519,10 +539,11 @@ class ExportAppleBuild : BuildCommand
 
 			if (BuildId == null)
 			{
-				Logger.LogError("Unable to find build for BundleId: {BundleId}, App Version {PreReleaseVersion}, Build Version {BundleVersion}. It could be due to:\n" + 
+				Logger.LogWarning("Unable to find build for BundleId: {BundleId}, App Version {PreReleaseVersion}, Build Version {BundleVersion}. It could be due to:\n" + 
 					"  * The build was never actually uploaded to AppStoreConnect to TestFlight\n" + 
 					"  * A later version was already there, and Apple changed the version number automatically. Run again with -CheckVersion -SkipUpload to check\n" + 
-					"  * There was a problem with the app. Check your email (or others on your team) to see if Apple rejected the app after uploading",
+					"  * There was a problem with the app. Check your email (or others on your team) to see if Apple rejected the app after uploading\n" +
+					"  * The build is potentially still validating and taking longer than our timeout.",
 					BundleId, PreReleaseVersion, BundleVersion);
 				return false;
 			}
