PR #13367: Add support for prepending a file to the Linux launch sh script

[RN] Add support for prepending project file Config\BootstrapPreamble.sh to the front of the generated linux bootstrap script

diff --git a/Engine/Source/Programs/AutomationTool/Linux/LinuxPlatform.Automation.cs b/Engine/Source/Programs/AutomationTool/Linux/LinuxPlatform.Automation.cs
--- a/Engine/Source/Programs/AutomationTool/Linux/LinuxPlatform.Automation.cs
+++ b/Engine/Source/Programs/AutomationTool/Linux/LinuxPlatform.Automation.cs
@@ -168,6 +168,16 @@ public abstract class BaseLinuxPlatform : Platform
 		StringBuilder Script = new StringBuilder();
 		string EOL = "\n";
 		Script.Append("#!/bin/sh" + EOL);
+
+		// If the BootstrapPreamble.sh file exists, prepend it to the regular launch script, so that things like display and required platform environment variables can be set up automatically
+		FileReference BootstrapPreambleFile = FileReference.Combine(SC.ProjectRoot, "Config", "BootstrapPreamble.sh");
+		if (FileReference.Exists(BootstrapPreambleFile))
+		{
+			Script.AppendFormat(EOL + "### Added from project Config BootstrapPreamble.sh" + EOL + EOL);
+			Script.Append(FileReference.ReadAllText(BootstrapPreambleFile));
+			Script.AppendFormat(EOL + "### Added from project Config BootstrapPreamble.sh" + EOL + EOL);			
+		}
+
 		// allow running from symlinks
 		Script.AppendFormat("UE_TRUE_SCRIPT_NAME=$(echo \\\"$0\\\" | xargs readlink -f)" + EOL);
 		Script.AppendFormat("UE_PROJECT_ROOT=$(dirname \"$UE_TRUE_SCRIPT_NAME\")" + EOL);
