Fail PGO automation if the PROFRAW file is zero bytes

I think this is happening because abnormal termination isn't giving the process a chance to flush the PGO data to disk.

diff --git a/Engine/Source/Programs/AutomationTool/Gauntlet/Platform/Windows/Gauntlet.WindowsPGO.cs b/Engine/Source/Programs/AutomationTool/Gauntlet/Platform/Windows/Gauntlet.WindowsPGO.cs
--- a/Engine/Source/Programs/AutomationTool/Gauntlet/Platform/Windows/Gauntlet.WindowsPGO.cs
+++ b/Engine/Source/Programs/AutomationTool/Gauntlet/Platform/Windows/Gauntlet.WindowsPGO.cs
@@ -80,8 +80,14 @@ namespace Gauntlet
 		{
 			string PGODir = Path.Combine(ArtifactPath, "PGO");
 
-			if (File.Exists(ExpectedProfrawFileName)) 
+			if (File.Exists(ExpectedProfrawFileName))
 			{
+				FileInfo Info = new FileInfo(ExpectedProfrawFileName);
+				if (Info.Length == 0)
+				{
+					throw new AutomationException($"file has zero length, probably crashed - {ExpectedProfrawFileName}");
+				}
+
 				Log.Info("Clang format profiling files detected");
 				GatherResultsClang();
 			}
