Prevent double-add of registry targets for role

diff --git a/Engine/Source/Programs/AutomationTool/Gauntlet/Framework/RpcFramework/Gauntlet.RpcTargetRegistry.cs b/Engine/Source/Programs/AutomationTool/Gauntlet/Framework/RpcFramework/Gauntlet.RpcTargetRegistry.cs
--- a/Engine/Source/Programs/AutomationTool/Gauntlet/Framework/RpcFramework/Gauntlet.RpcTargetRegistry.cs
+++ b/Engine/Source/Programs/AutomationTool/Gauntlet/Framework/RpcFramework/Gauntlet.RpcTargetRegistry.cs
@@ -233,10 +233,16 @@ namespace Gauntlet
 
 		#endregion // RPC Target Configuration Base Values
 
+		/// <summary>
+		/// Roles which have a RpcTarget registered, only one target should register for a single role
+		/// </summary>
+		private readonly ConcurrentDictionary<UnrealTestRole, RpcTarget> RoleTargets = new(); 
+
 		/// <summary>
 		/// All registered RPC Targets stored here, in a thread-safe container.
 		/// </summary>
-		public ConcurrentDictionary<string, RpcTarget> RpcTargets = new ConcurrentDictionary<string, RpcTarget>();
+		public readonly ConcurrentDictionary<string, RpcTarget> RpcTargets = new();
+
 
 		public RpcTargetRegistry()
 		{
@@ -450,7 +456,7 @@ namespace Gauntlet
 
 		/// <summary>
 		/// This is a command config helper to quickly add the correct command line parameters to Enable RPC's for the desired role
-		/// and register the RpcTarget with the RpcTargetRegistry.
+		/// and register the RpcTarget with the RpcTargetRegistry, if not already registered.
 		/// </summary>
 		/// <param name="TestRole">Role to enable </param>
 		/// <param name="PortOverride"></param>
@@ -460,47 +466,55 @@ namespace Gauntlet
 		{
 			try
 			{
-				RpcTarget NewRpcTarget = GetNextAvailableRpcTarget(TestRole.Type);
-
-				// override port if one was provided
-				if (PortOverride != 0)
+				RpcTarget CurrentTarget;
+				if (RoleTargets.TryGetValue(TestRole, out CurrentTarget))
 				{
-					NewRpcTarget.Port = PortOverride;
+					Log.Verbose($"RPC already enabled for {TestRole.Type}:{TestRole.GetHashCode()}");
 				}
-
-				if (NewRpcTarget == null)
+				else
 				{
-					Log.Error("UnrealTargetRole Unrecognized! Value: {0}", TestRole.Type.ToString());
-					// Shouldn't continue, as we were not able to recognize the role
-					throw new Exception(string.Format("UnrealTargetRole Unrecognized! Value: {0}", TestRole.Type.ToString()));
+					CurrentTarget = GetNextAvailableRpcTarget(TestRole.Type);
+					if(RoleTargets.TryAdd(TestRole, CurrentTarget))
+					{
+						AddNewTarget(CurrentTarget);
+						CacheFirstRpcTargetByType(CurrentTarget.TargetName, TestRole.Type);
+					}
+					else
+					{
+						// somewhere else raced to add the same key
+						CurrentTarget = RoleTargets[TestRole];
+						Log.VeryVerbose($"RPC already enabled while trying to register for {TestRole.Type}:{TestRole.GetHashCode()}");
+					}
+					Log.Verbose($"RPC enabled for {TestRole.Type}:{TestRole.GetHashCode()}");
 				}
 
-				// Get IP address of the target, using multi-home logic.
-				NewRpcTarget.IpAddress = ResolveProcessUri(TestRole);
+				if (PortOverride != 0)
+				{
+					CurrentTarget.Port = PortOverride;
+				}
 
-				// if a URI was provided, us that over everything else.
 				if (!string.IsNullOrEmpty(UriOverride))
 				{
-					NewRpcTarget.IpAddress = UriOverride;
+					CurrentTarget.IpAddress = UriOverride;
+				}
+				else
+				{
+					CurrentTarget.IpAddress = ResolveProcessUri(TestRole);
 				}
 
 				if (TestRole.DeferredLaunch)
 				{
-					NewRpcTarget.IsDeferStart = true;
+					CurrentTarget.IsDeferStart = true;
 				}
 
 				// defines the port the role's process should listen on and adds to command line
-				TestRole.CommandLineParams.AddRawCommandline(string.Format("-rpcport={0}", NewRpcTarget.Port));
+				TestRole.CommandLineParams.AddRawCommandline(string.Format("-rpcport={0}", CurrentTarget.Port));
 				// Override HTTPServer.Listeners .ini settings
-				TestRole.CommandLineParams.AddRawCommandline(string.Format("-ini:Engine:[HTTPServer.Listeners]:DefaultBindAddress={0},[HTTPServer.Listeners]:DefaultReuseAddressAndPort=true -logcmds=\"LogSockets VeryVerbose\"", NewRpcTarget.IpAddress));
+				TestRole.CommandLineParams.AddRawCommandline(string.Format("-ini:Engine:[HTTPServer.Listeners]:DefaultBindAddress={0},[HTTPServer.Listeners]:DefaultReuseAddressAndPort=true -logcmds=\"LogSockets VeryVerbose\"", CurrentTarget.IpAddress));
 				// tells process what its sender ID is
-				TestRole.CommandLineParams.Add("rpcsenderid", NewRpcTarget.TargetName);
+				TestRole.CommandLineParams.Add("rpcsenderid", CurrentTarget.TargetName);
 				// setup gauntlet ip and port so messages can be sent from process
 				TestRole.CommandLineParams.Add("externalrpclistenaddress", RpcExecutor.GetListenAddressAndPort());
-				// add to RpcTargetManager
-				AddNewTarget(NewRpcTarget);
-
-				CacheFirstRpcTargetByType(NewRpcTarget.TargetName, TestRole.Type);
 			}
 			catch (TestException)
 			{
