- Refactored the authtoken stuff in ExportAppleBuild to automatically re-create the token as needed, but wrapping the token and HttpClient together in a new JwtTokenClient

diff --git a/Engine/Source/Programs/AutomationTool/Apple/ExportAppleBuild.Automation.cs b/Engine/Source/Programs/AutomationTool/Apple/ExportAppleBuild.Automation.cs
--- a/Engine/Source/Programs/AutomationTool/Apple/ExportAppleBuild.Automation.cs
+++ b/Engine/Source/Programs/AutomationTool/Apple/ExportAppleBuild.Automation.cs
@@ -72,16 +72,12 @@ class ExportAppleBuild : BuildCommand
 			Logger.LogError("No XCArchive found, with -XCArchive or -LatestXCArchiveForTarget params");
 			return ExitCode.Error_Arguments;
 		}
-		
-		string JWTToken = null;
-		if (bCheckVersion || !string.IsNullOrEmpty(WhatsNew))
-		{
-			JWTToken = CreateJWTToken(ProjectFile);
-		}
+
+		TokenClient = new JwtTokenClient(ProjectFile);
 
 		if (bCheckVersion)
 		{
-			if (CheckVersion(JWTToken, ref XcArchive, VersionConflictSigningIdentity, VersionConflictIncrement) == false)
+			if (CheckVersion(TokenClient, ref XcArchive, VersionConflictSigningIdentity, VersionConflictIncrement) == false)
 			{
 				return ExitCode.Error_Arguments;
 			}
@@ -178,21 +174,12 @@ class ExportAppleBuild : BuildCommand
 			WhatsNew = WhatsNew.Replace("@uecommandline", CommandLine);
 			WhatsNew = WhatsNew.Replace("@commandline", CommandLine);
 
-			UpdateWhatsNew(JWTToken, BundleId, AppVersion, BundleVersion, WhatsNew, Delay, Tries);
+			UpdateWhatsNew(TokenClient, BundleId, AppVersion, BundleVersion, WhatsNew, Delay, Tries);
 		}
 
 		return (ExitCode)0;
 	}
 
-	HttpClient CreateClient(string Token)
-	{
-		HttpClient HttpClient = new HttpClient();
-		HttpClient.BaseAddress = new Uri("https://api.appstoreconnect.apple.com/v1/");
-		HttpClient.DefaultRequestHeaders.Authorization = new AuthenticationHeaderValue("Bearer", Token);
-		return HttpClient;
-	}
-
-
 
 	record Links(string self, string related);
 	record Relationship(Links links);
@@ -232,11 +219,11 @@ class ExportAppleBuild : BuildCommand
 	static CancellationToken CancellationToken = new();
 
 	// shared client object
-	HttpClient HttpClient = null;
+	JwtTokenClient TokenClient = null;
 
 	private async Task<TResponse> GetAsync<TResponse>(string RelativePath, CancellationToken CancellationToken = default)
 	{
-		TResponse Response = await HttpClient.GetFromJsonAsync<TResponse>(RelativePath, CancellationToken);
+		TResponse Response = await TokenClient.Client.GetFromJsonAsync<TResponse>(RelativePath, CancellationToken);
 		return Response ?? throw new InvalidCastException($"Expected non-null response from GET to {RelativePath}");
 	}
 
@@ -254,11 +241,11 @@ class ExportAppleBuild : BuildCommand
 		bool bHasProcessedAll = false;
 		while (!bHasProcessedAll)
 		{
-			PagedBuildsResponse Response = GetAsync<PagedBuildsResponse>(SubURL, CancellationToken).Result;			
+			PagedBuildsResponse Response = GetAsync<PagedBuildsResponse>(SubURL, CancellationToken).Result;
 			foreach (BuildInfo Build in Response.data)
 			{
 				// follow link to preReleaseVersion
-				SubURL = Build.relationships.preReleaseVersion.links.related.Replace(HttpClient.BaseAddress.ToString(), "");
+				SubURL = Build.relationships.preReleaseVersion.links.related.Replace(TokenClient.Client.BaseAddress.ToString(), "");
 				VersionResponse VersionResponse = GetAsync<VersionResponse>(SubURL, CancellationToken).Result;
 
 				// if it matches, return this Build!
@@ -275,7 +262,7 @@ class ExportAppleBuild : BuildCommand
 			else
 			{
 				// get suburl for next page
-				SubURL = Response.links.next.Replace(HttpClient.BaseAddress.ToString(), "");
+				SubURL = Response.links.next.Replace(TokenClient.Client.BaseAddress.ToString(), "");
 			}
 		}
 
@@ -300,161 +287,158 @@ class ExportAppleBuild : BuildCommand
 	}
 
 	// return true if the version doesn't exist - which means it can be uploaded
-	bool CheckVersion(string JWTToken, ref DirectoryReference XcArchive, string SigningIdentity, string Increment)
+	bool CheckVersion(JwtTokenClient TokenClient, ref DirectoryReference XcArchive, string SigningIdentity, string Increment)
 	{
 		// get some info from the .xcarchive
 		string BundleId, AppVersion, BundleVersion;
 		GetInfoFromXcArchive(XcArchive, out BundleId, out AppVersion, out BundleVersion, out _);
-
-
+		
 		Logger.LogInformation("Checking for version >= {Version}", BundleVersion);
 		Version BundleVersionNumber = new Version(BundleVersion);
-		using (HttpClient = CreateClient(JWTToken))
+
+		int Attempts = 0;
+		while (Attempts < 5)
 		{
-			int Attempts = 0;
-			while (Attempts < 5)
+			if (Attempts > 0)
 			{
-				if (Attempts > 0)
+				Thread.Sleep(30 * 1000);
+			}
+			Attempts++;
+
+			try
+			{
+				string AppId = GetAppId(BundleId);
+
+				// get prerelease version id
+				string SubURL = $"preReleaseVersions?filter[app]={AppId}&filter[version]={AppVersion}&fields[preReleaseVersions]=version";
+				VersionsResponse VersionResponse = GetAsync<VersionsResponse>(SubURL, CancellationToken).Result;
+				// if the app version doesn't even exist, then the bundle verrsion can't exist in it!
+				if (VersionResponse.data.Count == 0)
 				{
-					Thread.Sleep(30 * 1000);
+					return true;
 				}
-				Attempts++;
 
-				try
+				Version MaxVersion = new Version(0, 0);
+
+				// now walk over all builds until we find one >= the bundle version
+				string VersionId = VersionResponse.data[0].id;
+				SubURL = $"preReleaseVersions/{VersionId}/builds?limit=5&fields[builds]=version";
+				bool bHasProcessedAll = false;
+				while (!bHasProcessedAll)
 				{
-					string AppId = GetAppId(BundleId);
+					PagedBuildsResponse Response = GetAsync<PagedBuildsResponse>(SubURL, CancellationToken).Result;
 
-					// get prerelease version id
-					string SubURL = $"preReleaseVersions?filter[app]={AppId}&filter[version]={AppVersion}&fields[preReleaseVersions]=version";
-					VersionsResponse VersionResponse = GetAsync<VersionsResponse>(SubURL, CancellationToken).Result;
-					// if the app version doesn't even exist, then the bundle verrsion can't exist in it!
-					if (VersionResponse.data.Count == 0)
+					foreach (BuildInfo Build in Response.data)
 					{
-						return true;
-					}
+						string VersionString = Build.attributes.version;
+						// append .0.0 if needed
+						if (!VersionString.Contains("."))
+						{
+							VersionString += ".0.0";
+						}
 
-					Version MaxVersion = new Version(0, 0);
+						Version Version = new Version(VersionString);
+						if (Version > MaxVersion)
+						{
+							MaxVersion = Version;
+						}
+					}
 
-					// now walk over all builds until we find one >= the bundle version
-					string VersionId = VersionResponse.data[0].id;
-					SubURL = $"preReleaseVersions/{VersionId}/builds?limit=5&fields[builds]=version";
-					bool bHasProcessedAll = false;
-					while (!bHasProcessedAll)
+					if (Response.links.next == null)
+					{
+						bHasProcessedAll = true;
+					}
+					else
 					{
-						PagedBuildsResponse Response = GetAsync<PagedBuildsResponse>(SubURL, CancellationToken).Result;
+						// get suburl for next page
+						SubURL = Response.links.next.Replace(TokenClient.Client.BaseAddress.ToString(), "");
+					}
+				}
 
-						foreach (BuildInfo Build in Response.data)
+				if (MaxVersion >= BundleVersionNumber)
+				{
+					if (SigningIdentity != null && Increment != null)
+					{
+						Version IncrementVersion = new Version(Increment);
+						if (IncrementVersion == new Version(0, 0, 0))
 						{
-							string VersionString = Build.attributes.version;
-							// append .0.0 if needed
-							if (!VersionString.Contains("."))
-							{
-								VersionString += ".0.0";
-							}
-
-							Version Version = new Version(VersionString);
-							if (Version > MaxVersion)
-							{
-								MaxVersion = Version;
-							}
+							throw new Exception($"Invalid version increment {Increment}");
 						}
 
-						if (Response.links.next == null)
+						string TempDir = Path.GetTempFileName();
+						File.Delete(TempDir);
+						Directory.CreateDirectory(TempDir);
+
+						FileReference Script = FileReference.Combine(Unreal.EngineDirectory, "Build/BatchFiles/Mac/ResignApp.sh");
+						string FixedXcArchive = TempDir + "/FixedVersion.xcarchive";
+						Version FixedVersion = new Version(MaxVersion.Major + IncrementVersion.Major, MaxVersion.Minor + IncrementVersion.Minor, MaxVersion.Build + IncrementVersion.Build);
+						string Options = $"{Script} -s \"{XcArchive}\" -d \"{FixedXcArchive}\" -i \"{SigningIdentity}\" -bv {FixedVersion}";
+
+						if (MaxVersion == BundleVersionNumber)
 						{
-							bHasProcessedAll = true;
+							Logger.LogInformation("A build exists with same version that is in the .xcarchive ({BundleVersion}). Making a new\n" +
+							                      ".xcarchive ({FixedXcArchive}) with a version using the supplied increment. The new version is {FixedVersion}.",
+								BundleVersion, FixedXcArchive, FixedVersion);
 						}
 						else
 						{
-							// get suburl for next page
-							SubURL = Response.links.next.Replace(HttpClient.BaseAddress.ToString(), "");
+							Logger.LogInformation("A build exists with LARGER version ({MaxVersion}) that is in the .xcarchive ({BundleVersion})\n" +
+								"which would cause the uploaded build to be auto-incremented. Making a new .xcarchive ({FixedXcArchive})\n" +
+								"with a incremented FROM THE LARGER VERSION (not the version in your .xcarchive!). The new version is {FixedVersion}.",
+								MaxVersion, BundleVersion, FixedXcArchive, FixedVersion);
 						}
-					}
 
-					if (MaxVersion >= BundleVersionNumber)
-					{
-						if (SigningIdentity != null && Increment != null)
+						Logger.LogInformation("Running: {Script} {Options}...", "/usr/bin/env", Options);
+						int ExitCode;
+						string Output = Utils.RunLocalProcessAndReturnStdOut("/usr/bin/env", Options, out ExitCode);
+						if (ExitCode != 0)
 						{
-							Version IncrementVersion = new Version(Increment);
-							if (IncrementVersion == new Version(0, 0, 0))
-							{
-								throw new Exception($"Invalid version increment {Increment}");
-							}
-
-							string TempDir = Path.GetTempFileName();
-							File.Delete(TempDir);
-							Directory.CreateDirectory(TempDir);
-
-							FileReference Script = FileReference.Combine(Unreal.EngineDirectory, "Build/BatchFiles/Mac/ResignApp.sh");
-							string FixedXcArchive = TempDir + "/FixedVersion.xcarchive";
-							Version FixedVersion = new Version(MaxVersion.Major + IncrementVersion.Major, MaxVersion.Minor + IncrementVersion.Minor, MaxVersion.Build + IncrementVersion.Build);
-							string Options = $"{Script} -s \"{XcArchive}\" -d \"{FixedXcArchive}\" -i \"{SigningIdentity}\" -bv {FixedVersion}";
-
-							if (MaxVersion == BundleVersionNumber)
-							{
-								Logger.LogInformation("A build exists with same version that is in the .xcarchive ({BundleVersion}). Making a new\n" +
-									".xcarchive ({FixedXcArchive}) with a version using the supplied increment. The new version is {FixedVersion}.",
-									BundleVersion, FixedXcArchive, FixedVersion);
-							}
-							else
-							{
-								Logger.LogInformation("A build exists with LARGER version ({MaxVersion}) that is in the .xcarchive ({BundleVersion})\n" +
-									"which would cause the uploaded build to be auto-incremented. Making a new .xcarchive ({FixedXcArchive})\n" +
-									"with a incremented FROM THE LARGER VERSION (not the version in your .xcarchive!). The new version is {FixedVersion}.",
-									MaxVersion, BundleVersion, FixedXcArchive, FixedVersion);
-							}
-
-							Logger.LogInformation("Running: {Script} {Options}...", "/usr/bin/env", Options);
-							int ExitCode;
-							string Output = Utils.RunLocalProcessAndReturnStdOut("/usr/bin/env", Options, out ExitCode);
-							if (ExitCode != 0)
-							{
-								Logger.LogError("{Output}", Output);
-								return false;
-							}
-
-							Logger.LogInformation("{Output}", Output);
-
-							// update the XcArchive location to the resigned one
-							XcArchive = new DirectoryReference(FixedXcArchive);
-
-							// we know that this xcarchive has a version that is not in the appstore, so return true 
+							Logger.LogError("{Output}", Output);
+							return false;
 						}
-						else
-						{
-							Version PossibleVersion = new Version(MaxVersion.Major, MaxVersion.Minor + 1, MaxVersion.Build);
 
-							Logger.LogError("A build exists with same or larger version ({MaxVersion}) than is in the .xcarchive ({BundleVersion}). Either make a new build,\n" +
-								"or modify the build version in the .xcarchive with a command like:\n" +
-								"   Engine/Build/BatchFiles/Mac/ResignApp.sh -s {XcArchive} -d UpdatedVersion.xcarchive -i <signing identity> -bv {NewVersion}\n" +
-								"and then upload that .xcarchive",
-								MaxVersion, BundleVersion, XcArchive, PossibleVersion);
+						Logger.LogInformation("{Output}", Output);
 
-							return false;
-						}
+						// update the XcArchive location to the resigned one
+						XcArchive = new DirectoryReference(FixedXcArchive);
 
+						// we know that this xcarchive has a version that is not in the appstore, so return true 
+					}
+					else
+					{
+						Version PossibleVersion = new Version(MaxVersion.Major, MaxVersion.Minor + 1, MaxVersion.Build);
+
+						Logger.LogError("A build exists with same or larger version ({MaxVersion}) than is in the .xcarchive ({BundleVersion}). Either make a new build,\n" +
+							"or modify the build version in the .xcarchive with a command like:\n" +
+							"   Engine/Build/BatchFiles/Mac/ResignApp.sh -s {XcArchive} -d UpdatedVersion.xcarchive -i <signing identity> -bv {NewVersion}\n" +
+							"and then upload that .xcarchive",
+							MaxVersion, BundleVersion, XcArchive, PossibleVersion);
+
+						return false;
 					}
 
-					// if we get here, there is either no newer version, or we have successfully re-versioned the .app
-					return true;
 				}
-				catch (AggregateException Ex)
-				{
-					Ex.Handle(x => 
+
+				// if we get here, there is either no newer version, or we have successfully re-versioned the .app
+				return true;
+			}
+			catch (AggregateException Ex)
+			{
+				Ex.Handle(x => 
+					{
+						if (x is HttpRequestException)
 						{
-							if (x is HttpRequestException)
-							{
-								Logger.LogWarning("Got an exception while checking for existing version numbers: '{Msg}'. Will try again...", x.Message);
-								return true;
-							}
-							return false;
+							Logger.LogWarning("Got an exception while checking for existing version numbers: '{Msg}'. Will try again...", x.Message);
+							return true;
 						}
-					);
+						return false;
+					}
+				);
 
-				}
 			}
-
-			return false;
 		}
+
+		return false;
 	}
 
 
@@ -470,7 +454,7 @@ class ExportAppleBuild : BuildCommand
 		StringContent JsonContent = new StringContent(JsonSerializer.Serialize<LocalizationCreateRequest>(LocUpdate), System.Text.Encoding.UTF8, "application/json");
 
 		string SubURL = $"betaBuildLocalizations";
-		HttpResponseMessage Response = HttpClient.PostAsync(SubURL, JsonContent).Result;
+		HttpResponseMessage Response = TokenClient.Client.PostAsync(SubURL, JsonContent).Result;
 
 		if (!Response.IsSuccessStatusCode)
 		{
@@ -490,7 +474,7 @@ class ExportAppleBuild : BuildCommand
 		StringContent JsonContent = new StringContent(JsonSerializer.Serialize<LocalizationUpdate>(LocUpdate), System.Text.Encoding.UTF8, "application/json");
 
 		string SubURL = $"betaBuildLocalizations/{LocId}";
-		HttpResponseMessage Response = HttpClient.PatchAsync(SubURL, JsonContent).Result;
+		HttpResponseMessage Response = TokenClient.Client.PatchAsync(SubURL, JsonContent).Result;
 
 		if (!Response.IsSuccessStatusCode)
 		{
@@ -501,108 +485,102 @@ class ExportAppleBuild : BuildCommand
 		return true;
 	}
 
-	bool UpdateWhatsNew(string Token, string BundleId, string PreReleaseVersion, string BundleVersion, string Text, int Delay = 30, int Tries = 30)
+	bool UpdateWhatsNew(JwtTokenClient TokenClient, string BundleId, string PreReleaseVersion, string BundleVersion, string Text, int Delay = 30, int Tries = 30)
 	{
-		using (HttpClient = CreateClient(Token))
+		string BuildId = null;
+		while (BuildId == null && Tries-- > 0)
 		{
-			// wait Delay second between queries for the version - made this configurable as part of the call
-
-			string BuildId = null;
-			while (BuildId == null && Tries-- > 0)
+			try
 			{
-				try
-				{
-					string AppId = GetAppId(BundleId);
+				string AppId = GetAppId(BundleId);
 
-					BuildId = GetBuildId(AppId, PreReleaseVersion, BundleVersion);
-				}
-				catch (AggregateException Ex)
-				{
-					Ex.Handle(x => 
+				BuildId = GetBuildId(AppId, PreReleaseVersion, BundleVersion);
+			}
+			catch (AggregateException Ex)
+			{
+				Ex.Handle(x => 
+					{
+						if (x is HttpRequestException)
 						{
-							if (x is HttpRequestException)
-							{
-								Logger.LogWarning("Got an exception while checking for looking for uploaded version: '{Msg}'. Will try again...", Ex.Message);
-								return true;
-							}
-							return false;
+							Logger.LogWarning("Got an exception while checking for looking for uploaded version: '{Msg}'. Will try again...", Ex.Message);
+							return true;
 						}
-					);
-				}
-
-				if (BuildId == null)
-				{
-					Logger.LogInformation("Build not found yet, waiting and will try again...");
-					Thread.Sleep(Delay * 1000);
-				}
+						return false;
+					}
+				);
 			}
 
 			if (BuildId == null)
 			{
-				Logger.LogWarning("Unable to find build for BundleId: {BundleId}, App Version {PreReleaseVersion}, Build Version {BundleVersion}. It could be due to:\n" + 
-					"  * The build was never actually uploaded to AppStoreConnect to TestFlight\n" + 
-					"  * A later version was already there, and Apple changed the version number automatically. Run again with -CheckVersion -SkipUpload to check\n" + 
-					"  * There was a problem with the app. Check your email (or others on your team) to see if Apple rejected the app after uploading\n" +
-					"  * The build is potentially still validating and taking longer than our timeout.",
-					BundleId, PreReleaseVersion, BundleVersion);
-				return false;
+				Logger.LogInformation("Build not found yet, waiting and will try again...");
+				Thread.Sleep(Delay * 1000);
 			}
-			// in the case that we don't use "most-recent" build, then we are looking for exact build version match to bundle version
-			string BuildVersion = BundleVersion;
-			
-			// string BuildVersion;
-			// string BuildId = GetMostRecentBuildId(AppId, PreReleaseVersion, out BuildVersion);
-			// if (BuildId == null)
-			// {
-			// 	Logger.LogError("Unable to find a recent build for version {Version} of {BundleId}", PreReleaseVersion, BundleId);
-			// 	return false;
-			// }
-
-			Logger.LogInformation("Found a build with version: {BuildVersion} [id is {Id}]", BuildVersion, BuildId);
-
-			int Attempts = 0;
-			while (Attempts < 5)
+		}
+
+		if (BuildId == null)
+		{
+			Logger.LogWarning("Unable to find build for BundleId: {BundleId}, App Version {PreReleaseVersion}, Build Version {BundleVersion}. It could be due to:\n" + 
+				"  * The build was never actually uploaded to AppStoreConnect to TestFlight\n" + 
+				"  * A later version was already there, and Apple changed the version number automatically. Run again with -CheckVersion -SkipUpload to check\n" + 
+				"  * There was a problem with the app. Check your email (or others on your team) to see if Apple rejected the app after uploading",
+				BundleId, PreReleaseVersion, BundleVersion);
+			return false;
+		}
+		// in the case that we don't use "most-recent" build, then we are looking for exact build version match to bundle version
+		string BuildVersion = BundleVersion;
+
+		// string BuildVersion;
+		// string BuildId = GetMostRecentBuildId(AppId, PreReleaseVersion, out BuildVersion);
+		// if (BuildId == null)
+		// {
+		// 	Logger.LogError("Unable to find a recent build for version {Version} of {BundleId}", PreReleaseVersion, BundleId);
+		// 	return false;
+		// }
+
+		Logger.LogInformation("Found a build with version: {BuildVersion} [id is {Id}]", BuildVersion, BuildId);
+
+		int Attempts = 0;
+		while (Attempts < 5)
+		{
+			if (Attempts > 0)
+			{
+				Thread.Sleep(30 * 1000);
+			}
+			Attempts++;
+
+			try
 			{
-				if (Attempts > 0)
+				string LocId = GetLocalizationId(BuildId);	
+
+				bool bUpdateSucceeded;
+				if (LocId == null)
 				{
-					Thread.Sleep(30 * 1000);
+					bUpdateSucceeded = CreateWhatsNew(BuildId, Text);
 				}
-				Attempts++;
-
-				try
+				else 
 				{
-					string LocId = GetLocalizationId(BuildId);	
-
-					bool bUpdateSucceeded;
-					if (LocId == null)
-					{
-						bUpdateSucceeded = CreateWhatsNew(BuildId, Text);
-					}
-					else 
-					{
-						bUpdateSucceeded = SetWhatsNew(LocId, Text);
-					}
-					if (bUpdateSucceeded)
-					{
-						GetLocalizationId(BuildId);
-						return true;
-					}
-					return false;
+					bUpdateSucceeded = SetWhatsNew(LocId, Text);
 				}
-				catch (AggregateException Ex)
+				if (bUpdateSucceeded)
 				{
-					Ex.Handle(x => 
-						{
-							Logger.LogWarning("Got an unknown exception ('{Exception}') while updating the WhatsNew of the build. Will try again...", Ex.Message);
-							return true;
-						}
-					);
-					return false;
+					GetLocalizationId(BuildId);
+					return true;
 				}
+				return false;
+			}
+			catch (AggregateException Ex)
+			{
+				Ex.Handle(x => 
+					{
+						Logger.LogWarning("Got an unknown exception ('{Exception}') while updating the WhatsNew of the build. Will try again...", Ex.Message);
+						return true;
+					}
+				);
+				return false;
 			}
-			
-			return false;
 		}
+
+		return false;
 	}
 
 	static void GetInfoFromXcArchive(DirectoryReference XcArchive, out string BundleId, out string AppVersion, out string BundleVersion, out string CommandLine)
@@ -673,18 +651,54 @@ class ExportAppleBuild : BuildCommand
 			}
 		}
 	}
+}
 
-	static string CreateJWTToken(FileReference ProjectFile)
+class JwtTokenClient
+{
+	public JwtTokenClient(FileReference ProjectFile)
+	{
+		projectFile = ProjectFile;
+		
+		client = new HttpClient();
+		client.BaseAddress = new Uri("https://api.appstoreconnect.apple.com/v1/");
+	}
+	
+	public HttpClient Client
+	{
+		get
+		{
+			CheckExpiration();
+			return client;
+		}
+	}
+	
+	private DateTime expirationDate = DateTime.MinValue;
+	private string token = "";
+	private FileReference projectFile;
+	private HttpClient client;
+
+	void CheckExpiration()
+	{
+		if (token == "" || DateTime.UtcNow.AddSeconds(30) > expirationDate)
+		{
+			CreateJWTToken();
+
+			// update the auth header
+			client.DefaultRequestHeaders.Authorization = new AuthenticationHeaderValue("Bearer", token);
+		}
+	}
+	
+	void CreateJWTToken()
 	{
 		// @todo: refactor the AppleExports.GetXcodeBuildAuthOptions() code so that we can get this info in a shared way
-		ConfigHierarchy SharedPlatformIni = ConfigCache.ReadHierarchy(ConfigHierarchyType.Engine, ProjectFile?.Directory, UnrealTargetPlatform.Mac);
+		ConfigHierarchy SharedPlatformIni = ConfigCache.ReadHierarchy(ConfigHierarchyType.Engine, projectFile?.Directory, UnrealTargetPlatform.Mac);
 		string IssuerID, KeyID, KeyPath;
 		string P8Path;
 		if (SharedPlatformIni.TryGetValue("/Script/MacTargetPlatform.XcodeProjectSettings", "AppStoreConnectIssuerID", out IssuerID) &&
 			SharedPlatformIni.TryGetValue("/Script/MacTargetPlatform.XcodeProjectSettings", "AppStoreConnectKeyID", out KeyID) &&
 			SharedPlatformIni.TryGetValue("/Script/MacTargetPlatform.XcodeProjectSettings", "AppStoreConnectKeyPath", out KeyPath))
 		{
-			P8Path = AppleExports.ConvertFilePath(ProjectFile?.Directory, KeyPath).FullName;
+			P8Path = AppleExports.ConvertFilePath(projectFile?.Directory, KeyPath).FullName;
 		}
 		else
 		{
@@ -721,23 +735,24 @@ class ExportAppleBuild : BuildCommand
 		SigningCredentials Credentials = new SigningCredentials(SecurityKey, SecurityAlgorithms.EcdsaSha256);
 		JwtHeader Header = new JwtHeader(Credentials);
 		
+		// set expiration date to check Token validity
+		expirationDate = DateTime.UtcNow.AddMinutes(19);
+		
 		// set up the payload manually with the body that AppStoreConnect needs (aud, exp, and iss - iat and scope are optional)
 		// @todo: use scope to limit the token's access
 		JwtPayload Payload = new JwtPayload
 		{
 			{ "aud", "appstoreconnect-v1" },
-			{ "exp", (long)(DateTime.UtcNow - DateTime.UnixEpoch).TotalSeconds + 60 * 19 },
+			{ "exp", (long)(expirationDate - DateTime.UnixEpoch).TotalSeconds },
 			{ "iss", IssuerID }
 		};
 
 		// make a token and output as string
 		JwtSecurityToken JwtToken = new JwtSecurityToken(Header, Payload);
-		string FullToken = new JwtSecurityTokenHandler().WriteToken(JwtToken);
+		token = new JwtSecurityTokenHandler().WriteToken(JwtToken);
 
 		// When debugging things with curl, you'll need this token, but we don't write it by default since it could technically 
 		// be used to hack ASC within 20 minutes
-		//	Console.WriteLine("Auth header:\n  Authorization: Bearer {0}", FullToken);
-		
-		return FullToken;
+//		Console.WriteLine("Auth header:\n  Authorization: Bearer {0}", token);
 	}
 }
