Fixed some edge cases that caused new PO files to fail to add to P4

diff --git a/Engine/Source/Programs/AutomationTool/AutomationUtils/P4Utils.cs b/Engine/Source/Programs/AutomationTool/AutomationUtils/P4Utils.cs
--- a/Engine/Source/Programs/AutomationTool/AutomationUtils/P4Utils.cs
+++ b/Engine/Source/Programs/AutomationTool/AutomationUtils/P4Utils.cs
@@ -3319,14 +3319,13 @@ namespace AutomationTool
 		/// </summary>
 		/// <param name="DepotFile">Depot path</param>
 		/// <param name="AllowSpew"></param>
-		/// <returns>List of records describing the file's mapping. Usually just one, but may be more.</returns>
-		public bool FileExistsInDepot(string DepotFile, bool AllowSpew = true)
+		/// <returns>True if the file exists in the depot with a non-deleted head revision.</returns>
+		public bool FileExistsInDepot(string DepotFile, bool AllowSpew = false)
 		{
 			string CommandLine = String.Format("fstat {0}", CommandUtils.MakePathSafeToUseWithCommandLine(DepotFile));
 
 			string Output;
-			if(!LogP4Output(out Output, "-z tag", CommandLine, AllowSpew: false) || !Output.Contains("headRev"))
-
+			if (!LogP4Output(out Output, "-z tag", CommandLine, AllowSpew: AllowSpew) || !Output.Contains("headRev") || Output.Contains("headAction delete"))
 			{
 				return false;
 			}
diff --git a/Engine/Source/Programs/AutomationTool/Scripts/Localisation.Automation.cs b/Engine/Source/Programs/AutomationTool/Scripts/Localisation.Automation.cs
--- a/Engine/Source/Programs/AutomationTool/Scripts/Localisation.Automation.cs
+++ b/Engine/Source/Programs/AutomationTool/Scripts/Localisation.Automation.cs
@@ -933,14 +933,17 @@ class Localize : BuildCommand
 			var CurrentPOFileHashes = GetPOFileHashes(LocalizationBatches, UEProjectRoot);
 			foreach (var CurrentPOFileHashPair in CurrentPOFileHashes)
 			{
-				byte[] InitialPOFileHash;
-				if (InitalPOFileHashes.TryGetValue(CurrentPOFileHashPair.Key, out InitialPOFileHash) && InitialPOFileHash.SequenceEqual(CurrentPOFileHashPair.Value))
+				if (P4.FileExistsInDepot(CurrentPOFileHashPair.Key))
 				{
-					POFilesToRevert.Add(CurrentPOFileHashPair.Key);
-				}
-				else if (P4.FileExistsInDepot(CurrentPOFileHashPair.Key))
-				{
-					POFilesToEdit.Add(CurrentPOFileHashPair.Key);
+					byte[] InitialPOFileHash;
+					if (InitalPOFileHashes.TryGetValue(CurrentPOFileHashPair.Key, out InitialPOFileHash) && InitialPOFileHash.SequenceEqual(CurrentPOFileHashPair.Value))
+					{
+						POFilesToRevert.Add(CurrentPOFileHashPair.Key);
+					}
+					else
+					{
+						POFilesToEdit.Add(CurrentPOFileHashPair.Key);
+					}
 				}
 				else
 				{
