Fixed issue where Android.GetSdkDir() would return an empty string rather than the fallback value (DefaultSdkDir).

diff --git a/Engine/Source/Programs/AutomationTool/Android/AndroidPlatform.Automation.cs b/Engine/Source/Programs/AutomationTool/Android/AndroidPlatform.Automation.cs
--- a/Engine/Source/Programs/AutomationTool/Android/AndroidPlatform.Automation.cs
+++ b/Engine/Source/Programs/AutomationTool/Android/AndroidPlatform.Automation.cs
@@ -392,7 +392,7 @@ public class AndroidPlatform : Platform
 
 		string DefaultSdkDir = Path.Combine(Environment.GetFolderPath(Environment.SpecialFolder.LocalApplicationData), "Android", "Sdk");
 		string RegValue = Microsoft.Win32.Registry.GetValue("HKEY_LOCAL_MACHINE\\SOFTWARE\\Android Studio", "SdkPath", null) as string;
-		return RegValue == null ? DefaultSdkDir : RegValue;
+		return string.IsNullOrEmpty(RegValue) ? DefaultSdkDir : RegValue;
 	}
 
 	public override bool UpdateHostPrerequisites(BuildCommand Command, ITurnkeyContext TurnkeyContext, bool bVerifyOnly)
