LaunchMultiServer now supports -BotCount=XX which allows you to specify an overall number of bots which works in both -client=proxy and -client=direct mode.

diff --git a/Engine/Source/Programs/AutomationTool/Scripts/LaunchMultiServer.Automation.cs b/Engine/Source/Programs/AutomationTool/Scripts/LaunchMultiServer.Automation.cs
--- a/Engine/Source/Programs/AutomationTool/Scripts/LaunchMultiServer.Automation.cs
+++ b/Engine/Source/Programs/AutomationTool/Scripts/LaunchMultiServer.Automation.cs
@@ -257,13 +257,16 @@ Example running 4 clients locally connecting to single proxy server w/ a non-loc
 		[CommandLine, Help("DirectClientCount=##", "Specify how many clients to launch in client=direct mode.  If 0 (default), will launch one for each server.")]
 		public int DirectClientCount = 0;
 
+		[CommandLine, Help("BotCount=##", "Specify how many client bots to launch. Client bots are launched with -bot it is up to the game to specify the implementation.")]
+		public int BotCount = 0;
+
 		[CommandLine, Help("ProxyServerCount=##", "Specify how many proxy servers to launch in client=proxy mode")]
 		public int ProxyServerCount = 1;
 
 		[CommandLine, Help("ProxyClientCount=##", "Specify how many clients to launch PER SERVER in client=proxy mode")]
 		public int ProxyClientCount = 1;
 
-		[CommandLine, Help("ProxyBotCount=##", "Specify how many bots to launch PER SERVER in client=proxy mode")]
+		[CommandLine, Help("ProxyBotCount=##", "Specify how many bots to launch PER SERVER in client=proxy mode. Prefer -BotCount=XX.")]
 		public int ProxyBotCount = 0;
 
 		[CommandLine, Help("ProxyClientPrimary=##", "Specify which game server to prefer as the initial proxy connection, as an integer or the string \"random\"")]
@@ -302,7 +305,7 @@ Example running 4 clients locally connecting to single proxy server w/ a non-loc
 			DirectClient,
 			ProxyServer,
 			ProxyClient,
-			ProxyBot
+			BotClient,
 		}
 
 		protected virtual void ParseCommandLineArguments()
@@ -376,17 +379,39 @@ Example running 4 clients locally connecting to single proxy server w/ a non-loc
 						DirectClientCount = MaxDirectClientCount;
 					}
                     Logger.LogInformation($"Starting {DirectClientCount} direct client instances connecting to MultiServer instances");
-                    StartProcessesForClients(ProcessLaunchType.DirectClient, AllGameServerAddresses, CommonArgs, 1, DirectClientCount);
+					StartProcessesForClients(ProcessLaunchType.DirectClient, AllGameServerAddresses, CommonArgs, 1, DirectClientCount);
+
+					if (BotCount > 0)
+					{
+						// NumBotsPerServer = Ceil(BotCount / ServerCount);
+						int ServerCount = CountServers(AllGameServerAddresses);
+						int NumBotsPerServer = (BotCount + ServerCount - 1) / ServerCount;
+
+						Logger.LogInformation($"Starting {BotCount} bots connecting to {string.Join(',', AllGameServerAddresses)}. Roughly {NumBotsPerServer} per server.");
+						StartProcessesForClients(ProcessLaunchType.BotClient, AllGameServerAddresses, $"{CommonArgs} -nullrhi -bot -nosound -unattended -nosplash", NumBotsPerServer, BotCount);
+					}
 				}
 				else if (ClientValue.ToLower() == "proxy")
 				{
-					Logger.LogInformation($"Starting {ProxyServerCount} Proxy Servers each with {ProxyClientCount} Clients and {ProxyBotCount} Bots");
+					Logger.LogInformation($"Starting {ProxyServerCount} Proxy Servers each with {ProxyClientCount} Clients");
 
 					string[] LocalProxyAddresses = GenerateLocalServerAddressRange(ProxyServerCount, ProxyServerBasePort);
 					StartProcessesForProxyServer(LocalProxyAddresses, AllGameServerAddresses, CommonArgs);
-
 					StartProcessesForClients(ProcessLaunchType.ProxyClient, LocalProxyAddresses, CommonArgs, ProxyClientCount);
-					StartProcessesForClients(ProcessLaunchType.ProxyBot, LocalProxyAddresses, $"{CommonArgs} -nullrhi -bot -nosound -unattended -nosplash", ProxyBotCount);
+
+					if (ProxyBotCount > 0)
+					{
+						Logger.LogInformation($"Starting {ProxyBotCount} bots for each proxy address: {string.Join(',', LocalProxyAddresses)}.");
+						StartProcessesForClients(ProcessLaunchType.BotClient, LocalProxyAddresses, $"{CommonArgs} -nullrhi -bot -nosound -unattended -nosplash", ProxyBotCount);
+					}
+
+					if (BotCount > 0)
+					{
+						// NumBotsPerProxy = Ceil(BotCount / ProxyServerCount);
+						int NumBotsPerProxy = (BotCount + ProxyServerCount - 1) / ProxyServerCount;
+						Logger.LogInformation($"Starting {BotCount} bots connecting to {string.Join(',', LocalProxyAddresses)}. Roughly {NumBotsPerProxy} per proxy.");
+						StartProcessesForClients(ProcessLaunchType.BotClient, LocalProxyAddresses, $"{CommonArgs} -nullrhi -bot -nosound -unattended -nosplash", NumBotsPerProxy, BotCount);
+					}
 				}
 				else
 				{
@@ -553,12 +578,7 @@ Example running 4 clients locally connecting to single proxy server w/ a non-loc
                         {
                             break;
                         }
-					}
-                
-                    if (MaxClientCount != 0 && LaunchedClientCount >= MaxClientCount)
-                    {
-                        break;
-                    }
+					}                
                 }
             }
 		}
@@ -626,7 +646,7 @@ Example running 4 clients locally connecting to single proxy server w/ a non-loc
 				Command: this,
 				RawProjectPath: ProjectPath,
 				DedicatedServer: LaunchType == ProcessLaunchType.Server || LaunchType == ProcessLaunchType.ProxyServer,
-				Client: LaunchType == ProcessLaunchType.DirectClient || LaunchType == ProcessLaunchType.ProxyClient || LaunchType == ProcessLaunchType.ProxyBot
+				Client: LaunchType == ProcessLaunchType.DirectClient || LaunchType == ProcessLaunchType.ProxyClient || LaunchType == ProcessLaunchType.BotClient
 			);
 
 			var DeployContexts = AutomationScripts.Project.CreateDeploymentContext(Params, Params.DedicatedServer);
@@ -734,6 +754,26 @@ Example running 4 clients locally connecting to single proxy server w/ a non-loc
 			return true;
 		}
 
+		/// <summary>
+		/// Counts the number of servers specified in the passed-in string array. We account for port ranges.
+		/// </summary>
+		/// <param name="ServerAddresses">The server addresses (IpAddress:StartPort-EndPort)</param>
+		/// <returns>The number of servers specified (taking into account port ranges)</returns>
+		protected int CountServers(string[] ServerAddresses)
+		{
+			int TotalCount = 0;
+			foreach (string AddressToParse in ServerAddresses)
+			{
+				string BaseAddress;
+				int StartPort, EndPort;
+				ParseAddressAndPortRange(out BaseAddress, out StartPort, out EndPort, AddressToParse);
+
+				TotalCount += EndPort - StartPort + 1;
+			}
+
+			return TotalCount;
+		}
+
 		/// <summary>
 		/// Are we requesting a debugger be attached to this particular process?
 		/// </summary>
