[Gauntlet] OverlayExecutables will search for and prioritize exectuables with a "-Preflight" appended to the name.
This is useful for overlaying mobile builds where re-signing can cause build graph scripts to complain about build output modifications

diff --git a/Engine/Source/Programs/AutomationTool/Gauntlet/Framework/Base/Gauntlet.OverlayExecutable.cs b/Engine/Source/Programs/AutomationTool/Gauntlet/Framework/Base/Gauntlet.OverlayExecutable.cs
--- a/Engine/Source/Programs/AutomationTool/Gauntlet/Framework/Base/Gauntlet.OverlayExecutable.cs
+++ b/Engine/Source/Programs/AutomationTool/Gauntlet/Framework/Base/Gauntlet.OverlayExecutable.cs
@@ -57,9 +57,22 @@ namespace Gauntlet
 
 			// Ex. /UnrealGame/Binaries/Win64/UnrealClient.exe
 			string PlatformBinariesDirectory = Path.Combine(Environment.CurrentDirectory, ProjectName, "Binaries", Platform.ToString());
-			string ExecutableName = UnrealHelpers.GetExecutableName(ProjectName, Platform, Configuration, Role, Flavor, string.Empty);
 			string ExecutableExtension = string.IsNullOrEmpty(ExtensionOverride) ? Path.GetExtension(BaseExecutable) : ExtensionOverride;
+
+			string ExecutableName = UnrealHelpers.GetExecutableName(ProjectName, Platform, Configuration, Role, Flavor, string.Empty);
+			string PreflightExecutableName = ExecutableName + "-Preflight";
+
 			string LocalExecutable = Path.Combine(PlatformBinariesDirectory, ExecutableName + ExecutableExtension);
+			string LocalPreflightExecutable = Path.Combine(PlatformBinariesDirectory, PreflightExecutableName + ExecutableExtension);
+
+			// Prioritize using -Preflight executables, if they exist
+			if ((File.Exists(LocalPreflightExecutable) || Directory.Exists(LocalPreflightExecutable))
+				&& File.GetLastWriteTime(LocalPreflightExecutable) > File.GetLastWriteTime(BaseExecutable))
+			{
+				Log.Info("Applying newer preflight executable as overlay {PreflightExecutable}", LocalPreflightExecutable);
+				OverlayExecutable = LocalPreflightExecutable;
+				return true;
+			}
 
 			if (!File.Exists(LocalExecutable) && !Directory.Exists(LocalExecutable))
 			{
