LaunchMultiServer: Fixed DirectClientCount regression

Fixed MaxClientCount logic regression allowing more than the max number of clients to be run

CL 47538432 regression

diff --git a/Engine/Source/Programs/AutomationTool/Scripts/LaunchMultiServer.Automation.cs b/Engine/Source/Programs/AutomationTool/Scripts/LaunchMultiServer.Automation.cs
--- a/Engine/Source/Programs/AutomationTool/Scripts/LaunchMultiServer.Automation.cs
+++ b/Engine/Source/Programs/AutomationTool/Scripts/LaunchMultiServer.Automation.cs
@@ -576,10 +576,10 @@ Example running 4 clients locally connecting to single proxy server w/ a non-loc
                         ++LaunchedClientCount;
                         if (MaxClientCount != 0 && LaunchedClientCount >= MaxClientCount)
                         {
-                            break;
+							return;
                         }
 					}                
-                }
+				}
             }
 		}
 
