Horde Test Report - Ensure unique tags

#rnx

diff --git a/Engine/Source/Programs/AutomationTool/Gauntlet/Framework/Utils/Gauntlet.HordeReport.cs b/Engine/Source/Programs/AutomationTool/Gauntlet/Framework/Utils/Gauntlet.HordeReport.cs
--- a/Engine/Source/Programs/AutomationTool/Gauntlet/Framework/Utils/Gauntlet.HordeReport.cs
+++ b/Engine/Source/Programs/AutomationTool/Gauntlet/Framework/Utils/Gauntlet.HordeReport.cs
@@ -14,6 +14,7 @@ using Logging = Microsoft.Extensions.Logging;
 using System.Globalization;
 using static Gauntlet.HordeReport.TestDataCollection;
 using Gauntlet.Report;
+using System.Diagnostics.CodeAnalysis;
 
 namespace Gauntlet
 {
@@ -1494,7 +1495,7 @@ namespace Gauntlet
 				{
 					if (tags == null)
 					{
-						tags = new();
+						tags = new(new TagComparer());
 					}
 					tags.UnionWith(NewTags);
 				}
@@ -1510,6 +1511,27 @@ namespace Gauntlet
 				}
 			}
 
+			internal class TagComparer : IEqualityComparer<string>
+			{
+				public bool Equals(string? TagA, string? TagB)
+				{
+					if (TagA == null && TagB == null)
+					{
+						return true;
+					}
+					else if (TagA == null || TagB == null)
+					{
+						return false;
+					}
+					return TagA.Equals(TagB, StringComparison.CurrentCultureIgnoreCase);
+				}
+
+				public int GetHashCode([DisallowNull] string Tag)
+				{
+					return Tag.ToLower().GetHashCode();
+				}
+			}
+
 			/// <summary>
 			/// fork of the test session data (summary, phases, devices, eventStreams)
 			/// </summary>
@@ -1677,7 +1699,7 @@ namespace Gauntlet
 			{
 				if (Data.tags == null)
 				{
-					Data.tags = new();
+					Data.tags = new(new TagComparer());
 				}
 				Data.tags.UnionWith(NewTags);
 			}
