LaunchMultiServer UAT script now uses more consistent & descriptive log filenames for launched processes.

diff --git a/Engine/Source/Programs/AutomationTool/Scripts/LaunchMultiServer.Automation.cs b/Engine/Source/Programs/AutomationTool/Scripts/LaunchMultiServer.Automation.cs
--- a/Engine/Source/Programs/AutomationTool/Scripts/LaunchMultiServer.Automation.cs
+++ b/Engine/Source/Programs/AutomationTool/Scripts/LaunchMultiServer.Automation.cs
@@ -517,7 +517,7 @@ Example running 4 clients locally connecting to single proxy server w/ a non-loc
 				Logger.LogInformation("Starting Dedicated MultiServer for Game Port {0}", ServerGamePort);
 
 				string ServerArgs = ServerCommonArgs.Replace("{ServerId}", ServerId.ToString(), StringComparison.OrdinalIgnoreCase);
-				ServerArgs += String.Format(" -port={0} -log=MultiServer-{1}.log", ServerGamePort, ServerId);
+				ServerArgs += $" -port={ServerGamePort} -log={ProjectFullPath.GetFileNameWithoutAnyExtensions()}-Server-{ServerId}.log";
 
 				string AppTitle = $"Server ID {ServerId}";
 				ExecuteApp(ProcessLaunchType.Server, AppTitle, ServerArgs);
@@ -542,10 +542,10 @@ Example running 4 clients locally connecting to single proxy server w/ a non-loc
         /// <param name="MaxClientCount">The maximum number of clients to launch.  Zero for no limit (one for each server address or proxy server port)</param>
         virtual protected void StartProcessesForClients(ProcessLaunchType ClientType, string[] GameServerAddresses, string AdditionalClientArguments, int ClientsPerServer = 1, int MaxClientCount = 0)
 		{
-            int LaunchedClientCount = 0;
-            for (int Idx= 0; Idx < GameServerAddresses.Length; ++Idx)
+			int TotalClientCount = 0;
+			for (int ServerAddressIdx= 0; ServerAddressIdx < GameServerAddresses.Length; ++ServerAddressIdx)
 			{
-				string GameServerAddressWithPortRange = GameServerAddresses[Idx];
+				string GameServerAddressWithPortRange = GameServerAddresses[ServerAddressIdx];
 				string GameServerBaseAddress;
 				int StartPort, EndPort;
 
@@ -561,26 +561,28 @@ Example running 4 clients locally connecting to single proxy server w/ a non-loc
 					string ClientArgs = $" -game {AdditionalClientArguments}";
 					ClientArgs += NoTimeouts ? " -notimeouts" : "";
 					ClientArgs += $" {GameServerBaseAddress}:{Port}";
-					ClientArgs += $" -log={ClientType}-{{ClientId}}.log -windowed -SaveWinPos={{ClientId}}";
+					ClientArgs += $" -windowed";
 
 					// Potentially launch multiple clients per server address
-					for (int InstanceNum = 0; InstanceNum < ClientsPerServer; ++InstanceNum)
+					for (int ClientInstanceNum = 0; ClientInstanceNum < ClientsPerServer; ++ClientInstanceNum)
 					{
-						int ClientId = (Idx + 1 + InstanceNum);
-						string InstanceArgs = ClientArgs.Replace("{ClientId}", ClientId.ToString(), StringComparison.OrdinalIgnoreCase);
+						int ClientId = (ServerAddressIdx + 1 + ClientInstanceNum);
 
-						string InstanceDetails = (ClientsPerServer > 1) ? $"(Srv {Idx + 1})" : string.Empty;
+						++TotalClientCount;
+						string InstanceArgs = ClientArgs + $" -log={ProjectFullPath.GetFileNameWithoutAnyExtensions()}-{ClientType}-{TotalClientCount}.log";
+						InstanceArgs += $" -SaveWinPos={TotalClientCount}";
+
+						string InstanceDetails = (ClientsPerServer > 1) ? $"(Srv {ServerAddressIdx + 1})" : string.Empty;
 						string AppTitle = $"{ClientType} {ClientId} {InstanceDetails}";
 						ExecuteApp(ClientType, AppTitle, InstanceArgs);
 
-                        ++LaunchedClientCount;
-                        if (MaxClientCount != 0 && LaunchedClientCount >= MaxClientCount)
-                        {
+						if (MaxClientCount != 0 && TotalClientCount >= MaxClientCount)
+						{
 							return;
-                        }
-					}                
+						}
+					}
 				}
-            }
+			}
 		}
 
 		virtual protected void StartProcessesForProxyServer(string[] LocalProxyServerAddresses, string[] AllGameServerAddresses, string AdditionalProxyArguments)
@@ -624,6 +626,8 @@ Example running 4 clients locally connecting to single proxy server w/ a non-loc
 					.Replace("{ProxyId}", ProxyIdx.ToString(), StringComparison.OrdinalIgnoreCase)
 					.Replace("{ProxyServerPort}", ProxyServerPort.ToString(), StringComparison.OrdinalIgnoreCase);
 
+				ProxyServerArgs += $" -log={ProjectFullPath.GetFileNameWithoutAnyExtensions()}-Proxy-{ProxyIdx}.log";
+
 				string AppTitle = $"ProxyServer {ProxyIdx}";
 				ExecuteApp(ProcessLaunchType.ProxyServer, AppTitle, ProxyServerArgs);
 			}
