Localization:
Fixed bug in Localize UAT command where -EnableIncludedPlugins will only enable plugins if -AdditionalCommandletArguments was specified.
- Added feature to filter out GFPs from the provided -IncludePlugins/IncludePluginsDirectory lists to have localization respect the release versions of GFPs.
#test Ran the Localize command locally to make sure that all GFPs were filtered out

diff --git a/Engine/Source/Programs/AutomationTool/Scripts/Localisation.Automation.cs b/Engine/Source/Programs/AutomationTool/Scripts/Localisation.Automation.cs
--- a/Engine/Source/Programs/AutomationTool/Scripts/Localisation.Automation.cs
+++ b/Engine/Source/Programs/AutomationTool/Scripts/Localisation.Automation.cs
@@ -687,68 +687,106 @@ class Localize : BuildCommand
 			EditorArguments += " -multiprocess";
 		}
 
+		AppendEnableIncludedPluginsArgument();
+
+		if (!String.IsNullOrEmpty(AdditionalCommandletArguments))
+		{
+			EditorArguments += " " + AdditionalCommandletArguments;
+		}
+		return EditorArguments;
+	}
+
+	private void AppendEnableIncludedPluginsArgument()
+	{
+		if (!bShouldEnableIncludedPlugins)
+		{
+			Logger.LogInformation("Included plugins will not be auto enabled in commandlets.");
+			return;
+		}
+
+		Logger.LogInformation($"Included plugins requested to be enabled. All specified included plugins that are not excluded and pass filtering will be enabled.");
 		// We append all the included plugins to -EnablePlugins if -EnableIncludedPlugins is enabled. This wil ensure that the plugin content and metadata will be loaded.
 		// @TODOLocalization: Ideally the enabling of plugins should be per batch, otherwise each instance of the editor is enabling a bunch of plugins it doesn't need 
-		if (!string.IsNullOrEmpty(AdditionalCommandletArguments) && bShouldEnableIncludedPlugins && IncludePlugins.Count > 0)
+		HashSet<string> PluginsToEnableSet = IncludePlugins.Except(ExcludePlugins).ToHashSet();
+		string EnablePluginsToken = "-EnablePlugins=";
+		if (!string.IsNullOrEmpty(AdditionalCommandletArguments))
 		{
-			HashSet<string> PluginsToEnableSet = IncludePlugins.Except(ExcludePlugins).ToHashSet();
-
 			// It's possible that there are already specified values for -EnabledPlugins, we willneed to try and parse them first.
-			string EnablePluginsToken = "-EnablePlugins=";
-			string EnablePluginsNewValue = "";
 
 			int EnablePluginsTokenIndex = AdditionalCommandletArguments.IndexOf(EnablePluginsToken);
 			string EnablePluginsOldValue = "";
 			if (EnablePluginsTokenIndex > -1)
 			{
 				// -EnablePlugins token exists in the additional commandlet args. We need to process it 
+				Logger.LogInformation("-EnablePlugins already provided as an argument. Adding plugins provided to the full ist of plugins to enable.");
 				int EnablePluginsValueStartIndex = EnablePluginsTokenIndex + EnablePluginsToken.Length;
-				// We try and find the end of the string where it's separated by a space  between the next token 
-				int EnablePluginsValueEndIndex = AdditionalCommandletArguments.IndexOf(" ", EnablePluginsValueStartIndex);
-				// We can't find a next space. THis means we're the last parameter in AdditionalCommandletArguments. The end index will be the length of the string 
-				if (EnablePluginsValueEndIndex == -1)
+				// We try and find the end of the string where it's separated by a space  between the next token. IndexOf() returns -1 if it can't find the space.
+				// We decrement by 1 as we need the index of the last character not the index of the string 
+				int EnablePluginsValueEndIndex = AdditionalCommandletArguments.IndexOf(" ", EnablePluginsValueStartIndex) - 1;
+				// We can't find a next space. This means we're the last parameter in AdditionalCommandletArguments. The end index will be the length of the string - 1
+				if (EnablePluginsValueEndIndex < 0)
 				{
-					EnablePluginsValueEndIndex = AdditionalCommandletArguments.Length;
+					EnablePluginsValueEndIndex = AdditionalCommandletArguments.Length - 1;
 				}
+
 				// Isolate the value of -EnablePlugins and add them to our list of plugins to enable 
-				EnablePluginsOldValue = AdditionalCommandletArguments.Substring(EnablePluginsValueStartIndex, EnablePluginsValueEndIndex - EnablePluginsValueStartIndex);
+				EnablePluginsOldValue = AdditionalCommandletArguments.Substring(EnablePluginsValueStartIndex, EnablePluginsValueEndIndex - EnablePluginsValueStartIndex + 1);
 				foreach (string Plugin in EnablePluginsOldValue.Split(','))
 				{
-					PluginsToEnableSet.Add(Plugin);
+					PluginsToEnableSet.Add(Plugin.Trim());
 				}
+				// Now that we have parsed the values, we remove the field from the command line to add it in below 
+				AdditionalCommandletArguments = AdditionalCommandletArguments.Remove(EnablePluginsTokenIndex, EnablePluginsValueEndIndex - EnablePluginsTokenIndex + 1);
+				Logger.LogInformation($"Removed '{EnablePluginsToken}' argument from provided AdditionalCommandletArguments. Remaining value is '-AdditionalCommandletArgument={AdditionalCommandletArguments}'");
 			}
+		}
 
-			// Just a counter to help iterate through the set to build out the comma separated value 
-			int IterationCount = 0;
-			StringBuilder EnablePluginsBuilder = new StringBuilder();
-			foreach (string Plugin in PluginsToEnableSet)
+		if (PluginsToEnableSet.Count == 0)
+		{
+			Logger.LogInformation("No included plugins to enable.");
+			return;
+		}
+
+		Logger.LogInformation("Filtering out GFP plugins from being enabled.");
+		string GameFeaturesPluginDirectory = "GameFeatures";
+		// We filter out GFPs from being auto-enabled in commandlets as we want the release versions of plugisn to be respected
+		// The GFP system has ultimate authority of whether a GFP can be loaded. Thus force loading plugins can fail if the GFP system determines that it should not be loaded and cause errors. We bypass that by filtering GFPs out. 
+		List<string> PluginsToRemove = new();
+		foreach (string Plugin in PluginsToEnableSet)
+		{
+			PluginInfo Info = Plugins.GetPlugin(Plugin);
+			if (Info is null)
 			{
-				EnablePluginsBuilder.Append(Plugin);
-				if (IterationCount < PluginsToEnableSet.Count - 1)
-				{
-					EnablePluginsBuilder.Append(",");
-				}
-				++IterationCount;
+				PluginsToRemove.Add(Plugin);
+				continue;
 			}
-			EnablePluginsNewValue = EnablePluginsBuilder.ToString();
-			Logger.LogInformation($"Appending following plugins to be enabled: {EnablePluginsNewValue}");
-			// if we already had a value of -EnablePlugins in AdditionalCommandletArguments, we'll need to replace that with the new values we've created.
-			if (EnablePluginsTokenIndex > -1)
+
+			if (Info.Directory.ContainsName(GameFeaturesPluginDirectory, 0))
 			{
-				AdditionalCommandletArguments = AdditionalCommandletArguments.Replace(EnablePluginsToken + EnablePluginsOldValue, EnablePluginsToken + EnablePluginsNewValue);
+				PluginsToRemove.Add(Plugin);
 			}
-			else
+		}
+		if (PluginsToRemove.Count > 0)
+		{
+			Logger.LogInformation($"Removing {PluginsToRemove.Count} game feature plugins from plugins to be enabled.");
+			foreach (string Plugin in PluginsToRemove)
 			{
-				// The token doesn't exist, we'll add it to the end 
-				AdditionalCommandletArguments += " " + EnablePluginsToken + EnablePluginsNewValue;
+				Logger.LogInformation($"Removing GFP '{Plugin}' from being enabled.");
+				PluginsToEnableSet.Remove(Plugin);
 			}
 		}
-
-		if (!String.IsNullOrEmpty(AdditionalCommandletArguments))
+		
+		if (PluginsToEnableSet.Count == 0)
 		{
-			EditorArguments += " " + AdditionalCommandletArguments;
+			Logger.LogInformation("All requested included plugins to be enabled have been filtered out. No specified included plugins will be enabled.");
+			return;
 		}
-		return EditorArguments;
+		string EnablePluginsNewValue = String.Join(',', PluginsToEnableSet.ToArray());
+		Logger.LogInformation($"Appending following plugins to be enabled: {EnablePluginsNewValue}");
+
+		// The token was either previously removed or never existed, we'll add it to the end 
+		AdditionalCommandletArguments += " " + EnablePluginsToken + EnablePluginsNewValue;
+		Logger.LogInformation($"AdditionalCommandletArguments='{AdditionalCommandletArguments}'");	
 	}
 
 	private void InitializeLocalizationProvider(List<LocalizationTask> LocalizationTasks)
