Have build machines use "boost mode" (more cores) regardless of the parameters passed to RetrieveCloudArtifact.  This matches the behavior in CompileTask and other places that make decisions about core usage.

diff --git a/Engine/Source/Programs/AutomationTool/BuildGraph/Tasks/RetrieveCloudArtifactTask.cs b/Engine/Source/Programs/AutomationTool/BuildGraph/Tasks/RetrieveCloudArtifactTask.cs
--- a/Engine/Source/Programs/AutomationTool/BuildGraph/Tasks/RetrieveCloudArtifactTask.cs
+++ b/Engine/Source/Programs/AutomationTool/BuildGraph/Tasks/RetrieveCloudArtifactTask.cs
@@ -396,7 +396,7 @@ namespace AutomationTool.Tasks
 				{
 					AllowMultipart = _parameters.AllowMultipart, 
 					AssumeHttp2 = String.Equals(httpVersion, "http2-only", StringComparison.OrdinalIgnoreCase),
-					BoostWorkers = _parameters.BoostWorkers,
+					BoostWorkers = CommandUtils.IsBuildMachine || _parameters.BoostWorkers,
 					TraceFile = traceFile
 				}, optionsOverride: optionsOverride);
 		}
