- Check if app is foregrounded once we launch. If not, foreground it.
- Adds logging to WebdriverAgent process to determine where exactly it's failing

diff --git a/Engine/Source/Programs/AutomationTool/Gauntlet/Platform/IOS/Gauntlet.IOSAppium.cs b/Engine/Source/Programs/AutomationTool/Gauntlet/Platform/IOS/Gauntlet.IOSAppium.cs
--- a/Engine/Source/Programs/AutomationTool/Gauntlet/Platform/IOS/Gauntlet.IOSAppium.cs
+++ b/Engine/Source/Programs/AutomationTool/Gauntlet/Platform/IOS/Gauntlet.IOSAppium.cs
@@ -1,11 +1,12 @@
 // Copyright Epic Games, Inc. All Rights Reserved.
-using System;
-using System.Text.RegularExpressions;
 using AutomationTool;
 using Newtonsoft.Json;
 using OpenQA.Selenium.Appium;
+using OpenQA.Selenium.Appium.Enums;
 using OpenQA.Selenium.Appium.iOS;
 using OpenQA.Selenium.Appium.Service;
+using System;
+using System.Text.RegularExpressions;
 using static AutomationTool.CommandUtils;
 
 namespace Gauntlet
@@ -200,14 +201,15 @@ namespace Gauntlet
 				Disposed = true;
 			}
         }
-        #endregion
+		#endregion
 
-        public void Start(string PackageName)
-        {
-            Stop();
+		public void Start(string PackageName)
+		{
+			Stop();
 
-            try
-            {
+			// Install WebdriverAgent
+			try
+			{
 				IProcessResult InstallResult;
 				using (ScopedSuspendECErrorParsing ErrorSuspension = new())
 				{
@@ -222,14 +224,21 @@ namespace Gauntlet
 							string.Format("-b \"{0}\"", AppiumConfig.WdaAppPath), UUID, 60, AdditionalOptions: ERunOptions.NoStdOutRedirect);
 					}
 				}
-
 				if (InstallResult.ExitCode != 0)
 				{
 					throw new AutomationException("Failed to install WDA app ({0}): {1}", InstallResult.ExitCode, InstallResult.Output);
 				}
+			}
+			catch (Exception Ex)
+			{
+				throw new AutomationException($"Exception during WDA Install {Ex.Message}");
+			}
 
-				lock (Mutex)
-                {
+			lock (Mutex)
+			{
+				// Start Appium Server
+				try
+				{
 					AppiumServiceBuilder ServerBuilder = new AppiumServiceBuilder()
 						.UsingAnyFreePort()
 						.WithIPAddress("127.0.0.1");
@@ -239,19 +248,27 @@ namespace Gauntlet
 						ServerBuilder.WithAppiumJS(new System.IO.FileInfo(AppiumConfig.AppiumLocation));
 					}
 
-                    Server = ServerBuilder.Build();
-                    Server.Start();
-
-                    Driver = new IOSDriver(Server, AppiumConfig.GetCapabilities(UUID, PackageName));
-                }
-            }
-            catch (Exception Ex)
-            {
-                throw new DeviceException("Failed to start appium instance for {0}: {1}", UUID, Ex.Message);
-            }
-        }
+					Server = ServerBuilder.Build();
+					Server.Start();
+				}
+				catch (Exception Ex)
+				{
+					throw new AutomationException($"Exception during Appium Server Start {Ex.Message}");
+				}
+				// Create Webdriver instance
+				try
+				{
+					Driver = new IOSDriver(Server, AppiumConfig.GetCapabilities(UUID, PackageName));
+				}
+				catch (Exception Ex)
+				{
+					Server.Dispose();
+					throw new AutomationException($"Exception during WebdriverAgent Creation {Ex.Message}");
+				}
+			}
+		}
 
-        public void Stop()
+		public void Stop()
         {
             if (Driver != null)
             {
@@ -307,5 +324,36 @@ namespace Gauntlet
 				}
 			}
 		}
-    }
+
+		public bool VerifyAppInForeground(string appID)
+		{
+			bool bInForeground = false;
+			var appState = Driver.GetAppState(appID);
+
+			switch (appState)
+			{
+				case AppState.RunningInForeground:
+					bInForeground = true;
+					Log.Verbose("{0} app in Foreground", appID);
+					break;
+				case AppState.RunningInBackground:
+					Log.Verbose("{0} app in Background", appID);
+					break;
+				case AppState.RunningInBackgroundOrSuspended:
+					Log.Verbose("{0} app in Background or Suspended", appID);
+					break;
+				case AppState.NotRunning:
+					Log.Verbose("{0} app not running", appID);
+					break;
+			}
+
+			return bInForeground;
+		}
+
+		public void ForegroundApp(string appID)
+		{
+			Driver.ActivateApp(appID);
+		}
+
+	}
 }
\ No newline at end of file
diff --git a/Engine/Source/Programs/AutomationTool/Gauntlet/Platform/IOS/Gauntlet.TargetDeviceIOS.cs b/Engine/Source/Programs/AutomationTool/Gauntlet/Platform/IOS/Gauntlet.TargetDeviceIOS.cs
--- a/Engine/Source/Programs/AutomationTool/Gauntlet/Platform/IOS/Gauntlet.TargetDeviceIOS.cs
+++ b/Engine/Source/Programs/AutomationTool/Gauntlet/Platform/IOS/Gauntlet.TargetDeviceIOS.cs
@@ -12,6 +12,7 @@ using Gauntlet.Utils;
 using AutomationTool;
 using UnrealBuildTool;
 using static AutomationTool.CommandUtils;
+using System.Threading;
 
 /** IOS General Notes:
  *
@@ -829,6 +830,29 @@ namespace Gauntlet
 				throw new DeviceException("Failed to launch on {0}. {1}", Name, AppProcess.Output);
 			}
 
+			if (AppiumContainer != null)
+			{
+				lock (Mutex)
+				{
+					bool IsAppInForeground = false;
+					int Retries = 0;
+					while (!IsAppInForeground && Retries < 5)
+					{
+						IsAppInForeground = AppiumContainer.VerifyAppInForeground(IOSInstall.PackageName);
+						if (!IsAppInForeground)
+						{
+							Log.Info("{0} in background, trying to foreground Attempt #{1}", IOSInstall.PackageName, Retries);
+							AppiumContainer.ForegroundApp(IOSInstall.PackageName);
+							Thread.Sleep(500);
+							Retries++;
+						}
+					}
+					if (Retries == 5)
+					{
+						throw new DeviceException("Unable to foreground {0} after {1} attempts", IOSInstall.PackageName, Retries);
+					}
+				}
+			}
 			return new IOSAppInstance(IOSInstall, AppProcess, IOSInstall.CommandLine);
 		}
 
