LaunchMultiServer:  -AttachToServers no longer attaches to proxies.  Use the new -AttachToProxies.

#rnx

diff --git a/Engine/Source/Programs/AutomationTool/Scripts/LaunchMultiServer.Automation.cs b/Engine/Source/Programs/AutomationTool/Scripts/LaunchMultiServer.Automation.cs
--- a/Engine/Source/Programs/AutomationTool/Scripts/LaunchMultiServer.Automation.cs
+++ b/Engine/Source/Programs/AutomationTool/Scripts/LaunchMultiServer.Automation.cs
@@ -287,6 +287,9 @@ Example running 4 clients locally connecting to single proxy server w/ a non-loc
 		[CommandLine, Help("AttachToServers", "Force Visual Studio to attach to the servers when launched (requires VS be launched or it will hang)")]
 		public bool AttachToServers = false;
 
+		[CommandLine, Help("AttachToProxies", "Force Visual Studio to attach to the proxies when launched (requires VS be launched or it will hang)")]
+		public bool AttachToProxies = false;
+
 		[CommandLine, Help("AttachToClients", "Force Visual Studio to attach to the clients when launched (requires VS be launched or it will hang)")]
 		public bool AttachToClients = false;
 
@@ -785,7 +788,8 @@ Example running 4 clients locally connecting to single proxy server w/ a non-loc
 		/// <returns>true if the commandlet parameters have specified we attach a debugger to this process</returns>
 		protected bool IsAttachingDebugger(ProcessLaunchType LaunchType)
 		{
-			return (AttachToServers && (LaunchType == ProcessLaunchType.Server || LaunchType == ProcessLaunchType.ProxyServer)) ||
+			return (AttachToServers && (LaunchType == ProcessLaunchType.Server)) ||
+				 (AttachToProxies && (LaunchType == ProcessLaunchType.ProxyServer)) ||
 				(AttachToClients && (LaunchType == ProcessLaunchType.DirectClient || LaunchType == ProcessLaunchType.ProxyClient));
 		}
 
