Fix Mac Server build unable to pak properly

#rnx

diff --git a/Engine/Source/Programs/AutomationTool/Scripts/CopyBuildToStagingDirectory.Automation.cs b/Engine/Source/Programs/AutomationTool/Scripts/CopyBuildToStagingDirectory.Automation.cs
--- a/Engine/Source/Programs/AutomationTool/Scripts/CopyBuildToStagingDirectory.Automation.cs
+++ b/Engine/Source/Programs/AutomationTool/Scripts/CopyBuildToStagingDirectory.Automation.cs
@@ -4223,7 +4223,7 @@ namespace AutomationScripts
 
 						if (InternalUtils.SafeCopyFile(SourceOutputLocation.FullName, OutputLocation.FullName))
 						{
-							Logger.LogInformation("Copying source pak from {SourceOutputLocation} to {OutputLocation} instead of creating new pak", SourceOutputLocation, OutputLocation);
+							Logger.LogInformation("Copying source pak from {SourceOutputLocation} to {OutputLocation} instead of creating new pak", SourceOutputLocation.FullName, OutputLocation.FullName);
 							bCopiedExistingPak = true;
 
 							FileReference InSigFile = SourceOutputLocation.ChangeExtension(".sig");
@@ -4231,11 +4231,11 @@ namespace AutomationScripts
 							{
 								FileReference OutSigFile = OutputLocation.ChangeExtension(".sig");
 
-								Logger.LogInformation("Copying pak sig from {InSigFile} to {OutSigFile}", InSigFile, OutSigFile);
+								Logger.LogInformation("Copying pak sig from {InSigFile} to {OutSigFile}", InSigFile.FullName, OutSigFile.FullName);
 
 								if (!InternalUtils.SafeCopyFile(InSigFile.FullName, OutSigFile.FullName))
 								{
-									Logger.LogInformation("Failed to copy pak sig {InSigFile} to {InSigFile2}, creating new pak", InSigFile, InSigFile);
+									Logger.LogInformation("Failed to copy pak sig {InSigFile} to {InSigFile2}, creating new pak", InSigFile.FullName, InSigFile.FullName);
 									bCopiedExistingPak = false;
 								}
 							}
@@ -4252,7 +4252,7 @@ namespace AutomationScripts
 					}
 					if (!bCopiedExistingPak)
 					{
-						Logger.LogInformation("Failed to copy source pak from {SourceOutputLocation} to {OutputLocation}, creating new pak", SourceOutputLocation, OutputLocation);
+						Logger.LogInformation("Failed to copy source pak from {SourceOutputLocation} to {OutputLocation}, creating new pak", SourceOutputLocation.FullName, OutputLocation.FullName);
 					}
 
 					if (bCopiedExistingPak && !bCopiedExistingGlobalUtoc)
