[Gauntlet] Cleaning up Android artifacts

#rnx

diff --git a/Engine/Source/Programs/AutomationTool/Gauntlet/Platform/Android/Gauntlet.TargetDeviceAndroid.cs b/Engine/Source/Programs/AutomationTool/Gauntlet/Platform/Android/Gauntlet.TargetDeviceAndroid.cs
--- a/Engine/Source/Programs/AutomationTool/Gauntlet/Platform/Android/Gauntlet.TargetDeviceAndroid.cs
+++ b/Engine/Source/Programs/AutomationTool/Gauntlet/Platform/Android/Gauntlet.TargetDeviceAndroid.cs
@@ -368,14 +368,27 @@ namespace Gauntlet
 
 			bUseAFSWithoutToken = Build.UseAFSWithoutToken;
 
-			string ExternalStoragePath = StoragePath + "/UnrealGame/" + AppConfig.ProjectName;
-			string ExternalFilesPath = StoragePath + "/Android/data/" + Build.AndroidPackageName + "/files/UnrealGame/" + AppConfig.ProjectName;
+			Log.Info("Cleaning up artifacts for {PackageName}...", Build.AndroidPackageName);
+
+			string PackageOutput = RunAdbDeviceCommandAndGetOutput("shell pm list packages -3");
 
-			string ExternalAppStorage = string.Format("{0}/{1}/Saved", ExternalStoragePath, AppConfig.ProjectName);
-			string ExternalAppFiles = string.Format("{0}/{1}/Saved", ExternalFilesPath, AppConfig.ProjectName);
+			bool bIsPackageInstalled = PackageOutput.Contains($"package:{Build.AndroidPackageName}");
 
-			RunAdbDeviceCommand(string.Format("shell rm -r {0}/*", ExternalAppStorage));
-			RunAdbDeviceCommand(string.Format("shell rm -r {0}/*", ExternalAppFiles));
+			if (bIsPackageInstalled && UsingAndroidFileServer(AppConfig.ProjectFile, AppConfig.Configuration, out _, out string AFSToken, out _, out _, out _, bUseAFSWithoutToken))
+			{
+				RunAFSDeviceCommand(string.Format("-p \"{0}\" -k \"{1}\" rmdir -r \"^saved/\"", Build.AndroidPackageName, AFSToken));
+			}
+			else
+			{
+				string ExternalStoragePath = StoragePath + "/UnrealGame/" + AppConfig.ProjectName;
+				string ExternalFilesPath = StoragePath + "/Android/data/" + Build.AndroidPackageName + "/files/UnrealGame/" + AppConfig.ProjectName;
+
+				string ExternalAppStorage = string.Format("{0}/{1}/Saved", ExternalStoragePath, AppConfig.ProjectName);
+				string ExternalAppFiles = string.Format("{0}/{1}/Saved", ExternalFilesPath, AppConfig.ProjectName);
+
+				RunAdbDeviceCommand(string.Format("shell rm -r {0}/*", ExternalAppStorage));
+				RunAdbDeviceCommand(string.Format("shell rm -r {0}/*", ExternalAppFiles));
+			}
 		}
 
 		public void InstallBuild(UnrealAppConfig AppConfig)
