Standardize capturing Unreal Insights

diff --git a/Engine/Source/Programs/AutomationTool/Gauntlet/Unreal/Game/EpicGame.TestConfig.cs b/Engine/Source/Programs/AutomationTool/Gauntlet/Unreal/Game/EpicGame.TestConfig.cs
--- a/Engine/Source/Programs/AutomationTool/Gauntlet/Unreal/Game/EpicGame.TestConfig.cs
+++ b/Engine/Source/Programs/AutomationTool/Gauntlet/Unreal/Game/EpicGame.TestConfig.cs
@@ -48,20 +48,25 @@ namespace EpicGame
 		protected string BuildIDOverride = "";
 
 		/// <summary>
-		/// Unique server port to avoid matchmaking collisions
+		/// Initial value for a unique server port to avoid matchmaking collisions
 		/// </summary>
 		[AutoParam]
 		int ServerPortStart = 7777;
 
+		/// <summary>
+		/// Unique server port for the current Config object
+		/// </summary>
 		protected int ServerPort { get; private set; }
 
-
 		/// <summary>
-		/// Unique server beacon port to avoid matchmaking collisions
+		/// Initial value for a unique server beacon port to avoid matchmaking collisions
 		/// </summary>
 		[AutoParam]
 		int BeaconPortStart = 15000;
 
+		/// <summary>
+		/// Unique beacon port for the current Config object
+		/// </summary>
 		protected int BeaconPort { get; private set; }
 		
 		/// <summary>
@@ -107,27 +112,45 @@ namespace EpicGame
 		[AutoParam]
 		public string McpSubRegion = string.Empty;
 
-		/// Set traces to collect for editor roles, saved to a generated file.
+		/// <summary>
+		/// Comma-separated list of Unreal Insights trace channels to collect for Editor roles, saved to a file named EditorTraces{ID#}.utrace
 		/// </summary>
 		[AutoParam]
 		public string EditorTraces { get; set; }
 
+		/// <summary>
+		/// Comma-separated list of Unreal Insights trace channels to collect for Client roles, saved to a file named ClientTraces{ID#}.utrace
+		/// </summary>
+		[AutoParam]
+		public string ClientTraces { get; set; }
+
+		/// <summary>
+		/// Comma-separated list of Unreal Insights trace channels to collect for Server roles, saved to a file named ServerTraces{ID#}.utrace
+		/// </summary>
+		[AutoParam]
+		public string ServerTraces { get; set; }
+
+		/// <summary>
+		/// Filepaths to trace logging artifacts, copied after test execution ends
+		/// </summary>
+		public readonly List<string> TraceLogPaths = [];
+
 		// Allows a more derived test config to take over mcp setup.
 		protected bool SkipMcpSetup = false;
 
 		// incrementing value to ensure we can assign unique values to ports etc
 		static private int NumberOfConfigsCreated = 0;
+		static private int NumberOfEditorsConfigured = 0;
+		static private int NumberOfClientsConfigured = 0;
+		static private int NumberOfServersConfigured = 0;
+
+		public static readonly string AutomationToolLogsDir = Path.Combine(Unreal.EngineDirectory.FullName, "Programs", "AutomationTool", "Saved", "Logs");
 
 		public EpicGameTestConfig()
 		{
 			NumberOfConfigsCreated++;
 		}
 
-		~EpicGameTestConfig()
-		{
-		}
-
-
 		public void ParametersWereApplied(string[] Params)
 		{
 			
@@ -316,17 +339,55 @@ namespace EpicGame
 				{
 					AppConfig.CommandLine += " -nocrashlytics";
 				}
+
+				if (!string.IsNullOrEmpty(ClientTraces))
+				{
+					AppConfig.CommandLine += $" -trace={ClientTraces} -tracefile={SetupInsightsTraceFile(ConfigRole.RoleType, NumberOfClientsConfigured)} -traceautostart=1";
+				}
+
+				++NumberOfClientsConfigured;
+			}
+
+			if (ConfigRole.RoleType.IsServer())
+			{
+				if (!string.IsNullOrEmpty(ServerTraces))
+				{
+					AppConfig.CommandLine += $" -trace={ServerTraces} -tracefile={SetupInsightsTraceFile(ConfigRole.RoleType, NumberOfServersConfigured)} -traceautostart=1 -nofakeforking";
+				}
+
+				++NumberOfServersConfigured;
 			}
 
 			if (ConfigRole.RoleType.IsEditor())
 			{
 				if (!string.IsNullOrEmpty(EditorTraces))
 				{
-					string GeneratedTraceFile = Path.Combine(Unreal.EngineDirectory.FullName, "Programs", "AutomationTool", "Saved", "Logs", "EditorTraces.utrace");
-					AppConfig.CommandLine += $" -trace={EditorTraces} -tracefile={GeneratedTraceFile}";
+					AppConfig.CommandLine += $" -trace={EditorTraces} -tracefile={SetupInsightsTraceFile(ConfigRole.RoleType, NumberOfEditorsConfigured)} -traceautostart=1";
 				}
+
+				++NumberOfEditorsConfigured;
 			}
-		}		
+		}
+
+		private string SetupInsightsTraceFile(UnrealTargetRole RoleType, int RoleCount)
+		{
+			string GeneratedTraceFile = Path.Combine(AutomationToolLogsDir, $"{RoleType}Trace_{RoleCount}.utrace");
+			TraceLogPaths.Add(GeneratedTraceFile);
+			if (File.Exists(GeneratedTraceFile))
+			{
+				// unexpected: the Logs directory is cleared each time UAT executes
+				Log.Warning($"Removing previous Unreal Insights capture at {GeneratedTraceFile}");
+				try
+				{
+					File.Delete(GeneratedTraceFile);
+				}
+				catch(Exception ex)
+				{
+					Log.Warning($"Unable to delete previous Unreal Insights capture, a new result from this run will not be captured. Failure reason: {ex}");
+				}
+			}
+			return GeneratedTraceFile;
+		}
 	}
 
 	/// <summary>
@@ -339,6 +400,44 @@ namespace EpicGame
 		{
 
 		}
+
+		public override IEnumerable<UnrealRoleArtifacts> SaveRoleArtifacts(string OutputPath)
+		{
+			string OutputFolderName = new DirectoryInfo(OutputPath).Name;
+			foreach (string SourceTracePath in CachedConfig.TraceLogPaths)
+			{
+				var Filename = Path.GetFileName(SourceTracePath);
+				
+				if (CommandUtils.IsBuildMachine && Gauntlet.Horde.IsHordeJob)
+				{
+					string HordeTracePath = Path.GetFullPath(Path.Combine(CommandUtils.CmdEnv.LogFolder, OutputFolderName, Filename));
+					Log.Verbose($"Copying trace file for Horde to {HordeTracePath}");
+					try
+					{
+						Directory.CreateDirectory(Path.GetDirectoryName(HordeTracePath));
+						File.Copy(SourceTracePath, HordeTracePath, true);
+					}
+					catch (Exception Ex)
+					{
+						Log.Warning($"Failed to copy {SourceTracePath} to Horde Artifacts. Exception: {Ex}");
+					}
+				}
+
+				string OutputTracePath = Path.Combine(OutputPath, Filename);
+				Log.Verbose($"Moving trace file to {OutputTracePath}");
+				try
+				{
+					Directory.CreateDirectory(Path.GetDirectoryName(OutputTracePath));
+					File.Move(SourceTracePath, OutputTracePath, true);
+				}
+				catch (Exception Ex)
+				{
+					Log.Warning($"Failed to move {SourceTracePath} to Role Artifacts. Cached file will be overwritten on the next invocation of Unreal Automation Tool. Exception: {Ex}");
+				}
+			}
+
+			return base.SaveRoleArtifacts(OutputPath);
+		}
 	}
 
 }
\ No newline at end of file
