Fix a race condition in RPCRegistry where we can request the available RPCs before they have registered on the client and 0 RPCs are returned available, but we prevent requesting for updated RPCs later.

#rnx
#Synced-CL 47572819

diff --git a/Engine/Source/Programs/AutomationTool/Gauntlet/Framework/RpcFramework/Gauntlet.RpcTargetRegistry.cs b/Engine/Source/Programs/AutomationTool/Gauntlet/Framework/RpcFramework/Gauntlet.RpcTargetRegistry.cs
--- a/Engine/Source/Programs/AutomationTool/Gauntlet/Framework/RpcFramework/Gauntlet.RpcTargetRegistry.cs
+++ b/Engine/Source/Programs/AutomationTool/Gauntlet/Framework/RpcFramework/Gauntlet.RpcTargetRegistry.cs
@@ -26,7 +26,7 @@ namespace Gauntlet
 		public string TargetName;
 		public List<RpcDefinition> AvailableRPCs;
 		/// <summary>
-		/// If this is true, we know the target is capable of sending messages to us and we no longer need to manually poll those processes for messages.
+		/// If this is true, we know the target is capable of sending messages to us and has available RPCs and we no longer need to manually poll those processes for messages.
 		/// </summary>
 		public bool HasSentMessage = false;
 		public bool IsDeferStart = false;
@@ -101,7 +101,7 @@ namespace Gauntlet
 				// Not sure what to put here as RpcRegistry could quite possibly not be initialized by design.
 				return false;
 			}
-			return true;
+			return AvailableRPCs.Count > 0;
 		}
 
 		// Logging Helper extensions that prepends the TargetName to the log message to help identify which process it came from
@@ -296,8 +296,7 @@ namespace Gauntlet
 			RpcTarget InTarget = GetTarget(InMessage.SenderId);
 			if (InTarget != null)
 			{
-				InTarget.HasSentMessage = true;
-				InTarget.UpdateAvailableRpcs();
+				InTarget.HasSentMessage = InTarget.UpdateAvailableRpcs();
 			}
 		}
 
@@ -412,11 +411,11 @@ namespace Gauntlet
 
 		private bool InitialTargetRoleIsSet(UnrealTargetRole TargetRoleType) => TargetRoleType switch
 		{
-			UnrealTargetRole.Client		  => ClientTarget is not null,
-			UnrealTargetRole.EditorGame	  => ClientTarget is not null,
-			UnrealTargetRole.Server		  => ServerTarget is not null,
+			UnrealTargetRole.Client => ClientTarget is not null,
+			UnrealTargetRole.EditorGame => ClientTarget is not null,
+			UnrealTargetRole.Server => ServerTarget is not null,
 			UnrealTargetRole.EditorServer => ServerTarget is not null,
-			UnrealTargetRole.Editor		  => EditorTarget is not null,
+			UnrealTargetRole.Editor => EditorTarget is not null,
 			UnrealTargetRole.CookedEditor => EditorTarget is not null,
 			_ => false,
 		};
@@ -625,4 +624,4 @@ namespace Gauntlet
 			ReleaseUnmanagedResources();
 		}
 	}
-}
\ No newline at end of file
+}
