Small fix for iOS cancel button checking retry logic

#Synced-CL 47873187

diff --git a/Engine/Source/Programs/AutomationTool/Gauntlet/Platform/IOS/Gauntlet.IOSAppium.cs b/Engine/Source/Programs/AutomationTool/Gauntlet/Platform/IOS/Gauntlet.IOSAppium.cs
--- a/Engine/Source/Programs/AutomationTool/Gauntlet/Platform/IOS/Gauntlet.IOSAppium.cs
+++ b/Engine/Source/Programs/AutomationTool/Gauntlet/Platform/IOS/Gauntlet.IOSAppium.cs
@@ -337,7 +337,7 @@ namespace Gauntlet
 				IsAppInForeground = CheckAppInForeground(packageName);
 				if (!IsAppInForeground)
 				{
-					Log.Info("{0} in background, trying to foreground Attempt #{1}", packageName, retries);
+					Log.Info("{0} in background, trying to foreground Retries remaining:{1}", packageName, retries);
 					ForegroundApp(packageName);
 					Thread.Sleep(500);
 					retries--;
@@ -381,9 +381,9 @@ namespace Gauntlet
 
 		public void CancelAppleSignOnAlert(int retries = 5)
 		{
-			try
+			while (retries > 0)
 			{
-				while (retries > 0)
+				try
 				{
 					IAlert Alert = Driver.SwitchTo().Alert();
 
@@ -395,17 +395,17 @@ namespace Gauntlet
 						Log.Verbose("Click Cancel Button");
 						CancelButton.Click();
 					}
-
-					retries--;
 				}
-			}
-			catch (NoAlertPresentException)
-			{
-				Log.Verbose("No alert present");
-			}
-			catch (Exception ex)
-			{
-				Log.Verbose("An error occurred when dismissing an alert message: {0}", ex);
+				catch (NoAlertPresentException)
+				{
+					Log.Verbose("No alert present");
+				}
+				catch (Exception ex)
+				{
+					Log.Verbose("An error occurred when dismissing an alert message: {0}", ex);
+				}
+
+				retries--;
 			}
 		}
 	}
