Install Bundle Subsets
- Subsets are bundle variants with different payloads.
- We generate seperate download config for each variant.
- On non-bps platforms, we filter what files to mount at runtime.

#Synced-CL 48271216

diff --git a/Engine/Plugins/Experimental/DefaultInstallBundleManager/Source/Private/InstallBundleSourceBulk.cpp b/Engine/Plugins/Experimental/DefaultInstallBundleManager/Source/Private/InstallBundleSourceBulk.cpp
--- a/Engine/Plugins/Experimental/DefaultInstallBundleManager/Source/Private/InstallBundleSourceBulk.cpp
+++ b/Engine/Plugins/Experimental/DefaultInstallBundleManager/Source/Private/InstallBundleSourceBulk.cpp
@@ -8,6 +8,7 @@
 #include "HAL/PlatformFile.h"
 #include "IO/IoStoreOnDemand.h"
 #include "IPlatformFilePak.h"
+#include "InstallBundleManagerInterface.h"
 #include "InstallBundleManagerUtil.h"
 #include "Misc/CommandLine.h"
 #include "Misc/FileHelper.h"
@@ -323,7 +324,9 @@ void FInstallBundleSourceBulk::AsyncInit_MakeBundlesForBulkBuild()
 	{
 		//First load any existing BulkBuildBundleIni entries and use those to sort FoundFiles into bundles
 		const bool bHasAnyFilesToParse = FoundFiles.IsValid() && (FoundFiles->Num() > 0);
-		bool bDidLoadAllFilesFromMetadata = bHasAnyFilesToParse ? TryLoadBulkBuildBundleMetadata(*FoundFiles, BulkBuildBundles) : false;
+		// TODO: this cache is incompatible with bundle subsets. Its also unclear if its still needed, so I'm disabling it for now
+		// bool bDidLoadAllFilesFromMetadata = bHasAnyFilesToParse ? TryLoadBulkBuildBundleMetadata(*FoundFiles, BulkBuildBundles) : false;
+		bool bDidLoadAllFilesFromMetadata = false;
 
 		// Prune out any remaining paks that were mounted at startup
 		for (int i = 0; i < FoundFiles->Num();)
@@ -386,9 +389,11 @@ void FInstallBundleSourceBulk::AsyncInit_MakeBundlesForBulkBuild()
 			for (const FString& Section : SectionNames)
 			{
 				const FString BundleName = Section.RightChop(InstallBundleUtil::GetInstallBundleSectionPrefix().Len());
-				TArray<FString>& BundleFileList = BulkBuildBundles.FindOrAdd(FName(*BundleName));
+				BulkBuildBundles.FindOrAdd(FName(*BundleName));
 			}
 
+			TSharedPtr<IInstallBundleManager> BundleManger = IInstallBundleManager::GetPlatformInstallBundleManager();
+
 			//If we had files remaining to manually parse, now sort the remaining files into bundles
 			if (!bDidLoadAllFilesFromMetadata)
 			{
@@ -398,8 +403,23 @@ void FInstallBundleSourceBulk::AsyncInit_MakeBundlesForBulkBuild()
 					TArray<FString>& BundleFileList = BulkBuildBundles.FindOrAdd(FName(*BundleName));
 
 					TArray<FString> StrSearchRegexPatterns;
-					if (!InstallBundleConfig->GetArray(*Section, TEXT("FileRegex"), StrSearchRegexPatterns))
+					InstallBundleConfig->GetArray(*Section, TEXT("FileRegex"), StrSearchRegexPatterns);
+
+					if (StrSearchRegexPatterns.IsEmpty())
+					{
 						continue;
+					}
+
+					TArray<FString> BundleSubsets;
+					InstallBundleConfig->GetArray(*Section, TEXT("Subsets"), BundleSubsets);
+					const int32 iSubset = BundleSubsets.IsEmpty() ? INDEX_NONE : BundleManger->SelectBundleSubset(BundleName, BundleSubsets);
+
+					TArray<FString> StrSubsetSearchRegexPatterns;
+					if (BundleSubsets.IsValidIndex(iSubset))
+					{
+						const FString SubsetSection = InstallBundleUtil::GetInstallBundleSubsetSectionPrefix() + BundleSubsets[iSubset];
+						InstallBundleConfig->GetArray(*SubsetSection, TEXT("FileRegex"), StrSubsetSearchRegexPatterns);
+					}
 
 					TArray<FRegexPattern> SearchRegexPatterns;
 					SearchRegexPatterns.Reserve(StrSearchRegexPatterns.Num());
@@ -407,10 +427,18 @@ void FInstallBundleSourceBulk::AsyncInit_MakeBundlesForBulkBuild()
 					{
 						SearchRegexPatterns.Emplace(Str, ERegexPatternFlags::CaseInsensitive);
 					}
+
+					TArray<FRegexPattern> SubsetSearchRegexPatterns;
+					SubsetSearchRegexPatterns.Reserve(StrSubsetSearchRegexPatterns.Num());
+					for (const FString& Str : StrSubsetSearchRegexPatterns)
+					{
+						SubsetSearchRegexPatterns.Emplace(Str, ERegexPatternFlags::CaseInsensitive);
+					}
 				
 					for (int i = 0; i < FoundFiles->Num();)
 					{
-						const FString& File = (*FoundFiles)[i];
+						FString& File = (*FoundFiles)[i];
+
 						bool bMatches = false;
 						for (const FRegexPattern& Pattern : SearchRegexPatterns)
 						{
@@ -421,16 +449,34 @@ void FInstallBundleSourceBulk::AsyncInit_MakeBundlesForBulkBuild()
 							}
 						}
 
-						if (bMatches)
+						if (!bMatches)
+						{
+							++i;
+							continue;
+						}
+
+						LOG_SOURCE_BULK(Verbose, TEXT("Adding %s to Bundle %s"), *File, *BundleName);
+
+						FString MatchedFile = MoveTemp(File);
+						FoundFiles->RemoveAtSwap(i);
+
+						bMatches = SubsetSearchRegexPatterns.IsEmpty();
+						for (const FRegexPattern& Pattern : SubsetSearchRegexPatterns)
 						{
-							LOG_SOURCE_BULK(Verbose, TEXT("Adding %s to Bundle %s"), *File, *BundleName);
+							if (FRegexMatcher(Pattern, MatchedFile).FindNext())
+							{
+								bMatches = true;
+								break;
+							}
+						}
 						
-							BundleFileList.AddUnique(File);
-							FoundFiles->RemoveAtSwap(i);
+						if (bMatches)
+						{
+							BundleFileList.AddUnique(MoveTemp(MatchedFile));
 						}
 						else
 						{
-							++i;
+							LOG_SOURCE_BULK(Verbose, TEXT("Excluded %s from Bundle %s with subset %s"), *MatchedFile, *BundleName, *BundleSubsets[iSubset]);
 						}
 					}
 				}
@@ -438,7 +484,7 @@ void FInstallBundleSourceBulk::AsyncInit_MakeBundlesForBulkBuild()
 		}
 
 		//See if we should serialize out the manually parsed results
-		if (!bDidLoadAllFilesFromMetadata)
+		/*if (!bDidLoadAllFilesFromMetadata)
 		{
 			bool bShouldSerializeMissingBulkBuildDataIni = false;
 			if (!GConfig->GetBool(TEXT("InstallBundleSource.Bulk.MiscSettings"), TEXT("bShouldSerializeMissingBulkBuildDataIni"), bShouldSerializeMissingBulkBuildDataIni, GInstallBundleIni))
@@ -450,7 +496,7 @@ void FInstallBundleSourceBulk::AsyncInit_MakeBundlesForBulkBuild()
 			{
 				SerializeBulkBuildBundleMetadata(BulkBuildBundles);
 			}
-		}
+		}*/
 
 		LOG_SOURCE_BULK(Display, TEXT("Finished Making Bundles for Bulk Build"));
 		InitStepResult = EAsyncInitStepResult::Done;
@@ -546,8 +592,10 @@ void FInstallBundleSourceBulk::AsyncInit_QueryBundleInfo(FInstallBundleSourceQue
 		{
 			const FString& Section = Pair.Key;
 			FInstallBundleSourcePersistentBundleInfo BundleInfo;
-			if(!InstallBundleManagerUtil::LoadBundleSourceBundleInfoFromConfig(GetSourceType(), *InstallBundleConfig, Section, BundleInfo))
+			if (!InstallBundleManagerUtil::LoadBundleSourceBundleInfoFromConfig(GetSourceType(), *InstallBundleConfig, Section, BundleInfo))
+			{
 				continue;
+			}
 
 			BundleInfo.BundleContentState = GetBundleInstallState(BundleInfo.BundleName);
 
diff --git a/Engine/Plugins/Experimental/DefaultInstallBundleManager/Source/Private/InstallBundleSourcePlatformBase.cpp b/Engine/Plugins/Experimental/DefaultInstallBundleManager/Source/Private/InstallBundleSourcePlatformBase.cpp
--- a/Engine/Plugins/Experimental/DefaultInstallBundleManager/Source/Private/InstallBundleSourcePlatformBase.cpp
+++ b/Engine/Plugins/Experimental/DefaultInstallBundleManager/Source/Private/InstallBundleSourcePlatformBase.cpp
@@ -53,6 +53,7 @@ FInstallBundleSourceInitInfo FInstallBundleSourcePlatformBase::Init(
 	}
 
 	BundleRegexList = InstallBundleUtil::LoadBundleRegexFromConfig(*InstallBundleConfig, InstallBundleUtil::IsPlatformInstallBundlePredicate);
+	BundleSubsetRegexList = InstallBundleUtil::LoadBundleSubsetRegexFromConfig(*InstallBundleConfig);
 
 	return InitInfo;
 }
diff --git a/Engine/Plugins/Experimental/DefaultInstallBundleManager/Source/Private/InstallBundleSourcePlatformChunkInstall.cpp b/Engine/Plugins/Experimental/DefaultInstallBundleManager/Source/Private/InstallBundleSourcePlatformChunkInstall.cpp
--- a/Engine/Plugins/Experimental/DefaultInstallBundleManager/Source/Private/InstallBundleSourcePlatformChunkInstall.cpp
+++ b/Engine/Plugins/Experimental/DefaultInstallBundleManager/Source/Private/InstallBundleSourcePlatformChunkInstall.cpp
@@ -5,6 +5,7 @@
 #include "DefaultInstallBundleManagerPrivate.h"
 
 #include "InstallBundleManagerUtil.h"
+#include "InstallBundleManagerInterface.h"
 
 #include "Misc/CommandLine.h"
 #include "Misc/ConfigCacheIni.h"
@@ -151,23 +152,27 @@ void FInstallBundleSourcePlatformChunkInstall::AsyncInit(FInstallBundleSourceIni
 		LOG_SOURCE_CHUNKINSTALL(Warning, TEXT("Retrying initialization"));
 		InitState = EInstallBundleManagerInitState::NotInitialized;
 	}
-	
+
 	struct AsyncInitContext
 	{
 		EInstallBundleManagerInitResult InitResult = EInstallBundleManagerInitResult::OK;
 
 		TSet<FName> NamedChunks;
 		TMap<FName, FBundleInfo> BundleInfoMap;
+		TMap<FName, TArray<FString>> BundleToSubsetMap;
 		TArray<TPair<FString, TArray<FRegexPattern>>> BundleRegexList;
+		TArray<TPair<FString, TArray<FRegexPattern>>> BundleSubsetRegexList;
 		IPlatformChunkInstall* PlatformChunkInstall;
+		TSharedPtr<IInstallBundleManager> BundleManager;
 	};
 
 
 	// create and initialize the async init context
 	TSharedPtr<AsyncInitContext, ESPMode::ThreadSafe> Context = MakeShared<AsyncInitContext, ESPMode::ThreadSafe>();
 	Context->BundleRegexList = MoveTemp(BundleRegexList);
+	Context->BundleSubsetRegexList = MoveTemp(BundleSubsetRegexList);
 	Context->PlatformChunkInstall = PlatformChunkInstall;
-
+	Context->BundleManager = IInstallBundleManager::GetPlatformInstallBundleManager();
 	
 
 	auto AsyncInit = [Context]() mutable
@@ -232,6 +237,13 @@ void FInstallBundleSourcePlatformChunkInstall::AsyncInit(FInstallBundleSourceIni
 			InstallBundleConfig->GetString(*Section, TEXT("Priority"), PriorityString);
 			LexTryParseString(BundleInfo.Priority, *PriorityString);
 
+			TArray<FString> BundleSubsets;
+			InstallBundleConfig->GetArray(*Section, TEXT("Subsets"), BundleSubsets);
+			if (!BundleSubsets.IsEmpty())
+			{
+				Context->BundleToSubsetMap.Add(BundleName, MoveTemp(BundleSubsets));
+			}
+
 			if (KnownNamedChunks.Contains(BundleInfo.NamedChunk))
 			{
 				continue;
@@ -293,8 +305,17 @@ void FInstallBundleSourcePlatformChunkInstall::AsyncInit(FInstallBundleSourceIni
 				FString BundleName;
 				if (InstallBundleUtil::MatchBundleRegex(Context->BundleRegexList, Path, BundleName))
 				{
-					FBundleInfo& BundleInfo = Context->BundleInfoMap.FindChecked(*BundleName);
-					BundleInfo.FilePaths.AddUnique(MoveTemp(Path));
+					FName BundleFName(*BundleName);
+					FBundleInfo& BundleInfo = Context->BundleInfoMap.FindChecked(BundleFName);
+					TArray<FString>* BundleSubsets = Context->BundleToSubsetMap.Find(BundleFName);
+
+					const int32 iSubset = (!BundleSubsets) ? INDEX_NONE : Context->BundleManager->SelectBundleSubset(BundleName, *BundleSubsets);
+
+					if (!BundleSubsets || !BundleSubsets->IsValidIndex(iSubset) ||
+						InstallBundleUtil::MatchBundleSubsetRegex((*BundleSubsets)[iSubset], Context->BundleSubsetRegexList, Path))
+					{
+						BundleInfo.FilePaths.AddUnique(MoveTemp(Path));
+					}
 				}
 				else
 				{
diff --git a/Engine/Plugins/Experimental/DefaultInstallBundleManager/Source/Public/InstallBundleSourcePlatformBase.h b/Engine/Plugins/Experimental/DefaultInstallBundleManager/Source/Public/InstallBundleSourcePlatformBase.h
--- a/Engine/Plugins/Experimental/DefaultInstallBundleManager/Source/Public/InstallBundleSourcePlatformBase.h
+++ b/Engine/Plugins/Experimental/DefaultInstallBundleManager/Source/Public/InstallBundleSourcePlatformBase.h
@@ -62,5 +62,6 @@ protected:
 	EInstallBundleManagerInitState InitState = EInstallBundleManagerInitState::NotInitialized;
 
 	TArray<TPair<FString, TArray<FRegexPattern>>> BundleRegexList; // BundleName -> FileRegex
+	TArray<TPair<FString, TArray<FRegexPattern>>> BundleSubsetRegexList;
 };
 #endif //WITH_PLATFORM_INSTALL_BUNDLE_SOURCE
diff --git a/Engine/Source/Programs/AutomationTool/AutomationUtils/BundleUtils.cs b/Engine/Source/Programs/AutomationTool/AutomationUtils/BundleUtils.cs
--- a/Engine/Source/Programs/AutomationTool/AutomationUtils/BundleUtils.cs
+++ b/Engine/Source/Programs/AutomationTool/AutomationUtils/BundleUtils.cs
@@ -1,12 +1,14 @@
 ﻿// Copyright Epic Games, Inc. All Rights Reserved.
+using AutomationTool;
+using EpicGames.Core;
+using EpicGames.Horde.Projects;
 using System;
 using System.Collections.Generic;
 using System.Linq;
+using System.Reflection.Metadata;
 using System.Text;
-using System.Threading.Tasks;
 using System.Text.RegularExpressions;
-using EpicGames.Core;
-using AutomationTool;
+using System.Threading.Tasks;
 using UnrealBuildTool;
 
 namespace AutomationUtils.Automation
@@ -20,7 +22,7 @@ namespace AutomationUtils.Automation
 			public List<string> Dependencies { get; set; }
 			public List<string> FileRegex { get; set; }
 			public List<string> Files { get; set; }
-			public bool bFoundParent { get; set; }
+			public List<string> Subsets { get; set; }
 			public bool bContainsShaderLibrary { get; set; }
 			public int Order { get; set; }
 			public bool	UseChunkDBs { get; set; }
@@ -29,26 +31,44 @@ namespace AutomationUtils.Automation
 			public string CloudAppNameOverride { get; set; }
 		}
 
+		public class BundleSubsetSettings
+		{
+			public string Name { get; set; }
+			public List<string> Tags { get; set; }
+			public List<string> FileRegex { get; set; }
+		}
+
+		public static ConfigHierarchy LoadConfigHierarchy(DirectoryReference ProjectDir, UnrealTargetPlatform Platform, string CustomConfig)
+		{
+			ConfigHierarchy BundleConfig = ConfigCache.ReadHierarchy(ConfigHierarchyType.InstallBundle, ProjectDir, Platform, CustomConfig);
+			return BundleConfig;
+		}
+
+		public static void LoadBundleConfig(ConfigHierarchy BundleConfig, out IReadOnlyDictionary<string, BundleSettings> Bundles)
+		{
+			LoadBundleConfig<BundleSettings>(BundleConfig, out Bundles, delegate (BundleSettings Settings, ConfigHierarchy BundleConfig, string Section) { });
+		}
+
 		public static void LoadBundleConfig(DirectoryReference ProjectDir, UnrealTargetPlatform Platform, string CustomConfig, out IReadOnlyDictionary<string, BundleSettings> Bundles)
 		{
 			LoadBundleConfig<BundleSettings>(ProjectDir, Platform, CustomConfig, out Bundles, delegate (BundleSettings Settings, ConfigHierarchy BundleConfig, string Section) { });
 		}
 
-		public static void LoadBundleConfig<TPlatformBundleSettings>(DirectoryReference ProjectDir, UnrealTargetPlatform Platform, string CustomConfig, 
+		public static void LoadBundleConfig<TPlatformBundleSettings>(ConfigHierarchy BundleConfig,
 			out IReadOnlyDictionary<string, TPlatformBundleSettings> Bundles, 
 			Action<TPlatformBundleSettings, ConfigHierarchy, string> GetPlatformSettings) 
 			where TPlatformBundleSettings : BundleSettings, new()
 		{
 			var Results = new List<TPlatformBundleSettings>();
 
-			ConfigHierarchy BundleConfig = ConfigCache.ReadHierarchy(ConfigHierarchyType.InstallBundle, ProjectDir, Platform, CustomConfig);
-
 			const string BundleDefinitionPrefix = "InstallBundleDefinition ";
 
 			foreach (string SectionName in BundleConfig.SectionNames)
 			{
 				if (!SectionName.StartsWith(BundleDefinitionPrefix))
+				{
 					continue;
+				}
 
 				TPlatformBundleSettings Bundle = new TPlatformBundleSettings();
 				Bundle.Name = SectionName.Substring(BundleDefinitionPrefix.Length);
@@ -107,6 +127,17 @@ namespace AutomationUtils.Automation
 						Bundle.Files = new List<string>();
 					}
 				}
+				{
+					List<string> Subsets;
+					if (BundleConfig.GetArray(SectionName, "Subsets", out Subsets))
+					{
+						Bundle.Subsets = Subsets;
+					}
+					else
+					{
+						Bundle.Subsets = new List<string>();
+					}
+				}
 				{
 					bool bContainsShaderLibrary;
 					if (BundleConfig.GetBool(SectionName, "ContainsShaderLibrary", out bContainsShaderLibrary))
@@ -161,6 +192,67 @@ namespace AutomationUtils.Automation
 			Bundles = Results.OrderBy(b => b.Order).ToDictionary(b => b.Name, b => b);
 		}
 
+		public static void LoadBundleConfig<TPlatformBundleSettings>(DirectoryReference ProjectDir, UnrealTargetPlatform Platform, string CustomConfig,
+			out IReadOnlyDictionary<string, TPlatformBundleSettings> Bundles,
+			Action<TPlatformBundleSettings, ConfigHierarchy, string> GetPlatformSettings)
+			where TPlatformBundleSettings : BundleSettings, new()
+		{
+			ConfigHierarchy BundleConfig = LoadConfigHierarchy(ProjectDir, Platform, CustomConfig);
+			LoadBundleConfig(BundleConfig, out Bundles, GetPlatformSettings);
+		}
+
+		public static void LoadBundleSubsetConfig(ConfigHierarchy BundleConfig, out IReadOnlyDictionary<string, BundleSubsetSettings> BundlesSubsets)
+		{
+			var Results = new Dictionary<string, BundleSubsetSettings>();
+
+			const string BundleDefinitionSubsetPrefix = "InstallBundleSubsetDefinition ";
+
+			foreach (string SectionName in BundleConfig.SectionNames)
+			{
+				if (!SectionName.StartsWith(BundleDefinitionSubsetPrefix))
+				{
+					continue;
+				}
+
+				BundleSubsetSettings Settings = new BundleSubsetSettings();
+				Settings.Name = SectionName.Substring(BundleDefinitionSubsetPrefix.Length);
+
+				{
+					List<string> Tags;
+					if (BundleConfig.GetArray(SectionName, "Tags", out Tags))
+					{
+						Settings.Tags = Tags;
+					}
+					else
+					{
+						Settings.Tags = new List<string>();
+					}
+				}
+				{
+					List<string> FileRegex;
+					if (BundleConfig.GetArray(SectionName, "FileRegex", out FileRegex))
+					{
+						Settings.FileRegex = FileRegex;
+					}
+					else
+					{
+						Settings.FileRegex = new List<string>();
+					}
+				}
+
+				Results.Add(Settings.Name, Settings);
+			}
+
+			BundlesSubsets = Results;
+		}
+
+		public static void LoadBundleSubsetConfig(DirectoryReference ProjectDir, UnrealTargetPlatform Platform, string CustomConfig,
+			out IReadOnlyDictionary<string, BundleSubsetSettings> BundlesSubsets)
+		{
+			ConfigHierarchy BundleConfig = LoadConfigHierarchy(ProjectDir, Platform, CustomConfig);
+			LoadBundleSubsetConfig(BundleConfig, out BundlesSubsets);
+		}
+
 		public static TPlatformBundleSettings MatchBundleSettings<TPlatformBundleSettings>(
 			string FileName, IReadOnlyDictionary<string, TPlatformBundleSettings> InstallBundles) where TPlatformBundleSettings : BundleSettings
 		{
@@ -186,15 +278,28 @@ namespace AutomationUtils.Automation
 			return null;
 		}
 
-		public static IEnumerable<string> GetCloudAppsForBundles<TPlatformBundleSettings>(IReadOnlyDictionary<string, TPlatformBundleSettings> InstallBundles) where TPlatformBundleSettings : BundleSettings
+		public static IList<BundleSubsetSettings> MatchBundleSubsetSettings(BundleSettings Bundle, string FileName,
+			IReadOnlyDictionary<string, BundleSubsetSettings> InstallBundleSubsets)
 		{
-			// Gather all bundle clouda app name overrides
-			var CloudAppNames = new HashSet<string>();
-			foreach (var Bundle in InstallBundles.Values)
+			List<BundleSubsetSettings> Results = new List<BundleSubsetSettings>();
+
+			// Find matching subsets in order specified by bundle config
+			foreach (string Subset in Bundle.Subsets)
 			{
-				CloudAppNames.Add(Bundle.CloudAppNameOverride);
+				if (InstallBundleSubsets.TryGetValue(Subset, out BundleSubsetSettings SubsetSettings))
+				{
+					foreach (string RegexString in SubsetSettings.FileRegex)
+					{
+						if (Regex.Match(FileName, RegexString, RegexOptions.IgnoreCase).Success)
+						{
+							Results.Add(SubsetSettings);
+							break;
+						}
+					}
+				}
 			}
-			return CloudAppNames;
+
+			return Results;
 		}
 
 		public static IEnumerable<TPlatformBundleSettings> GetBundleDependencies<TPlatformBundleSettings>(
@@ -205,15 +310,15 @@ namespace AutomationUtils.Automation
 			var CurrentDependencies = new List<TPlatformBundleSettings>();
 			CurrentDependencies.Add(Bundle);
 
-			while(CurrentDependencies.Count > 0)
+			while (CurrentDependencies.Count > 0)
 			{
 				var NextDependencies = new List<TPlatformBundleSettings>();
-				foreach(var Dep in CurrentDependencies)
+				foreach (var Dep in CurrentDependencies)
 				{
 					if (FoundDependencies.ContainsKey(Dep.Name))
 						continue;
 					
-					FoundDependencies.Add(Dep.Name, Dep);					
+					FoundDependencies.Add(Dep.Name, Dep);
 
 					foreach (var NextDepName in Dep.Dependencies)
 					{
diff --git a/Engine/Source/Runtime/InstallBundleManager/Private/InstallBundleUtils.cpp b/Engine/Source/Runtime/InstallBundleManager/Private/InstallBundleUtils.cpp
--- a/Engine/Source/Runtime/InstallBundleManager/Private/InstallBundleUtils.cpp
+++ b/Engine/Source/Runtime/InstallBundleManager/Private/InstallBundleUtils.cpp
@@ -100,6 +100,12 @@ namespace InstallBundleUtil
 		return Prefix;
 	}
 
+	const FString& GetInstallBundleSubsetSectionPrefix()
+	{
+		static FString Prefix("InstallBundleSubsetDefinition "); // trailing space intentional
+		return Prefix;
+	}
+
 	bool GetConfiguredBundleSources(TArray<FString>& OutSources, TMap<FString, FString>& OutFallbackSources)
 	{
 #ifdef INSTALL_BUNDLE_SOURCES_FALLBACK_CONFIG_SECTION
@@ -248,6 +254,35 @@ namespace InstallBundleUtil
 		return BundleRegexList;
 	}
 
+	TArray<TPair<FString, TArray<FRegexPattern>>> LoadBundleSubsetRegexFromConfig(const FConfigFile& InstallBundleConfig)
+	{
+		TArray<TPair<FString, TArray<FRegexPattern>>> BundleSubsetRegexList; // BundleSubsetName -> FileRegex
+
+		for (const TPair<FString, FConfigSection>& Pair : InstallBundleConfig)
+		{
+			const FString& Section = Pair.Key;
+			if (!Section.StartsWith(InstallBundleUtil::GetInstallBundleSubsetSectionPrefix()))
+			{
+				continue;
+			}
+
+			TArray<FString> StrSearchRegexPatterns;
+			InstallBundleConfig.GetArray(*Section, TEXT("FileRegex"), StrSearchRegexPatterns);
+
+			TArray<FRegexPattern> SearchRegexPatterns;
+			SearchRegexPatterns.Reserve(StrSearchRegexPatterns.Num());
+			for (const FString& Str : StrSearchRegexPatterns)
+			{
+				SearchRegexPatterns.Emplace(Str, ERegexPatternFlags::CaseInsensitive);
+			}
+
+			const FString BundleSubsetName = Section.RightChop(InstallBundleUtil::GetInstallBundleSubsetSectionPrefix().Len());
+			BundleSubsetRegexList.Emplace(TPair<FString, TArray<FRegexPattern>>(BundleSubsetName, MoveTemp(SearchRegexPatterns)));
+		}
+
+		return BundleSubsetRegexList;
+	}
+
 	bool MatchBundleRegex(
 		const TArray<TPair<FString, TArray<FRegexPattern>>>& BundleRegexList,
 		const FString& Path,
@@ -272,6 +307,29 @@ namespace InstallBundleUtil
 		return false;
 	}
 
+	bool MatchBundleSubsetRegex(
+		const FString& BundleSubset,
+		const TArray<TPair<FString, TArray<FRegexPattern>>>& BundleSubsetRegexList,
+		const FString& Path)
+	{
+		const TPair<FString, TArray<FRegexPattern>>* BundleRegexPair = Algo::FindByPredicate(BundleSubsetRegexList,
+			[&BundleSubset, &Path](const TPair<FString, TArray<FRegexPattern>>& Pair)
+			{
+				if (Pair.Key != BundleSubset)
+				{
+					return false;
+				}
+
+				const TArray<FRegexPattern>& SearchRegexPatterns = Pair.Value;
+				return Algo::AnyOf(SearchRegexPatterns, [&Path](const FRegexPattern& Pattern)
+				{
+					return FRegexMatcher(Pattern, Path).FindNext();
+				});
+			});
+
+		return !!BundleRegexPair;
+	}
+
 	FName FInstallBundleManagerKeepAwake::Tag(TEXT("InstallBundleManagerKeepAwake"));
 	FName FInstallBundleManagerKeepAwake::TagWithRendering(TEXT("InstallBundleManagerKeepAwakeWithRendering"));
 
diff --git a/Engine/Source/Runtime/InstallBundleManager/Public/InstallBundleManagerInterface.h b/Engine/Source/Runtime/InstallBundleManager/Public/InstallBundleManagerInterface.h
--- a/Engine/Source/Runtime/InstallBundleManager/Public/InstallBundleManagerInterface.h
+++ b/Engine/Source/Runtime/InstallBundleManager/Public/InstallBundleManagerInterface.h
@@ -148,6 +148,8 @@ public:
 
 	INSTALLBUNDLEMANAGER_API virtual const TSharedPtr<IInstallBundleSource> GetBundleSource(FInstallBundleSourceType SourceType) const;
 
+	// Note: Some bundles sources may call this from a background task, so it should be threadsafe
+	virtual int32 SelectBundleSubset(const FString& BundleName, TConstArrayView<FString> BundleSubsets) const { return INDEX_NONE; }
 
 	virtual FDelegateHandle PushInitErrorCallback(FInstallBundleManagerInitErrorHandler Callback) = 0;
 	virtual void PopInitErrorCallback() = 0;
diff --git a/Engine/Source/Runtime/InstallBundleManager/Public/InstallBundleUtils.h b/Engine/Source/Runtime/InstallBundleManager/Public/InstallBundleUtils.h
--- a/Engine/Source/Runtime/InstallBundleManager/Public/InstallBundleUtils.h
+++ b/Engine/Source/Runtime/InstallBundleManager/Public/InstallBundleUtils.h
@@ -39,6 +39,8 @@ namespace InstallBundleUtil
 
 	INSTALLBUNDLEMANAGER_API const FString& GetInstallBundleSectionPrefix();
 
+	INSTALLBUNDLEMANAGER_API const FString& GetInstallBundleSubsetSectionPrefix();
+
 	INSTALLBUNDLEMANAGER_API bool GetConfiguredBundleSources(TArray<FString>& OutSources, TMap<FString, FString>& OutFallbackSources);
 
 	// returns true if the given BundleName exists in the InstallBundle.ini config
@@ -63,12 +65,22 @@ namespace InstallBundleUtil
 		const FConfigFile& InstallBundleConfig,
 		TFunctionRef<bool(const FConfigFile& InstallBundleConfig, const FString& Section)> SectionPredicate = AllInstallBundlePredicate);
 
+	// Returns an ordered regex list used to map files to bundle subsets
+	INSTALLBUNDLEMANAGER_API TArray<TPair<FString, TArray<FRegexPattern>>> LoadBundleSubsetRegexFromConfig(
+		const FConfigFile& InstallBundleConfig);
+
 	// Finds a matching install bundle for a given path
 	INSTALLBUNDLEMANAGER_API bool MatchBundleRegex(
 		const TArray<TPair<FString, TArray<FRegexPattern>>>& BundleRegexList,
 		const FString& Path,
 		FString& OutBundleName);
 
+	// Return true for a matching install bundle subset for a given path
+	INSTALLBUNDLEMANAGER_API bool MatchBundleSubsetRegex(
+		const FString& BundleSubset,
+		const TArray<TPair<FString, TArray<FRegexPattern>>>& BundleSubsetRegexList,
+		const FString& Path);
+
 	constexpr float MinimumBundleWeight = 0.05f;
 
 	// It would really be nice to have these in core
