Fix bad merge

#Synced-CL 46060913

diff --git a/Engine/Source/Programs/AutomationTool/AutomationUtils/BundleUtils.cs b/Engine/Source/Programs/AutomationTool/AutomationUtils/BundleUtils.cs
--- a/Engine/Source/Programs/AutomationTool/AutomationUtils/BundleUtils.cs
+++ b/Engine/Source/Programs/AutomationTool/AutomationUtils/BundleUtils.cs
@@ -302,6 +302,17 @@ namespace AutomationUtils.Automation
 			return Results;
 		}
 
+		public static IEnumerable<string> GetCloudAppsForBundles<TPlatformBundleSettings>(IReadOnlyDictionary<string, TPlatformBundleSettings> InstallBundles) where TPlatformBundleSettings : BundleSettings
+		{
+			// Gather all bundle clouda app name overrides
+			var CloudAppNames = new HashSet<string>();
+			foreach (var Bundle in InstallBundles.Values)
+			{
+				CloudAppNames.Add(Bundle.CloudAppNameOverride);
+			}
+			return CloudAppNames;
+		}
+
 		public static IEnumerable<TPlatformBundleSettings> GetBundleDependencies<TPlatformBundleSettings>(
 			TPlatformBundleSettings Bundle, IReadOnlyDictionary<string, TPlatformBundleSettings> InstallBundles) where TPlatformBundleSettings : BundleSettings
 		{
