UAT: Add a -AutomationProject=name1+name2+...+nameN parameter to filer the found .csproj projects to build and load.
Allow for faster iteration when the user knows the project of the UAT command to run.

Example:
"del /S /Q Engine\Binaries\DotNET\AutomationTool && Engine\Build\BatchFiles\RunUAT.bat P4WriteConfig"
"del /S /Q Engine\Binaries\DotNET\AutomationTool && Engine\Build\BatchFiles\RunUAT.bat P4WriteConfig -AutomationProject=AutomationUtils"
=> Total script initialization time down from 26s to 3s

#rnx

#Synced-CL 0

diff --git a/Engine/Source/Programs/AutomationTool/Program.cs b/Engine/Source/Programs/AutomationTool/Program.cs
--- a/Engine/Source/Programs/AutomationTool/Program.cs
+++ b/Engine/Source/Programs/AutomationTool/Program.cs
@@ -40,8 +40,9 @@ namespace AutomationToolDriver
 			bool bGlobalParam = AutomationToolCommandLine.TrySetGlobal(CurrentParam);
 
 			// Global value parameters, handled explicitly
-			string Option_ScriptsForProject = "-ScriptsForProject";
-			string Option_ScriptDir = "-ScriptDir";
+			string Option_ScriptsForProject = "-ScriptsForProject";	// restrict game search directories to the the specified uproject
+			string Option_ScriptDir = "-ScriptDir";					// additional search directory for .csproj files
+			string Option_AutomationProjects = "-AutomationProject";// filter the found .csproj projects to build and load, specified as name1+name2+...+nameN
 			string Option_Telemetry = "-Telemetry";
 			string Option_WaitForStdStreams = "-WaitForStdStreams";
 
@@ -101,6 +102,20 @@ namespace AutomationToolDriver
 					}
 				}
 			}
+			else if (CurrentParam.StartsWith(Option_AutomationProjects + "=", StringComparison.InvariantCultureIgnoreCase))
+            {
+				string AutomationProjectString = CurrentParam.Substring(CurrentParam.IndexOf('=') + 1);
+				List<string> AutomationProjects = AutomationProjectString.Split("+".ToCharArray()).ToList();
+				if (AutomationProjects.Count > 0)
+				{
+					// AutomationUtils is a necessity, see MainProc()
+					if (!AutomationProjects.Contains("AutomationUtils"))
+					{
+						AutomationProjects.Add("AutomationUtils");
+					}
+					AutomationToolCommandLine.SetUnchecked(Option_AutomationProjects, AutomationProjects);
+				}
+			}
 			else if (CurrentParam.StartsWith(Option_Telemetry + "=", StringComparison.InvariantCultureIgnoreCase))
             {
 				string TelemetryPath = CurrentParam.Substring(CurrentParam.IndexOf('=') + 1);
@@ -487,6 +502,7 @@ namespace AutomationToolDriver
 			Stopwatch Timer = Stopwatch.StartNew();
 			string ScriptsForProject = (string)AutomationToolCommandLine.GetValueUnchecked("-ScriptsForProject");
 			List<string> AdditionalScriptDirs = (List<string>) AutomationToolCommandLine.GetValueUnchecked("-ScriptDir");
+			List<string> AutomationProjects = (List<string>) AutomationToolCommandLine.GetValueUnchecked("-AutomationProject");
 			bool bForceCompile = AutomationToolCommandLine.IsSetGlobal("-Compile");
 			bool bNoCompile = AutomationToolCommandLine.IsSetGlobal("-NoCompile");
 			bool bUseBuildRecords = !AutomationToolCommandLine.IsSetGlobal("-IgnoreBuildRecords");
@@ -495,8 +511,15 @@ namespace AutomationToolDriver
 				: AutomationToolCommandLine.CommandsToExecute;
 			bool bBuildSuccess;
 			HashSet<FileReference> ScriptModuleAssemblyPaths = CompileScriptModule.InitializeScriptModules(
-					Rules.RulesFileType.AutomationModule, ScriptsForProject, AdditionalScriptDirs, bForceCompile, bNoCompile, bUseBuildRecords, 
-					out bBuildSuccess, Logger);
+					Rules.RulesFileType.AutomationModule,
+					ScriptsForProject,
+					AdditionalScriptDirs,
+					AutomationProjects,
+					bForceCompile,
+					bNoCompile,
+					bUseBuildRecords, 
+					out bBuildSuccess,
+					Logger);
 
 			if (!bBuildSuccess)
             {
diff --git a/Engine/Source/Programs/Shared/EpicGames.Build/System/CompileScriptModules.cs b/Engine/Source/Programs/Shared/EpicGames.Build/System/CompileScriptModules.cs
--- a/Engine/Source/Programs/Shared/EpicGames.Build/System/CompileScriptModules.cs
+++ b/Engine/Source/Programs/Shared/EpicGames.Build/System/CompileScriptModules.cs
@@ -138,6 +138,7 @@ namespace UnrealBuildBase
 		/// <param name="RulesFileType"></param>
 		/// <param name="ScriptsForProjectFileName"></param>
 		/// <param name="AdditionalScriptsFolders"></param>
+		/// <param name="AutomationProjectIncludeFilter"></param>
 		/// <param name="bForceCompile"></param>
 		/// <param name="bNoCompile"></param>
 		/// <param name="bUseBuildRecords"></param>
@@ -145,7 +146,7 @@ namespace UnrealBuildBase
 		/// <param name="Logger"></param>
 		/// <returns>Collection of all the projects.  They will have been compiled.</returns>
 		public static HashSet<FileReference> InitializeScriptModules(Rules.RulesFileType RulesFileType,
-			string? ScriptsForProjectFileName, List<string>? AdditionalScriptsFolders, bool bForceCompile, bool bNoCompile, bool bUseBuildRecords,
+			string? ScriptsForProjectFileName, List<string>? AdditionalScriptsFolders, List<string>? AutomationProjectIncludeFilter, bool bForceCompile, bool bNoCompile, bool bUseBuildRecords,
 			out bool bBuildSuccess, ILogger Logger)
 		{
 			Stopwatch Timer = Stopwatch.StartNew();
@@ -169,6 +170,10 @@ namespace UnrealBuildBase
 				// Project scripts require source engine builds
 				GameFolders: Unreal.IsEngineInstalled() ? GameDirectories : [],
 				ForeignPlugins: null, AdditionalSearchPaths: [.. AdditionalDirectories, .. GameBuildDirectories])];
+			if (AutomationProjectIncludeFilter != null && AutomationProjectIncludeFilter.Count > 0)
+			{
+				FoundProjects = FoundProjects.Where(x => (AutomationProjectIncludeFilter.Contains(x.GetFileNameWithoutAnyExtensions(), StringComparer.OrdinalIgnoreCase))).ToHashSet();
+			}
 			Logger.LogDebug("Find {Count} projects time: {TimeSeconds:0.00} s", FoundProjects.Count(), Timer.Elapsed.TotalSeconds);
 
 			BuildFlags BuildFlags = BuildFlags.None;
