Automatically configure the targets in a registry with the IP of the device for their requested role.

#rnx

#Synced-CL 48351862

diff --git a/Engine/Source/Programs/AutomationTool/Gauntlet/Framework/RpcFramework/Gauntlet.AsyncTestNode.cs b/Engine/Source/Programs/AutomationTool/Gauntlet/Framework/RpcFramework/Gauntlet.AsyncTestNode.cs
--- a/Engine/Source/Programs/AutomationTool/Gauntlet/Framework/RpcFramework/Gauntlet.AsyncTestNode.cs
+++ b/Engine/Source/Programs/AutomationTool/Gauntlet/Framework/RpcFramework/Gauntlet.AsyncTestNode.cs
@@ -80,6 +80,7 @@ namespace Gauntlet
 			RpcExecutor.Instance.StartListenThread();
 			RpcExecutor.Instance.SetRpcTargetRegistry(TargetRegistry);
 			TargetRegistry.UpdateAvailableRpcForAllTargets();
+			TargetRegistry.UpdateIPsForRemoteAndNonDesktopTargets(UnrealApp, ReportInterrupt);
 		}
 
 		/// <summary>
diff --git a/Engine/Source/Programs/AutomationTool/Gauntlet/Framework/RpcFramework/Gauntlet.RpcTargetRegistry.cs b/Engine/Source/Programs/AutomationTool/Gauntlet/Framework/RpcFramework/Gauntlet.RpcTargetRegistry.cs
--- a/Engine/Source/Programs/AutomationTool/Gauntlet/Framework/RpcFramework/Gauntlet.RpcTargetRegistry.cs
+++ b/Engine/Source/Programs/AutomationTool/Gauntlet/Framework/RpcFramework/Gauntlet.RpcTargetRegistry.cs
@@ -527,6 +527,40 @@ namespace Gauntlet
 			}
 		}
 
+		/// <summary>
+		/// Replace the current IP address values with the values obtained during devices reservation.
+		/// IP addresses are replaced for all remote devices, as well as for local non-desktop devices.
+		/// </summary>
+		/// <param name="UnrealApp"></param>
+		/// <param name="ActionReportInterrupt">Required for correct completion of test in case of failure and facilitating debugging.</param>
+		/// <exception cref="TestException"></exception>
+		public void UpdateIPsForRemoteAndNonDesktopTargets(UnrealSession UnrealApp, Action<string, object[]> ActionReportInterrupt)
+		{
+			foreach (ITargetDevice ReservedDevice in UnrealApp.UnrealDeviceReservation.ReservedDevices)
+			{
+				bool IsLocalDevice = ReservedDevice.Name.StartsWith("LocalDevice");
+
+				if ((ReservedDevice is not TargetDeviceDesktopCommon) || (!IsLocalDevice))
+				{
+					UnrealSessionRole TargetRole = UnrealApp.RolesToDevices.FirstOrDefault(x => x.Value == ReservedDevice).Key;
+					string RpcTargetName = TargetRole.CommandLineParams.GetParamValue("rpcsenderid").ToString();
+
+					try
+					{
+						this.RpcTargets[RpcTargetName].IpAddress = ReservedDevice.Name;
+					}
+					catch (KeyNotFoundException Ex)
+					{
+						string ErrorMsg = string.Format("Test failed due to absence of required RpcTargetName in RpcTargets dictionary." +
+														" Message: {0}, Call Stack: {1}", Ex.Message, Ex.StackTrace);
+						ActionReportInterrupt(ErrorMsg, [Ex]);
+						
+						throw new TestException(ErrorMsg);
+					}
+				}
+			}
+		}
+
 		/// <summary>
 		/// This logic is similar to determining if -multihome should be used so being on VPN doesn't
 		/// cause IP's to conflict.
