Windows Performance Recorder capture of ETW Trace artifacts

#Synced-CL 48361264

diff --git a/Engine/Source/Programs/AutomationTool/Gauntlet/Unreal/Game/EpicGame.TestConfig.cs b/Engine/Source/Programs/AutomationTool/Gauntlet/Unreal/Game/EpicGame.TestConfig.cs
--- a/Engine/Source/Programs/AutomationTool/Gauntlet/Unreal/Game/EpicGame.TestConfig.cs
+++ b/Engine/Source/Programs/AutomationTool/Gauntlet/Unreal/Game/EpicGame.TestConfig.cs
@@ -5,10 +5,13 @@ using Gauntlet;
 using System;
 using System.Collections.Generic;
 using System.Data;
+using System.Diagnostics;
 using System.IO;
 using System.Linq;
 using System.Net;
 using System.Net.NetworkInformation;
+using System.Runtime.InteropServices;
+using System.Security.Principal;
 using UnrealBuildBase;
 using UnrealBuildTool;
 
@@ -113,23 +116,31 @@ namespace EpicGame
 		public string McpSubRegion = string.Empty;
 
 		/// <summary>
-		/// Comma-separated list of Unreal Insights trace channels to collect for Editor roles, saved to a file named EditorTraces{ID#}.utrace
+		/// Comma-separated list of Unreal Insights trace channels to collect for Editor roles, saved to a file named EditorTrace{ID#}.utrace
 		/// </summary>
 		[AutoParam]
 		public string EditorTraces { get; set; }
 
 		/// <summary>
-		/// Comma-separated list of Unreal Insights trace channels to collect for Client roles, saved to a file named ClientTraces{ID#}.utrace
+		/// Comma-separated list of Unreal Insights trace channels to collect for Client roles, saved to a file named ClientTrace{ID#}.utrace
 		/// </summary>
 		[AutoParam]
 		public string ClientTraces { get; set; }
 
 		/// <summary>
-		/// Comma-separated list of Unreal Insights trace channels to collect for Server roles, saved to a file named ServerTraces{ID#}.utrace
+		/// Comma-separated list of Unreal Insights trace channels to collect for Server roles, saved to a file named ServerTrace{ID#}.utrace
 		/// </summary>
 		[AutoParam]
 		public string ServerTraces { get; set; }
 
+		/// <summary>
+		/// Enables Windows Performance Recorder on the local machine, saved to a file named WinPerfReportTrace{ID#}.etl
+		/// Can produce around 1 GB/minute, avoid for long-running tests.
+		/// Requires running with Admin privileges.
+		/// </summary>
+		[AutoParam]
+		public bool WprTraces = false;
+
 		/// <summary>
 		/// Filepaths to trace logging artifacts, copied after test execution ends
 		/// </summary>
@@ -139,6 +150,7 @@ namespace EpicGame
 		protected bool SkipMcpSetup = false;
 
 		// incrementing value to ensure we can assign unique values to ports etc
+		private readonly int CurrentConfigIndex;
 		static private int NumberOfConfigsCreated = 0;
 		static private int NumberOfEditorsConfigured = 0;
 		static private int NumberOfClientsConfigured = 0;
@@ -148,6 +160,7 @@ namespace EpicGame
 
 		public EpicGameTestConfig()
 		{
+			CurrentConfigIndex = NumberOfConfigsCreated;
 			NumberOfConfigsCreated++;
 		}
 
@@ -388,6 +401,11 @@ namespace EpicGame
 			}
 			return GeneratedTraceFile;
 		}
+
+		/// <summary>
+		/// Path to generate Windows Performance Recorder artifacts
+		/// </summary>
+		public string CurrentConfigWprPath() => Path.Join(EpicGameTestConfig.AutomationToolLogsDir, $"WinPerfReportTrace{CurrentConfigIndex}.etl");
 	}
 
 	/// <summary>
@@ -401,13 +419,130 @@ namespace EpicGame
 
 		}
 
+		private static bool IsWindowsAdmin()
+		{
+			if (RuntimeInformation.IsOSPlatform(OSPlatform.Windows))
+			{
+				using var Identity = WindowsIdentity.GetCurrent();
+				return new WindowsPrincipal(Identity).IsInRole(WindowsBuiltInRole.Administrator);
+			}
+			return false;
+		}
+
+		public override bool StartTest(int Pass, int InNumPasses)
+		{
+			if(CachedConfig.WprTraces)
+			{
+				StartWprCapture();
+			}
+			return base.StartTest(Pass, InNumPasses);
+		}
+
+		private void StartWprCapture()
+		{
+			if (!IsWindowsAdmin())
+			{
+				Log.Warning("Windows admin privilege is required for Windows Performance Recorder. Capture not enabled.");
+			}
+			else
+			{
+				SafeCancelWprCapture();
+				Log.Info($"Starting capture from Windows Performance Recorder");
+				// Light profiles limit artifact size, WPR profiles are verbose by default which can capture multiple gigabytes per minute
+				var WprStartInfo = new ProcessStartInfo("wpr", "-start CPU.light -start FileIO.light -start DiskIO.light -filemode")
+				{
+					// prevent wpr from modifying our io streams to display secondary perf statistics
+					UseShellExecute = false,
+					RedirectStandardInput = true,
+					RedirectStandardOutput = true,
+					RedirectStandardError = true,
+				};
+
+				Globals.PostAbortHandlers.Add(() =>
+				{
+					Log.Warning("Cancelling Windows Performance Recorder capture, please wait while this operation completes");
+					SafeCancelWprCapture();
+				});
+
+				Process WprProcess = Process.Start(WprStartInfo);
+				if (!WprProcess.WaitForExit(TimeSpan.FromSeconds(10)))
+				{
+					// unexpected, should immediately return
+					Log.Warning("Windows Performance Recorder unexpectedly failed to exit while starting trace collection, capture may be unavailable");
+				}
+				else
+				{
+					if (WprProcess.ExitCode != 0)
+					{
+						Log.Warning($"Windows Performance Recorder returned exit code '{WprProcess.ExitCode}' while starting trace collection, capture may be unavailable");
+					}
+				}
+			}
+		}
+
+		private void SafeCancelWprCapture()
+		{
+			if (IsWindowsAdmin())
+			{
+				Log.Verbose("Cancelling WPR capture");
+				var WprCancelArgs = new ProcessStartInfo("wpr", $"-cancel")
+				{
+					// prevent wpr from modifying our io streams to report progress
+					UseShellExecute = false,
+					RedirectStandardInput = true,
+					RedirectStandardOutput = true,
+					RedirectStandardError = true,
+				};
+				Process WprProcess = Process.Start(WprCancelArgs);
+				if (!WprProcess.WaitForExit(TimeSpan.FromSeconds(10)))
+				{
+					// unexpected, should immediately return
+					Log.Warning("Windows Performance Recorder unexpectedly failed to exit while cancelling trace collection");
+				} // ignore exit code, cancel errors when there is no active capture
+			}
+		}
+
+		private void EndWprCapture(string EtlOutputPath)
+		{
+			if (IsWindowsAdmin())
+			{
+				Log.Verbose($"Stopping WPR to output results into: {EtlOutputPath}");
+				CachedConfig.TraceLogPaths.Add(EtlOutputPath);
+
+				var EtlDescription = $"Gauntlet Test '{Name}'";
+				var WprStopArgs = new ProcessStartInfo("wpr", $"-stop \"{EtlOutputPath}\" \"{EtlDescription}\" -skipPdbGen")
+				{
+					// prevent wpr from modifying our io streams to report progress
+					UseShellExecute = false,
+					RedirectStandardInput = true,
+					RedirectStandardOutput = true,
+					RedirectStandardError = true,
+				};
+				Process WprProcess = Process.Start(WprStopArgs);
+				if (!WprProcess.WaitForExit(TimeSpan.FromMinutes(10))) // generation can take a few minutes
+				{
+					Log.Error("Windows Performance Recorder unexpectedly failed to exit while stopping trace collection");
+				}
+			}
+		}
+
 		public override IEnumerable<UnrealRoleArtifacts> SaveRoleArtifacts(string OutputPath)
 		{
+			if(CachedConfig.WprTraces)
+			{
+				EndWprCapture(CachedConfig.CurrentConfigWprPath());
+			}
+			
 			string OutputFolderName = new DirectoryInfo(OutputPath).Name;
 			foreach (string SourceTracePath in CachedConfig.TraceLogPaths)
 			{
+				if (!File.Exists(SourceTracePath))
+				{
+					Log.Warning($"Skipping save of missing trace artifact: {SourceTracePath}");
+					continue;
+				}
+
 				var Filename = Path.GetFileName(SourceTracePath);
-				
 				if (CommandUtils.IsBuildMachine && Gauntlet.Horde.IsHordeJob)
 				{
 					string HordeTracePath = Path.GetFullPath(Path.Combine(CommandUtils.CmdEnv.LogFolder, OutputFolderName, Filename));
