Deletes unnecessary .bat file.

#Synced-CL 48346153

diff --git a/Engine/Build/Graph/Tests/IncrementalCompare.xml b/Engine/Build/Graph/Tests/IncrementalCompare.xml
--- a/Engine/Build/Graph/Tests/IncrementalCompare.xml
+++ b/Engine/Build/Graph/Tests/IncrementalCompare.xml
@@ -188,17 +188,9 @@
 			<Commandlet Name="Cook" Project="$(RelativeProjectFile)" Arguments="$(CommonCookArgs) -forcerecook"/>
 		</Node>
 		<Node Name="$(DiffCDOsNodeName)" Requires="$(FullCookNodeName)">
-			<!-- The First run of DumpArchetypeInfo to dump the CDOs into a folder -->
-			<Log Message="$(NewLine)========== IncrementalCompare.xml: 1st Dump CDO =========="/>
-			<Commandlet Name="DumpArchetypeInfo" Project="$(RelativeProjectFile)" Arguments="test_001"/>
-
-			<!-- The Second run of DumpArchetypeInfo to dump the CDOs into a folder -->
-			<Log Message="$(NewLine)========== IncrementalCompare.xml: 2nd Dump CDO =========="/>
-			<Commandlet Name="DumpArchetypeInfo" Project="$(RelativeProjectFile)" Arguments="test_002"/>
-
-			<!-- Diff the two CDOs -->
+			<!-- Generate and Diff the two CDO folders -->
 			<Log Message="$(NewLine)========== IncrementalCompare.xml: Diff CDOs =========="/>
-			<Command Name="DiffCDOFolders" Arguments="-FirstRunFolder=$(RelativeProjectDir)\\Saved\\Temp\\test_001 -SecondRunFolder=$(RelativeProjectDir)\\Saved\\Temp\\test_002"/>
+			<Command Name="DiffCDOFolders" Arguments="-UnrealProjectPath=$(RelativeProjectFile) -FirstRunFolderName=test_001 -SecondRunFolderName=test_002"/>
 		</Node>
 		<Node Name="$(StoreFullCookNodeName)" Requires="$(DiffCDOsNodeName)">
 			<!-- Copy all cook artifacts to the compare_incremental dir. Copy rather than move so they remain
diff --git a/Engine/Source/Programs/AutomationTool/Scripts/DiffCDOFolders.cs b/Engine/Source/Programs/AutomationTool/Scripts/DiffCDOFolders.cs
--- a/Engine/Source/Programs/AutomationTool/Scripts/DiffCDOFolders.cs
+++ b/Engine/Source/Programs/AutomationTool/Scripts/DiffCDOFolders.cs
@@ -4,6 +4,8 @@ using System;
 using System.IO;
 using System.Text;
 using System.Text.Json;
+using EpicGames.Core;
+using Microsoft.Extensions.Logging;
 
 namespace AutomationTool
 {
@@ -11,18 +13,20 @@ namespace AutomationTool
 	{
 		public override void ExecuteBuild()
 		{
-			string firstRunPath = ParseParamValue("FirstRunFolder") ?? "";
-			string secondRunPath = ParseParamValue("SecondRunFolder") ?? "";
-
-			if (string.IsNullOrEmpty(firstRunPath) || string.IsNullOrEmpty(secondRunPath))
+			if (ParseParam("Help"))
 			{
-				throw new AutomationException("Usage: FirstRunFolder=\"PATH from root of the project\" SecondRunFolder=\"PATH from root of the project\"");
+				Logger.LogInformation("Usage: DiffCDOFolders -UnrealProjectPath=\"Path to the unreal project\" -FirstRunFolderName=\"Name of the first run\" -SecondRunFolderName=\"Name of the second run\" -NoGenerate (optional to not generate the CDO folders, and use the existing ones)");
+				return;
 			}
 
-			firstRunPath = Path.Combine(Directory.GetCurrentDirectory(), firstRunPath);
-			secondRunPath = Path.Combine(Directory.GetCurrentDirectory(), secondRunPath);
+			if (!ParseParam("NoGenerate")) {
+				GenerateCDOFolders(new FileReference(ParseParamValue("UnrealProjectPath")), ParseParamValue("FirstRunFolderName"), ParseParamValue("SecondRunFolderName"));
+			}
 
-			if (!Directory.Exists(firstRunPath) || !Directory.Exists(secondRunPath))
+			string firstRunPath = Path.Combine(Path.GetDirectoryName(ParseParamValue("UnrealProjectPath")),"Saved", "Temp", ParseParamValue("FirstRunFolderName"));
+			string secondRunPath = Path.Combine(Path.GetDirectoryName(ParseParamValue("UnrealProjectPath")),"Saved", "Temp", ParseParamValue("SecondRunFolderName"));
+
+			if (string.IsNullOrEmpty(firstRunPath) || string.IsNullOrEmpty(secondRunPath))
 			{
 				throw new AutomationException("Both paths must exist. FirstRunFolder: " + firstRunPath + " SecondRunFolder: " + secondRunPath);
 			}
@@ -78,8 +82,21 @@ namespace AutomationTool
 			}
 			if (textDiff.Length > 0)
 			{
-				throw new AutomationException("Diff was detected.\n" + textDiff.ToString() + "\n To run locally, run Engine\\Source\\Programs\\AutomationTool\\Scripts\\diffCDOs.bat (from the root of the depot)");
+				throw new AutomationException("Diff was detected.\n" + textDiff.ToString() + "\n To run locally, run .\\runUAT.bat DiffCDOFolders -UnrealProjectPath=\"Path to the unreal project\" -FirstRunFolderName=\"Name of the first run\" -SecondRunFolderName=\"Name of the second run\"");
+			} 
+			else
+			{
+				Logger.LogInformation("No diff was detected.");
 			}
 		}
+
+		private void GenerateCDOFolders(FileReference unrealProjectPath, string firstRunFolderName, string secondRunFolderName)
+		{
+			Logger.LogInformation("Generating CDO folders...");
+			string editorPath = Path.Combine(CmdEnv.LocalRoot, "Engine", "Binaries", "Win64", "UnrealEditor.exe");
+
+			RunCommandlet(unrealProjectPath, editorPath, "DumpArchetypeInfo", firstRunFolderName);
+			RunCommandlet(unrealProjectPath, editorPath, "DumpArchetypeInfo", secondRunFolderName);
+		}
 	}
 }
diff --git a/Engine/Source/Programs/AutomationTool/Scripts/diffCDOs.bat b/Engine/Source/Programs/AutomationTool/Scripts/diffCDOs.bat
deleted file mode 100755
