Removing appium backgrounding and alert checking temporarily.
#Synced-CL 47873187

diff --git a/Engine/Source/Programs/AutomationTool/Gauntlet/Platform/IOS/Gauntlet.TargetDeviceIOS.cs b/Engine/Source/Programs/AutomationTool/Gauntlet/Platform/IOS/Gauntlet.TargetDeviceIOS.cs
--- a/Engine/Source/Programs/AutomationTool/Gauntlet/Platform/IOS/Gauntlet.TargetDeviceIOS.cs
+++ b/Engine/Source/Programs/AutomationTool/Gauntlet/Platform/IOS/Gauntlet.TargetDeviceIOS.cs
@@ -833,16 +833,6 @@ namespace Gauntlet
 				throw new DeviceException("Failed to launch on {0}. {1}", Name, AppProcess.Output);
 			}
 
-			if (AppiumContainer != null)
-			{
-				lock (Mutex)
-				{
-					// once in foreground, check for the apple sign on screen popup
-					// TODO: use MDM profile to bypass this
-
-					AppiumContainer.CancelAppleSignOnAlert(IOSInstall.PackageName);
-				}
-			}
 			return new IOSAppInstance(IOSInstall, AppProcess, IOSInstall.CommandLine);
 		}
 
