Change bucket name sanitization in ZenExportSnapshot task to allow 256 characters in the bucket name, not 64 characters as is the limit for StringIds.

diff --git a/Engine/Source/Programs/AutomationTool/BuildGraph/Tasks/ZenExportSnapshotTask.cs b/Engine/Source/Programs/AutomationTool/BuildGraph/Tasks/ZenExportSnapshotTask.cs
--- a/Engine/Source/Programs/AutomationTool/BuildGraph/Tasks/ZenExportSnapshotTask.cs
+++ b/Engine/Source/Programs/AutomationTool/BuildGraph/Tasks/ZenExportSnapshotTask.cs
@@ -467,9 +467,65 @@ namespace AutomationTool.Tasks
 			return false;
 		}
 
+		private static bool IsValidCharacter(char character) => character <= 0x7f && IsValidCharacter((byte)character);
+
+		private static bool IsValidCharacter(byte character)
+		{
+			if (character >= 'a' && character <= 'z')
+			{
+				return true;
+			}
+			if (character >= '0' && character <= '9')
+			{
+				return true;
+			}
+			if (character == '-' || character == '_' || character == '.')
+			{
+				return true;
+			}
+			return false;
+		}
+
 		private static string SanitizeBucketName(string inString)
 		{
-			return StringId.Sanitize(inString).ToString();
+			// Originally matched StringId.Sanitize, but with a MaxLength of 256 instead of 64
+			const int MaxLength = 256;
+			StringBuilder result = new StringBuilder();
+			for (int idx = 0; idx < inString.Length && result.Length < MaxLength; idx++)
+			{
+				char character = (char)inString[idx];
+				if (character >= 'A' && character <= 'Z')
+				{
+					result.Append((char)('a' + (character - 'A')));
+				}
+				else if (character == '.')
+				{
+					if (result.Length > 0)
+					{
+						if (result[^1] == '-')
+						{
+							result[^1] = character;
+						}
+						else
+						{
+							result.Append(character);
+						}
+					}
+				}
+				else if (IsValidCharacter(character))
+				{
+					result.Append(character);
+				}
+				else if (result.Length > 0 && result[^1] != '-' && result[^1] != '.')
+				{
+					result.Append('-');
+				}
+			}
+			while (result.Length > 0 && (result[^1] == '-' || result[^1] == '.'))
+			{
+				result.Remove(result.Length - 1, 1);
+			}
+			return result.ToString();
 		}
 
 		/// <summary>
