Remove un-necessary using declaration in ZenExportSnapshot.

diff --git a/Engine/Source/Programs/AutomationTool/BuildGraph/Tasks/ZenExportSnapshotTask.cs b/Engine/Source/Programs/AutomationTool/BuildGraph/Tasks/ZenExportSnapshotTask.cs
--- a/Engine/Source/Programs/AutomationTool/BuildGraph/Tasks/ZenExportSnapshotTask.cs
+++ b/Engine/Source/Programs/AutomationTool/BuildGraph/Tasks/ZenExportSnapshotTask.cs
@@ -13,11 +13,10 @@ using System.Threading;
 using System.Threading.Tasks;
 using System.Xml;
 using EpicGames.Core;
-using EpicGames.Horde;
 using EpicGames.ProjectStore;
 using EpicGames.Serialization;
-using UnrealBuildTool;
 using Microsoft.Extensions.Logging;
+using UnrealBuildTool;
 
 namespace AutomationTool.Tasks
 {
