generate and stage umeta file.

diff --git a/Engine/Source/Programs/AutomationTool/Scripts/CopyBuildToStagingDirectory.Automation.cs b/Engine/Source/Programs/AutomationTool/Scripts/CopyBuildToStagingDirectory.Automation.cs
--- a/Engine/Source/Programs/AutomationTool/Scripts/CopyBuildToStagingDirectory.Automation.cs
+++ b/Engine/Source/Programs/AutomationTool/Scripts/CopyBuildToStagingDirectory.Automation.cs
@@ -3648,6 +3648,15 @@ namespace AutomationScripts
 				}
 			}
 
+			FileReference UMetaFile = InUtocFile.ChangeExtension(".umeta");
+			if (FileReference.Exists(UMetaFile))
+			{
+				if (!InternalUtils.SafeCopyFile(UMetaFile.FullName, OutputLocation.ChangeExtension(".umeta").FullName))
+				{
+					Logger.LogWarning("Failed to copy umeta {Src} to {Dst}", UMetaFile.FullName, OutputLocation.ChangeExtension(".umeta").FullName);
+				}
+			}
+
 			return bCopiedExistingPak;
 		}
 
@@ -4125,6 +4134,13 @@ namespace AutomationScripts
 				CommonAdditionalArgs += " -fallbackOrderForNonUassetFiles";
 			}
 
+			bool bGenerateContainerMetadata = false; // do we want container metadata to be generated (can be used to reverse lookup package names)
+			PlatformGameConfig.GetBool("/Script/UnrealEd.ProjectPackagingSettings", "bGenerateContainerMetadata", out bGenerateContainerMetadata);
+			if (bGenerateContainerMetadata) 
+			{
+				CommonAdditionalArgs += " -ContainerMeta";
+			}
+
 			if (bHasPlatformEngineConfig)
 			{
 				// if the runtime will want to reduce memory usage, we have to disable the pak index freezing
@@ -4584,6 +4600,13 @@ namespace AutomationScripts
 
 				AdditionalArgs += $" -platform={ConfigHierarchy.GetIniPlatformName(SC.StageTargetPlatform.IniPlatformType)}";
 
+				bool bGenerateIoContainerMetadata = false; // do we want container metadata to be generated (can be used to reverse lookup package names)
+				PlatformGameConfig.GetBool("/Script/UnrealEd.ProjectPackagingSettings", "bGenerateContainerMetadata", out bGenerateIoContainerMetadata );
+				if (bGenerateIoContainerMetadata )
+				{
+					AdditionalArgs += " -ContainerMeta";
+				}
+
 				AdditionalArgs += " " + Params.AdditionalIoStoreOptions;
 
 				RunIoStore(Params, SC, IoStoreCommandsFileName, FinalOrderFiles, AdditionalArgs);
@@ -4843,6 +4866,7 @@ namespace AutomationScripts
 								IncludedExtensions.Add(".ucas");
 								IncludedExtensions.Add(".utoc");
 								IncludedExtensions.Add(".uondemandtoc");
+								IncludedExtensions.Add(".umeta");
 							}
 							IEnumerable<string> PakFileSet = Directory.EnumerateFiles(ExistingPatchSearchPath, PakName + "-" + SC.FinalCookPlatform + "*.*");
 							foreach (string PakFilePath in PakFileSet)
diff --git a/Engine/Source/Runtime/Experimental/IoStore/OnDemand/Private/OnDemandContentInstaller.cpp b/Engine/Source/Runtime/Experimental/IoStore/OnDemand/Private/OnDemandContentInstaller.cpp
--- a/Engine/Source/Runtime/Experimental/IoStore/OnDemand/Private/OnDemandContentInstaller.cpp
+++ b/Engine/Source/Runtime/Experimental/IoStore/OnDemand/Private/OnDemandContentInstaller.cpp
@@ -50,6 +50,13 @@ namespace UE::IoStore
 				IoStoreSimulateHttpError,
 				TEXT("Simulate http error code in iostore.")
 			);
+
+		float IoStoreChunkMissingErrorChance = 0.0f;
+		static FAutoConsoleVariableRef CVar_IoStoreChunkMissingErrorChange(
+			TEXT("iostore.ChunksMissingRandomErrorChance"),
+			IoStoreChunkMissingErrorChance,
+			TEXT("Simulate chunk missing error based on chance 0 - 100.")
+			);
 #endif // !UE_BUILD_SHIPPING
 
 		static bool bContentInstallerFlushOnIdle = false;
@@ -113,6 +120,15 @@ void ResolvePackageDependencies(
 		const FPackageStoreEntry& PackageStoreEntry,
 		const TSharedPtr<IPackageStoreBackend>& Backend) -> bool
 	{
+#if !UE_BUILD_SHIPPING
+		if (UE::IoStore::CVars::IoStoreChunkMissingErrorChance > 0.0f)
+		{
+			if (FMath::RandRange(0.0f, 100.0f) < UE::IoStore::CVars::IoStoreChunkMissingErrorChance)
+			{
+				EntryStatus = EPackageStoreEntryStatus::Missing;
+			}
+		}
+#endif
 		if (EntryStatus != EPackageStoreEntryStatus::Missing)
 		{
 			OutResolved.Add(PackageId);
diff --git a/Engine/Source/Runtime/Experimental/IoStore/OnDemandCore/Private/OnDemandError.cpp b/Engine/Source/Runtime/Experimental/IoStore/OnDemandCore/Private/OnDemandError.cpp
--- a/Engine/Source/Runtime/Experimental/IoStore/OnDemandCore/Private/OnDemandError.cpp
+++ b/Engine/Source/Runtime/Experimental/IoStore/OnDemandCore/Private/OnDemandError.cpp
@@ -1,6 +1,7 @@
 // Copyright Epic Games, Inc. All Rights Reserved.
 
 #include "IO/OnDemandError.h"
+#include "IO/IoStore.h"
 
 UE_DEFINE_ERROR_MODULE(IoStoreOnDemand);
 UE_DEFINE_ERROR(HttpError, IoStoreOnDemand);
@@ -18,3 +19,45 @@ UE_DEFINE_ERROR(CasSnapshotError, IoStoreOnDemand);
 namespace UE::UnifiedError::IoStore
 {
 } // namespace UE::UnifiedError::IoStore
+
+
+
+////////////////////////////////////////////////////////////////////////////////
+
+
+void SerializeForLog(FCbWriter& Writer, const FIoChunkId& ChunkId)
+{
+	Writer.BeginObject();
+	Writer.AddString(ANSITEXTVIEW("$type"), ANSITEXTVIEW("FIoChunkId"));
+	FUtf8StringBuilderBase Filename;
+	FUtf8StringView ContainerName;
+	if (!FIoDispatcher::Get().GetFilename(ChunkId, Filename, ContainerName).IsEmpty())
+	{
+		Writer.AddString(ANSITEXTVIEW("$format"), TEXT("({Filename}:{ContainerName}:{ChunkId})"));
+		Writer.AddString(ANSITEXTVIEW("Filename"), Filename);
+		Writer.AddString(ANSITEXTVIEW("ContainerName"), ContainerName);
+	}
+	else
+	{
+		Writer.AddString(ANSITEXTVIEW("$format"), TEXT("{ChunkId}"));
+	}
+	Writer.AddString(ANSITEXTVIEW("ChunkId"), LexToString(ChunkId));
+	Writer.EndObject();
+}
+
+void SerializeForLog(FCbWriter& Writer, const UE::UnifiedError::IoStoreOnDemand::FChunkMissingErrorContext& Ctx)
+{
+	Writer.BeginObject();
+	Writer.AddString(ANSITEXTVIEW("$type"), TErrorStructFeatures<UE::UnifiedError::IoStoreOnDemand::FChunkMissingErrorContext>::GetErrorContextTypeNameAsString());
+	Writer.AddString(ANSITEXTVIEW("$format"), TEXT("(DiscoveredCount: {DiscoveredCount}, MissingCount: {MissingCount}, MissingChunkIds: [{ChunkIds}])"));
+	Writer.AddInteger(ANSITEXTVIEW("MissingCount"), Ctx.ChunkIds.Num());
+	Writer.AddInteger(ANSITEXTVIEW("DiscoveredCount"), Ctx.DiscoveredChunksCount);
+
+	Writer.BeginArray(ANSITEXTVIEW("ChunkIds"));
+	for (const auto& ChunkId : Ctx.ChunkIds)
+	{
+		SerializeForLog(Writer, ChunkId);
+	}
+	Writer.EndArray();
+	Writer.EndObject();
+}
diff --git a/Engine/Source/Runtime/Experimental/IoStore/OnDemandCore/Public/IO/OnDemandError.h b/Engine/Source/Runtime/Experimental/IoStore/OnDemandCore/Public/IO/OnDemandError.h
--- a/Engine/Source/Runtime/Experimental/IoStore/OnDemandCore/Public/IO/OnDemandError.h
+++ b/Engine/Source/Runtime/Experimental/IoStore/OnDemandCore/Public/IO/OnDemandError.h
@@ -197,32 +197,11 @@ inline void SerializeForLog(FCbWriter& Writer, const UE::UnifiedError::IoStoreOn
 }
 
 
-inline void SerializeForLog(FCbWriter& Writer, const FIoChunkId& ChunkId)
-{
-	Writer.BeginObject();
-	Writer.AddString(ANSITEXTVIEW("$type"), ANSITEXTVIEW("FIoChunkId"));
-	Writer.AddString(ANSITEXTVIEW("$text"), LexToString(ChunkId));
-	Writer.EndObject();
-}
+UE_API void SerializeForLog(FCbWriter& Writer, const FIoChunkId& ChunkId);
 
 ////////////////////////////////////////////////////////////////////////////////
 UE_DECLARE_ERRORSTRUCT_FEATURES(IoStoreOnDemand, FChunkMissingErrorContext);
-inline void SerializeForLog(FCbWriter& Writer, const UE::UnifiedError::IoStoreOnDemand::FChunkMissingErrorContext& Ctx)
-{
-	Writer.BeginObject();
-	Writer.AddString(ANSITEXTVIEW("$type"), TErrorStructFeatures<UE::UnifiedError::IoStoreOnDemand::FChunkMissingErrorContext>::GetErrorContextTypeNameAsString());
-	Writer.AddString(ANSITEXTVIEW("$format"), TEXT("(DiscoveredCount: {DiscoveredCount}, MissingCount: {MissingCount}, MissingChunkIds: [{ChunkIds}])"));
-	Writer.AddInteger(ANSITEXTVIEW("MissingCount"), Ctx.ChunkIds.Num());
-	Writer.AddInteger(ANSITEXTVIEW("DiscoveredCount"), Ctx.DiscoveredChunksCount);
-
-	Writer.BeginArray(ANSITEXTVIEW("ChunkIds"));
-	for (const auto& ChunkId : Ctx.ChunkIds)
-	{
-		SerializeForLog(Writer, ChunkId);
-	}
-	Writer.EndArray();
-	Writer.EndObject();
-}
+UE_API void SerializeForLog(FCbWriter& Writer, const UE::UnifiedError::IoStoreOnDemand::FChunkMissingErrorContext& Ctx);
 
 ////////////////////////////////////////////////////////////////////////////////
 UE_DECLARE_ERRORSTRUCT_FEATURES(IoStoreOnDemand, FChunkHashMismatchErrorContext);
