Copy global on-demand TOC if we find the global TOC

#rnx

#Synced-CL 48408831

diff --git a/Engine/Source/Programs/AutomationTool/Scripts/CopyBuildToStagingDirectory.Automation.cs b/Engine/Source/Programs/AutomationTool/Scripts/CopyBuildToStagingDirectory.Automation.cs
--- a/Engine/Source/Programs/AutomationTool/Scripts/CopyBuildToStagingDirectory.Automation.cs
+++ b/Engine/Source/Programs/AutomationTool/Scripts/CopyBuildToStagingDirectory.Automation.cs
@@ -4280,6 +4280,19 @@ namespace AutomationScripts
 						{
 							bCopiedExistingGlobalUtoc = CopyExistingContainer(InUtocFile, OutUtocFile);
 						}
+
+						// Also look for the optional global on-demand TOC
+						string GlobalOnDemandFileName = "global.uondemandtoc";
+						FileReference InFile = FileReference.Combine(SourceOutputLocation.Directory, GlobalOnDemandFileName);
+						if (FileExists_NoExceptions(InFile.FullName))
+						{
+							FileReference OutFile = FileReference.Combine(OutputLocation.Directory, GlobalOnDemandFileName);
+							Logger.LogInformation("Copying global on-demand TOC from {InFile} to {OutFile}", InFile, OutFile);
+							if (!InternalUtils.SafeCopyFile(InFile.FullName, OutFile.FullName))
+							{
+								Logger.LogInformation("Failed to copy global on-demand TOC from {InFile} to {OutFile}", InFile, OutFile);
+							}
+						}
 					}
 				}
 
