Added a list of default values for the parameter to the ParseValues method.

#rnx
#Synced-CL 48351862

diff --git a/Engine/Source/Programs/AutomationTool/Gauntlet/Framework/Gauntlet.Params.cs b/Engine/Source/Programs/AutomationTool/Gauntlet/Framework/Gauntlet.Params.cs
--- a/Engine/Source/Programs/AutomationTool/Gauntlet/Framework/Gauntlet.Params.cs
+++ b/Engine/Source/Programs/AutomationTool/Gauntlet/Framework/Gauntlet.Params.cs
@@ -66,8 +66,9 @@ namespace Gauntlet
 		/// </summary>
 		/// <param name="Param"></param>
 		/// <param name="CommaSeparated"></param>
-		/// <returns></returns>
-		public List<string> ParseValues(string Param, bool CommaSeparated=false)
+		/// <param name="Default">List with default parameter values ​​if the required values ​​are not found in CLI.</param>
+		/// <returns>Parameter values or default values or empty list if param is not found and default list is not set</returns>
+		public List<string> ParseValues(string Param, bool CommaSeparated=false, List<string> Default = default)
 		{
 			List<string> FoundValues = new List<string>();
 
@@ -126,6 +127,11 @@ namespace Gauntlet
 				}
 			}
 
+			if ((FoundValues.Count == 0) && (Default != default))
+			{
+				return Default;
+			}
+
 			return FoundValues;
 		}
 
diff --git a/Engine/Source/Programs/AutomationTool/Gauntlet/SelfTest/Framework/Gauntlet.SelfTest.TestParsing.cs b/Engine/Source/Programs/AutomationTool/Gauntlet/SelfTest/Framework/Gauntlet.SelfTest.TestParsing.cs
--- a/Engine/Source/Programs/AutomationTool/Gauntlet/SelfTest/Framework/Gauntlet.SelfTest.TestParsing.cs
+++ b/Engine/Source/Programs/AutomationTool/Gauntlet/SelfTest/Framework/Gauntlet.SelfTest.TestParsing.cs
@@ -32,7 +32,20 @@ namespace Gauntlet.SelfTest
 			CheckResult(MultiArgs3Result.Count == 2, "Incorrect result count");
 			CheckResult(MultiArgs3Result[0] == "one" && MultiArgs3Result[1] == "two", "Incorrect result count");
 
-			
+			// check default value
+			Params MultiArgs4 = new Params(new[] { "foo=one", "foo=two", "foo=three" });
+			List<string> MultiDefaultResult4 = ["four", "five"];
+			List<string> MultiArgs4Result = MultiArgs4.ParseValues("dummy", Default: MultiDefaultResult4);
+			CheckResult(MultiArgs4Result.Count == 2, "Incorrect result count");
+			CheckResult(MultiArgs4Result[0] == "four" && MultiArgs4Result[1] == "five", "Incorrect result count");
+
+			// negative check
+			Params MultiArgs5 = new Params(new[] { "foo=one" });
+			List<string> MultiArgs5Result = MultiArgs5.ParseValues("dummy");
+			CheckResult(MultiArgs5Result != null, "Result is NULL");
+			CheckResult(MultiArgs5Result.Count == 0, "Incorrect result count");
+
+
 			// simple comma-separated list with no subparams
 			string TestCase1 = "Test1, Test2";
 			var Results1 = ArgumentWithParams.CreateFromString(TestCase1);
