AutomationTool: Remove AutomationTool.csproj.References import

#rnx
#Synced-CL 47605596

diff --git a/Engine/Source/Programs/AutomationTool/AutomationTool.csproj b/Engine/Source/Programs/AutomationTool/AutomationTool.csproj
--- a/Engine/Source/Programs/AutomationTool/AutomationTool.csproj
+++ b/Engine/Source/Programs/AutomationTool/AutomationTool.csproj
@@ -18,9 +18,7 @@
   <PropertyGroup>
     <Configuration Condition=" '$(Configuration)' == '' ">Development</Configuration>
     <OutputType>Exe</OutputType>
-    <AppendTargetFrameworkToOutputPath>false</AppendTargetFrameworkToOutputPath>
     <GenerateAssemblyInfo>false</GenerateAssemblyInfo>
-    <GenerateTargetFrameworkAttribute>false</GenerateTargetFrameworkAttribute>
     <Configurations>Debug;Release;Development;Analyze</Configurations>
     <ApplicationManifest>app.manifest</ApplicationManifest>
     <RootNamespace>AutomationTool</RootNamespace>
@@ -112,7 +110,4 @@
       <Link>$([System.Text.RegularExpressions.Regex]::Replace(%(FullPath), '^.+?[\\/]BatchFiles[\\/]', 'BatchFiles/'))</Link>
     </None>
   </ItemGroup>
-  
-  <!-- Reference file is generated by GenerateProjectFiles scripts (calling UnrealBuildTool). It contains a list of all project files known to UBT based on how it was run. -->
-  <Import Project="..\..\..\Intermediate\ProjectFiles\AutomationTool.csproj.References" Condition="Exists('..\..\..\Intermediate\ProjectFiles\AutomationTool.csproj.References') And '$(AutomationToolProjectOnly)' == ''" />
 </Project>
\ No newline at end of file
