More formally define StepMessages and handle them in a more reliable way.

#Synced-CL 48539628

diff --git a/Engine/Source/Programs/AutomationTool/Gauntlet/Framework/Gauntlet.TestExecutor.cs b/Engine/Source/Programs/AutomationTool/Gauntlet/Framework/Gauntlet.TestExecutor.cs
--- a/Engine/Source/Programs/AutomationTool/Gauntlet/Framework/Gauntlet.TestExecutor.cs
+++ b/Engine/Source/Programs/AutomationTool/Gauntlet/Framework/Gauntlet.TestExecutor.cs
@@ -270,7 +270,7 @@ namespace Gauntlet
 									Node.AddTestEvent(new UnrealTestEvent(EventSeverity.Error, "Test Failed to Start", new List<string> { ex.Message }));
 									PendingTests[i] = null;
 									NodeInfo.TimeSetupBegan = NodeInfo.TimeSetupEnded = NodeInfo.TimeTestEnded = DateTime.Now;
-									Horde.ReportInvalidatingIssue("FailedToStart", "The test encountered a failure on startup.");
+									Horde.ReportInvalidatingIssue(Horde.GauntletHordeStepMessage.FailedToStart, "The test encountered a failure on startup.");
 									NodeInfo.Result = TestExecutionInfo.ExecutionResult.Interrupted;
 									CompletedTests.Add(NodeInfo);
 								}
@@ -309,7 +309,7 @@ namespace Gauntlet
 										PendingTests[i] = null;
 										NodeInfo.TimeSetupBegan = NodeInfo.TimeSetupEnded = NodeInfo.TimeTestEnded = DateTime.Now;
 										NodeInfo.Result = TestExecutionInfo.ExecutionResult.NotStarted;
-										Horde.ReportInvalidatingIssue("NotEnoughDevices", "The test was unable to secure enough devices to run.");
+										Horde.ReportInvalidatingIssue(Horde.GauntletHordeStepMessage.NotEnoughDevices, "The test was unable to secure enough devices to run.");
 										CompletedTests.Add(NodeInfo);
 									}
 								}
diff --git a/Engine/Source/Programs/AutomationTool/Gauntlet/Framework/Gauntlet.Utils.cs b/Engine/Source/Programs/AutomationTool/Gauntlet/Framework/Gauntlet.Utils.cs
--- a/Engine/Source/Programs/AutomationTool/Gauntlet/Framework/Gauntlet.Utils.cs
+++ b/Engine/Source/Programs/AutomationTool/Gauntlet/Framework/Gauntlet.Utils.cs
@@ -149,7 +149,7 @@ namespace Gauntlet
 					Console.CancelKeyPress += new ConsoleCancelEventHandler((obj, args) =>
 					{
 						CancelSignalled = true;
-
+						Horde.ReportInvalidatingIssue(Horde.GauntletHordeStepMessage.TestCancelled, "Test was manually cancelled");
 						// fire all abort handlers
 						foreach (var Handler in AbortHandlers)
 						{
diff --git a/Engine/Source/Programs/AutomationTool/Gauntlet/Framework/Horde/Gauntlet.Horde.cs b/Engine/Source/Programs/AutomationTool/Gauntlet/Framework/Horde/Gauntlet.Horde.cs
--- a/Engine/Source/Programs/AutomationTool/Gauntlet/Framework/Horde/Gauntlet.Horde.cs
+++ b/Engine/Source/Programs/AutomationTool/Gauntlet/Framework/Horde/Gauntlet.Horde.cs
@@ -47,6 +47,64 @@ namespace Gauntlet
 			Error
 		}
 
+		public class HordeStepMessage
+		{
+			public JobStepMessageLevels Severity { get; private set; }
+			public string Message { get; private set; }
+			public HordeStepMessage(string InMessage, JobStepMessageLevels InSeverity)
+			{
+				Message = InMessage;
+				Severity = InSeverity;
+			}
+		}
+		
+
+		public partial class GauntletHordeStepMessage
+		{
+			public static HordeStepMessage TestStepMessage = new HordeStepMessage("TestCategory", JobStepMessageLevels.Error);
+
+			// Failures unrelated to game functionality
+			public static HordeStepMessage CannotAccessBuild = new HordeStepMessage("CannotAccessBuild", JobStepMessageLevels.Error);
+			public static HordeStepMessage NotEnoughDriveSpace = new HordeStepMessage("NotEnoughDriveSpace", JobStepMessageLevels.Error);
+			public static HordeStepMessage InvalidBuildSource = new HordeStepMessage("InvalidBuildSource", JobStepMessageLevels.Error);
+			public static HordeStepMessage InvalidPerfSpec = new HordeStepMessage("InvalidPerfSpec", JobStepMessageLevels.Error);
+			public static HordeStepMessage InstallationFailed = new HordeStepMessage("InstallationFailed", JobStepMessageLevels.Error);
+			public static HordeStepMessage InvalidConfig = new HordeStepMessage("InvalidConfig", JobStepMessageLevels.Error);
+			public static HordeStepMessage NotEnoughDevices = new HordeStepMessage("NotEnoughDevices", JobStepMessageLevels.Error);
+			public static HordeStepMessage BuildNotFound = new HordeStepMessage("BuildNotFound", JobStepMessageLevels.Error);
+			public static HordeStepMessage BuildNotSpecified = new HordeStepMessage("BuildNotSpecified", JobStepMessageLevels.Error);
+			public static HordeStepMessage FailedToStart = new HordeStepMessage("FailedToStart", JobStepMessageLevels.Error);
+			public static HordeStepMessage InvalidBackend = new HordeStepMessage("InvalidBackend", JobStepMessageLevels.Error);
+			public static HordeStepMessage WrongAccountUsed = new HordeStepMessage("WrongAccountUsed", JobStepMessageLevels.Error);
+			public static HordeStepMessage TestCancelled = new HordeStepMessage("TestCancelled", JobStepMessageLevels.Error);
+
+			// Failures related to game/service functionality
+			public static HordeStepMessage LoginFailed = new HordeStepMessage("LoginFailed", JobStepMessageLevels.Error);
+			public static HordeStepMessage MatchmakingFailed = new HordeStepMessage("MatchmakingFailed", JobStepMessageLevels.Error);
+			public static HordeStepMessage TestTimeoutExceeded = new HordeStepMessage("TestTimeoutExceeded", JobStepMessageLevels.Error);
+			public static HordeStepMessage HeartbeatTimeout = new HordeStepMessage("HeartbeatTimeout", JobStepMessageLevels.Error);
+			public static HordeStepMessage OverlayException = new HordeStepMessage("OverlayException", JobStepMessageLevels.Error);
+			public static HordeStepMessage TestCasesFailed = new HordeStepMessage("TestCasesFailed", JobStepMessageLevels.Error);
+			public static HordeStepMessage MissingArtifacts = new HordeStepMessage("MissingArtifacts", JobStepMessageLevels.Error);
+			public static HordeStepMessage UnacceptablePerf = new HordeStepMessage("UnacceptablePerf", JobStepMessageLevels.Error);
+			public static HordeStepMessage UnacceptableMem = new HordeStepMessage("UnacceptableMem", JobStepMessageLevels.Error);
+			public static HordeStepMessage HardwareFailure = new HordeStepMessage("HardwareFailure", JobStepMessageLevels.Error);
+			public static HordeStepMessage ServerCrashDetected = new HordeStepMessage("ServerCrashDetected", JobStepMessageLevels.Error);
+			public static HordeStepMessage ClientCrashDetected = new HordeStepMessage("ClientCrashDetected", JobStepMessageLevels.Error);
+			public static HordeStepMessage EditorCrashDetected = new HordeStepMessage("EditorCrashDetected", JobStepMessageLevels.Error);
+
+			// Warning states for if things are getting close to bad
+			public static HordeStepMessage TestWarningsDetected = new HordeStepMessage("TestWarningsDetected", JobStepMessageLevels.Warning);
+			public static HordeStepMessage PerfWarningsDetected = new HordeStepMessage("PerfWarningsDetected", JobStepMessageLevels.Warning);
+			public static HordeStepMessage MemWarningsDetected = new HordeStepMessage("MemWarningsDetected", JobStepMessageLevels.Warning);
+			public static HordeStepMessage ReplayGettingStale = new HordeStepMessage("ReplayGettingStale", JobStepMessageLevels.Warning);
+			public static HordeStepMessage ImproperRpcCallDetected = new HordeStepMessage("ImproperRpcCallDetected", JobStepMessageLevels.Warning);
+
+			// Success states
+			public static HordeStepMessage TestSucceeded = new HordeStepMessage("TestSucceeded", JobStepMessageLevels.Info);
+			public static HordeStepMessage TestDataGenerated = new HordeStepMessage("TestDataGenerated", JobStepMessageLevels.Info);
+
+		}
 
 #nullable enable
 		private static string? curStepMessage { get; set; } = null;
@@ -59,12 +117,11 @@ namespace Gauntlet
 		/// Good message example "TestTimedOut", or "IslandNotFound", or "LoginFailed"
 		/// Bad message example "This test timed out after 3 hours", "7 of 8 islands found"
 		/// </summary>
-		/// <param name="inMessage"></param>
-		/// <param name="inSeverity"></param>
-		public static void UpdateStepMessage(string inMessage, JobStepMessageLevels inSeverity)
+		/// <param name="InStepMessage">A predefined object that describes what happened.</param>
+		public static void UpdateStepMessage(HordeStepMessage InStepMessage)
 		{
-			curStepMessage = inMessage;
-			curStepSeverity = inSeverity.ToString();
+			curStepMessage = InStepMessage.Message;
+			curStepSeverity = InStepMessage.Severity.ToString();
 		}
 
 		/// <summary>
@@ -72,20 +129,20 @@ namespace Gauntlet
 		/// be used right before an exception is thrown. If it is being used outside of a currently running test (as in, the test is actively ticking), you want to set
 		/// bGenerateSummary to true
 		/// </summary>
-		/// <param name="IssueType">Describes the type of issue (ex. NotEnoughDevices) - should be general so it can be shown in one line</param>
+		/// <param name="InMessage">A predefined object that describes what happened.</param>
 		/// <param name="IssueDescription">In-depth description of the issue. If the issue type was NotEnoughDevices, for example, this is where you would say what sort of device
 		/// you were looking for.</param>
 		/// <param name="bGenerateSummary"> Whether or not we should immediately generate the summary log. If this is being called from outside of a test node, set to yes.</param>
-		public static void ReportInvalidatingIssue(string IssueType, string IssueDescription, bool bGenerateSummary = true)
+		public static void ReportInvalidatingIssue(HordeStepMessage InMessage, string IssueDescription, bool bGenerateSummary = true)
 		{
-			string InvalidatingIssueString = $"**Invalidating issue found:** {IssueType}\n";
+			string InvalidatingIssueString = $"**Invalidating issue found:** {InMessage.Message}\n";
 			OnStepSummary += () => InvalidatingIssueString;
 			if (!string.IsNullOrEmpty(IssueDescription))
 			{
 				OnStepSummary += () => $"More details: {IssueDescription}";
 			}
 
-			UpdateStepMessage(IssueType, JobStepMessageLevels.Error);
+			UpdateStepMessage(InMessage);
 			if (bGenerateSummary)
 			{
 				GenerateSummary();
diff --git a/Engine/Source/Programs/AutomationTool/Gauntlet/Unreal/Automation/UE.BootTest.cs b/Engine/Source/Programs/AutomationTool/Gauntlet/Unreal/Automation/UE.BootTest.cs
--- a/Engine/Source/Programs/AutomationTool/Gauntlet/Unreal/Automation/UE.BootTest.cs
+++ b/Engine/Source/Programs/AutomationTool/Gauntlet/Unreal/Automation/UE.BootTest.cs
@@ -67,7 +67,6 @@ namespace UE
 			{
 				kTimeOutDuration = logIdleTimeout;
 			}
-			Gauntlet.Horde.UpdateStepMessage("DidNotFinish", Gauntlet.Horde.JobStepMessageLevels.Error);
 			return Config;
 		}
 
@@ -165,7 +164,7 @@ namespace UE
 			{
 				if (!string.IsNullOrEmpty(GetCompletionString()) && !DidDetectLaunch)
 				{
-					Gauntlet.Horde.UpdateStepMessage("LaunchNotCompleted", Gauntlet.Horde.JobStepMessageLevels.Error);
+					Gauntlet.Horde.UpdateStepMessage(Gauntlet.Horde.GauntletHordeStepMessage.FailedToStart);
 					ReportError("Failed to detect completion of launch");
 				}
 				else
@@ -178,14 +177,14 @@ namespace UE
 						if (!File.Exists(RoleResult.Artifacts.LogPath))
 						{
 							MissingLogs = true;
-							Gauntlet.Horde.UpdateStepMessage("LogsNotFound", Gauntlet.Horde.JobStepMessageLevels.Error);
+							Gauntlet.Horde.UpdateStepMessage(Gauntlet.Horde.GauntletHordeStepMessage.MissingArtifacts);
 							ReportError("No log files found for {0}. Were they not retrieved from the device?", RoleResult.Artifacts.SessionRole);
 						}
 					}
 
 					if (!MissingLogs)
 					{
-						Gauntlet.Horde.UpdateStepMessage("CreateReportSucceeded", Gauntlet.Horde.JobStepMessageLevels.Info);
+						Gauntlet.Horde.UpdateStepMessage(Gauntlet.Horde.GauntletHordeStepMessage.TestDataGenerated);
 						Log.Info("Found valid log artifacts for test");
 					}
 				}
diff --git a/Engine/Source/Programs/AutomationTool/Gauntlet/Unreal/Base/Gauntlet.UnrealSession.cs b/Engine/Source/Programs/AutomationTool/Gauntlet/Unreal/Base/Gauntlet.UnrealSession.cs
--- a/Engine/Source/Programs/AutomationTool/Gauntlet/Unreal/Base/Gauntlet.UnrealSession.cs
+++ b/Engine/Source/Programs/AutomationTool/Gauntlet/Unreal/Base/Gauntlet.UnrealSession.cs
@@ -176,7 +176,7 @@ namespace Gauntlet
 		/// <summary>
 		/// What to do if this role exits unexpectedly.
 		/// </summary>
-		public UnexpectedProcessExitDelegate OnProcessExitedUnexpectedly;
+		public OnTryToRestoreProcessDelegate OnTryToRestoreProcess;
 
 		/// <summary>
 		/// Is this role Null?
@@ -740,7 +740,7 @@ namespace Gauntlet
 			if (!CheckRolesArePossible(ref ValidationIssues))
 			{
 				ValidationIssues.ForEach(S => Log.Error(KnownLogEvents.Gauntlet_BuildDropEvent, S));
-				Horde.ReportInvalidatingIssue("BuildNotFound", "The test was unable to find the necessary executables.");
+				Horde.ReportInvalidatingIssue(Horde.GauntletHordeStepMessage.BuildNotFound, "The test was unable to find the necessary executables.");
 				throw new AutomationException("One or more issues occurred when validating build {0} against requested roles", InSource.BuildName);
 			}
 
@@ -905,7 +905,7 @@ namespace Gauntlet
 						return null;
 					}
 					// If device reservation fails, the device pool cannot support this launch.
-					Horde.ReportInvalidatingIssue("NotEnoughDevices", "The test was unable to secure enough devices to run.");
+					Horde.ReportInvalidatingIssue(Horde.GauntletHordeStepMessage.NotEnoughDevices, "The test was unable to secure enough devices to run.");
 					throw new AutomationException("Failed to acquire all devices for launch. See above for details.");
 				}
 
@@ -1525,6 +1525,7 @@ namespace Gauntlet
 							string ErrorMessage = string.Format("Encountered error setting up device {0} for role {1}. {2}", Device, Role, Ex);
 							if (ErrorMessage.Contains("not enough space") && Device.Platform == BuildHostPlatform.Current.Platform)
 							{
+								Horde.ReportInvalidatingIssue(Horde.GauntletHordeStepMessage.NotEnoughDriveSpace, ErrorMessage);
 								Log.Error(ErrorMessage);
 								// If on desktop platform, we are not retrying.
 								// It is unlikely that space is going to be made and InstallBuildParallel has marked the build path as problematic.
@@ -1578,7 +1579,7 @@ namespace Gauntlet
 
 					if (SessionRetries == 0)
 					{
-						Horde.ReportInvalidatingIssue("ApplicationInstallFailed", "Unable to install application for session.");
+						Horde.ReportInvalidatingIssue(Horde.GauntletHordeStepMessage.InvalidPerfSpec, "Unable to install application for session.");
 						throw new AutomationException("Unable to install application for session.");
 					}
 
@@ -1848,12 +1849,16 @@ namespace Gauntlet
 					{
 						// If on desktop platform, we are not retrying.
 						// It is unlikely that space is going to be made and InstallBuildParallel has marked the build path as problematic.
+
+						Horde.ReportInvalidatingIssue(Horde.GauntletHordeStepMessage.NotEnoughDriveSpace, Message);
 						Log.Error(EventType, Message);
 						throw;
 					}
 					else if (IsOverlayException(Ex))
 					{
 						// Errors with Overlay executables are caused by missing files or improper setup, inform user and exit early
+
+						Horde.ReportInvalidatingIssue(Horde.GauntletHordeStepMessage.OverlayException, Message);
 						Log.Error(EventType, "Device {DeviceName} {Message}", Device.Name, Message);
 						throw;
 					}
diff --git a/Engine/Source/Programs/AutomationTool/Gauntlet/Unreal/Base/Gauntlet.UnrealTestConfiguration.cs b/Engine/Source/Programs/AutomationTool/Gauntlet/Unreal/Base/Gauntlet.UnrealTestConfiguration.cs
--- a/Engine/Source/Programs/AutomationTool/Gauntlet/Unreal/Base/Gauntlet.UnrealTestConfiguration.cs
+++ b/Engine/Source/Programs/AutomationTool/Gauntlet/Unreal/Base/Gauntlet.UnrealTestConfiguration.cs
@@ -386,7 +386,7 @@ namespace Gauntlet
 	/// Delegate to fire when a process exits unexpectedly.
 	/// We pass in the Role Instance instead of the Role so we can get all of the instantiated data.
 	/// </summary>
-	public delegate void UnexpectedProcessExitDelegate(UnrealSessionInstance.RoleInstance InRole);
+	public delegate void OnTryToRestoreProcessDelegate(UnrealSessionInstance.RoleInstance InRole);
 
 	/// <summary>
 	/// This class represents a process-role in a test and defines the type, command line,
@@ -542,9 +542,9 @@ namespace Gauntlet
 		public bool RespawnProcessOnUnexpectedExit { get; set; }
 
 		/// <summary>
-		/// What to do if this role exits unexpectedly.
+		/// What to do if this role exits unexpectedly and we try to restore it
 		/// </summary>
-		public UnexpectedProcessExitDelegate OnProcessExitedUnexpectedly;
+		public OnTryToRestoreProcessDelegate OnTryToRestoreProcess;
 
 	}
 
diff --git a/Engine/Source/Programs/AutomationTool/Gauntlet/Unreal/Base/Gauntlet.UnrealTestNode.cs b/Engine/Source/Programs/AutomationTool/Gauntlet/Unreal/Base/Gauntlet.UnrealTestNode.cs
--- a/Engine/Source/Programs/AutomationTool/Gauntlet/Unreal/Base/Gauntlet.UnrealTestNode.cs
+++ b/Engine/Source/Programs/AutomationTool/Gauntlet/Unreal/Base/Gauntlet.UnrealTestNode.cs
@@ -644,7 +644,10 @@ namespace Gauntlet
 				{
 					return false;
 				}
-				RoleInstance.Role.OnProcessExitedUnexpectedly.Invoke(RoleInstance);
+				if (RoleInstance.Role.OnTryToRestoreProcess != null)
+				{
+					RoleInstance.Role.OnTryToRestoreProcess.Invoke(RoleInstance);
+				}
 				IAppInstall DesiredInstall = UnrealApp.RolesToInstalls[RoleInstance.Role];
 				UnrealApp.SessionInstance.RecoverProcessForRole(RoleInstance.Role, DesiredInstall);
 			}
@@ -709,6 +712,7 @@ namespace Gauntlet
 
 			if (Config == null)
 			{
+				Horde.ReportInvalidatingIssue(Horde.GauntletHordeStepMessage.InvalidConfig, "Null config returned for this test.");
 				throw new AutomationException("Test {0} returned null config!", this);
 			}
 
@@ -784,7 +788,7 @@ namespace Gauntlet
 					SessionRole.SkipCleanDeviceArtifacts = TestRole.SkipCleanDeviceArtifacts;
 					SessionRole.ArchiveDevArtifacts = TestRole.ArchiveDevArtifacts;
 					SessionRole.RespawnProcessOnUnexpectedExit = TestRole.RespawnProcessOnUnexpectedExit;
-					SessionRole.OnProcessExitedUnexpectedly = TestRole.OnProcessExitedUnexpectedly;
+					SessionRole.OnTryToRestoreProcess = TestRole.OnTryToRestoreProcess;
 
 					Log.Verbose("Created SessionRole {Role} from RoleContext {Context} (RoleType={Type})", SessionRole, RoleContext, TypesToRoles.Key);
 
@@ -833,7 +837,7 @@ namespace Gauntlet
 					SessionRole.CompressScreenshots = TestRole.CompressScreenshots;
 					SessionRole.CreateGifFromScreenshots = TestRole.CreateGifFromScreenshots;
 					SessionRole.RespawnProcessOnUnexpectedExit = TestRole.RespawnProcessOnUnexpectedExit;
-					SessionRole.OnProcessExitedUnexpectedly = TestRole.OnProcessExitedUnexpectedly;
+					SessionRole.OnTryToRestoreProcess = TestRole.OnTryToRestoreProcess;
 					SessionRoles.Add(SessionRole);
 				}
 			}
@@ -1302,6 +1306,7 @@ namespace Gauntlet
 			{
 				if (RoleResults.Any(R => R.ProcessResult == UnrealProcessResult.TimeOut))
 				{
+					Horde.UpdateStepMessage(Horde.GauntletHordeStepMessage.TestTimeoutExceeded);
 					SetTestResult(TestResult.TimedOut);
 				}
 				else
@@ -2316,6 +2321,8 @@ namespace Gauntlet
 				double SecondsSinceSessionStart = DateTime.Now.Subtract(SessionStartTime).TotalSeconds;
 				if (SecondsSinceSessionStart > HeartbeatOptions.TimeoutBeforeFirstActiveHeartbeat)
 				{
+
+					Horde.UpdateStepMessage(Horde.GauntletHordeStepMessage.HeartbeatTimeout);
 					Log.Error(KnownLogEvents.Gauntlet_TestEvent, "{Time} seconds have passed without detecting the first active Gauntlet heartbeat.", HeartbeatOptions.TimeoutBeforeFirstActiveHeartbeat);
 					MarkTestComplete();
 					SetUnrealTestResult(TestResult.TimedOut);
@@ -2328,6 +2335,7 @@ namespace Gauntlet
 				double SecondsSinceLastActiveHeartbeat = DateTime.Now.Subtract(LastActiveHeartbeatTime).TotalSeconds;
 				if (SecondsSinceLastActiveHeartbeat > HeartbeatOptions.TimeoutBetweenActiveHeartbeats)
 				{
+					Horde.UpdateStepMessage(Horde.GauntletHordeStepMessage.HeartbeatTimeout);
 					Log.Error(KnownLogEvents.Gauntlet_TestEvent, "{Time} seconds have passed without detecting any active Gauntlet heartbeats.", HeartbeatOptions.TimeoutBetweenActiveHeartbeats);
 					MarkTestComplete();
 					SetUnrealTestResult(TestResult.TimedOut);
@@ -2340,6 +2348,7 @@ namespace Gauntlet
 				double SecondsSinceLastHeartbeat = DateTime.Now.Subtract(LastHeartbeatTime).TotalSeconds;
 				if (SecondsSinceLastHeartbeat > HeartbeatOptions.TimeoutBetweenAnyHeartbeats)
 				{
+					Horde.UpdateStepMessage(Horde.GauntletHordeStepMessage.HeartbeatTimeout);
 					Log.Error(KnownLogEvents.Gauntlet_TestEvent, "{Time} seconds have passed without detecting any Gauntlet heartbeats.", HeartbeatOptions.TimeoutBetweenAnyHeartbeats);
 					MarkTestComplete();
 					SetUnrealTestResult(TestResult.TimedOut);
diff --git a/Engine/Source/Programs/AutomationTool/Gauntlet/Unreal/RunUnreal.cs b/Engine/Source/Programs/AutomationTool/Gauntlet/Unreal/RunUnreal.cs
--- a/Engine/Source/Programs/AutomationTool/Gauntlet/Unreal/RunUnreal.cs
+++ b/Engine/Source/Programs/AutomationTool/Gauntlet/Unreal/RunUnreal.cs
@@ -120,7 +120,7 @@ namespace Gauntlet
 			if (ParseParam("forcestepfail"))
 			{
 				string ErrorMessage = "-forcestepfail specified on commandline.";
-				Horde.ReportInvalidatingIssue("TestCategory", ErrorMessage);
+				Horde.ReportInvalidatingIssue(Horde.GauntletHordeStepMessage.TestStepMessage, ErrorMessage);
 				throw new AutomationException(ErrorMessage);
 			}
 			if (ParseParam("log"))
@@ -142,14 +142,14 @@ namespace Gauntlet
 			if (string.IsNullOrEmpty(ContextOptions.Build))
 			{
 				string ErrorMessage = "No build specified. Use -build=p:\\path\\to\\build";
-				Horde.ReportInvalidatingIssue("BuildNotSpecified", ErrorMessage);
+				Horde.ReportInvalidatingIssue(Horde.GauntletHordeStepMessage.BuildNotSpecified, ErrorMessage);
 				throw new AutomationException(ErrorMessage);
 			}
 
 			if (typeof(UnrealBuildSource).IsAssignableFrom(ContextOptions.BuildSourceType) == false)
 			{
 				string ErrorMessage = "Provided BuildSource type does not inherit from UnrealBuildSource";
-				Horde.ReportInvalidatingIssue("InvalidBuildSource", ErrorMessage);
+				Horde.ReportInvalidatingIssue(Horde.GauntletHordeStepMessage.InvalidBuildSource, ErrorMessage);
 				throw new AutomationException(ErrorMessage);
 			}
 
@@ -435,7 +435,7 @@ namespace Gauntlet
 				if (!Enum.TryParse<EPerfSpec>(PerfSpecArg, true, out PerfSpec))
 				{
 					string ErrorMessage = string.Format("Unable to convert perfspec '{0}' into an EPerfSpec", PerfSpec);
-					Horde.ReportInvalidatingIssue("InvalidPerfSpec", ErrorMessage);
+					Horde.ReportInvalidatingIssue(Horde.GauntletHordeStepMessage.InvalidPerfSpec, ErrorMessage);
 					throw new AutomationException(ErrorMessage);
 				}
 
