Add WindowsMSG label

diff --git a/Engine/Source/Programs/AutomationTool/AutomationUtils/MCPPublic.cs b/Engine/Source/Programs/AutomationTool/AutomationUtils/MCPPublic.cs
--- a/Engine/Source/Programs/AutomationTool/AutomationUtils/MCPPublic.cs
+++ b/Engine/Source/Programs/AutomationTool/AutomationUtils/MCPPublic.cs
@@ -260,6 +260,11 @@ namespace EpicGames.MCP.Automation
 		/// WindowsLuna Platform.
 		/// </summary>
 		public static MCPPlatform WindowsLuna = FindOrAddByName("WindowsLuna");
+		
+		/// <summary>
+		/// WindowsMSG Platform.
+		/// </summary>
+		public static MCPPlatform WindowsMSG = FindOrAddByName("WindowsMSG");
 
 		/// <summary>
 		/// IOSCN Platform.
