Localization:
Upgrade Smartling localization files failing to download/upload from log to warning.

diff --git a/Engine/Source/Programs/AutomationTool/SmartlingLocalization/SmartlingLocalizationProvider.cs b/Engine/Source/Programs/AutomationTool/SmartlingLocalization/SmartlingLocalizationProvider.cs
--- a/Engine/Source/Programs/AutomationTool/SmartlingLocalization/SmartlingLocalizationProvider.cs
+++ b/Engine/Source/Programs/AutomationTool/SmartlingLocalization/SmartlingLocalizationProvider.cs
@@ -121,7 +121,14 @@ namespace EpicGames.SmartlingLocalization
 			{
 				foreach (string Log in Logs)
 				{
-					Logger.LogInformation(Log);
+					if (Log.StartsWith("[FAILED]"))
+					{
+						Logger.LogWarning(Log);
+					}
+					else
+					{
+						Logger.LogInformation(Log);
+					}
 				}
 			}
 		}
@@ -313,7 +320,15 @@ namespace EpicGames.SmartlingLocalization
 			{
 				foreach (string Log in Logs)
 				{
-					Logger.LogInformation(Log);
+					if (Log.StartsWith("[FAILED]"))
+					{
+						Logger.LogWarning(Log);
+					}
+					else
+					{
+						Logger.LogInformation(Log);
+					}
+					
 				}
 			}
 
