Gauntlet default output to Horde Artifacts instead of P drive

#Synced-CL 48519518

diff --git a/Engine/Source/Programs/AutomationTool/Gauntlet/Unreal/Game/EpicGame.TestConfig.cs b/Engine/Source/Programs/AutomationTool/Gauntlet/Unreal/Game/EpicGame.TestConfig.cs
--- a/Engine/Source/Programs/AutomationTool/Gauntlet/Unreal/Game/EpicGame.TestConfig.cs
+++ b/Engine/Source/Programs/AutomationTool/Gauntlet/Unreal/Game/EpicGame.TestConfig.cs
@@ -542,23 +542,7 @@ namespace EpicGame
 					continue;
 				}
 
-				var Filename = Path.GetFileName(SourceTracePath);
-				if (CommandUtils.IsBuildMachine && Gauntlet.Horde.IsHordeJob)
-				{
-					string HordeTracePath = Path.GetFullPath(Path.Combine(CommandUtils.CmdEnv.LogFolder, OutputFolderName, Filename));
-					Log.Verbose($"Copying trace file for Horde to {HordeTracePath}");
-					try
-					{
-						Directory.CreateDirectory(Path.GetDirectoryName(HordeTracePath));
-						File.Copy(SourceTracePath, HordeTracePath, true);
-					}
-					catch (Exception Ex)
-					{
-						Log.Warning($"Failed to copy {SourceTracePath} to Horde Artifacts. Exception: {Ex}");
-					}
-				}
-
-				string OutputTracePath = Path.Combine(OutputPath, Filename);
+				string OutputTracePath = Path.Combine(OutputPath, Path.GetFileName(SourceTracePath));
 				Log.Verbose($"Moving trace file to {OutputTracePath}");
 				try
 				{
