Include .uondemandtoc file(s) in platform packages.

#rnx

#Synced-CL 48597399

diff --git a/Engine/Source/Programs/AutomationTool/Gauntlet/Unreal/Automation/UE.ZenLoaderTest.cs b/Engine/Source/Programs/AutomationTool/Gauntlet/Unreal/Automation/UE.ZenLoaderTest.cs
--- a/Engine/Source/Programs/AutomationTool/Gauntlet/Unreal/Automation/UE.ZenLoaderTest.cs
+++ b/Engine/Source/Programs/AutomationTool/Gauntlet/Unreal/Automation/UE.ZenLoaderTest.cs
@@ -74,7 +74,7 @@ namespace UE
 			ActualStagingOutputConditions =
 			[
 				new(StagingOutput.Utoc, () =>
-					DirectoryUtils.FindMatchingFiles(ContentPath, @"^.*\.(ucas|utoc)$", -1).Any()),
+					DirectoryUtils.FindMatchingFiles(ContentPath, @"^.*\.(ucas|utoc|uondemandtoc)$", -1).Any()),
 				new(StagingOutput.LegacyPak, () =>
 					DirectoryUtils.FindMatchingFiles(ContentPath, @"^.*\.pak$", -1).Any()),
 				new(StagingOutput.Streaming, () =>
