Localization:
Introduce the [WARNING] tag to the logs for Smartling Provider. allows failures in uploads/downloads to be filtered out and reported in build machines.

diff --git a/Engine/Source/Programs/AutomationTool/SmartlingLocalization/SmartlingLocalizationProvider.cs b/Engine/Source/Programs/AutomationTool/SmartlingLocalization/SmartlingLocalizationProvider.cs
--- a/Engine/Source/Programs/AutomationTool/SmartlingLocalization/SmartlingLocalizationProvider.cs
+++ b/Engine/Source/Programs/AutomationTool/SmartlingLocalization/SmartlingLocalizationProvider.cs
@@ -121,7 +121,7 @@ namespace EpicGames.SmartlingLocalization
 			{
 				foreach (string Log in Logs)
 				{
-					if (Log.StartsWith("[FAILED]"))
+					if (Log.StartsWith("[WARNING]"))
 					{
 						Logger.LogWarning(Log);
 					}
@@ -184,7 +184,7 @@ namespace EpicGames.SmartlingLocalization
 						if (CurrentTries > MaxTries)
 						{
 							DownloadStopWatch.Stop();
-							Logs.Add($"[FAILED] Downloading: '{ExportFile.FullName}' ({EpicLocale}) exhausted all retries after {DownloadStopWatch.ElapsedMilliseconds / 1000}.");
+							Logs.Add($"[WARNING] Failed in downloading: '{ExportFile.FullName}' ({EpicLocale}) exhausted all retries after {DownloadStopWatch.ElapsedMilliseconds / 1000}.");
 							return Logs;
 						}
 						Logs.Add($"Encountered HTTP Status Code 401. Authentication most likely expired. Retrying {CurrentTries}/{MaxTries} times with refreshed authentication token.");
@@ -199,7 +199,7 @@ namespace EpicGames.SmartlingLocalization
 					if (CurrentTries > MaxTries)
 					{
 						DownloadStopWatch.Stop();
-						Logs.Add($"[FAILED] Downloading: '{ExportFile.FullName}' ({EpicLocale}) exhausted all retries in {DownloadStopWatch.ElapsedMilliseconds / 1000.0:F2} seconds. - {Ex}");
+						Logs.Add($"[WARNING] Failed in downloading: '{ExportFile.FullName}' ({EpicLocale}) exhausted all retries in {DownloadStopWatch.ElapsedMilliseconds / 1000.0:F2} seconds. - {Ex}");
 						return Logs;
 					}
 					Logs.Add($"Failed to get download response. Retrying {CurrentTries}/{MaxTries} times.");
@@ -320,7 +320,7 @@ namespace EpicGames.SmartlingLocalization
 			{
 				foreach (string Log in Logs)
 				{
-					if (Log.StartsWith("[FAILED]"))
+					if (Log.StartsWith("[WARNING]"))
 					{
 						Logger.LogWarning(Log);
 					}
@@ -395,7 +395,7 @@ namespace EpicGames.SmartlingLocalization
 						++CurrentTries;
 						if (CurrentTries > MaxTries)
 						{
-							Logs.Add($"[FAILED] Uploading: '{FileToUpload.FullName}' ({EpicLocale}). Exhausted all retries.");
+							Logs.Add($"[WARNING] Failed in uploading: '{FileToUpload.FullName}' ({EpicLocale}). Exhausted all retries.");
 							return Logs;
 						}
 						Logs.Add($"Encountered HTTP Status Code 401. Authentication token most likely expired. Retrying {CurrentTries}/{MaxTries} times with refreshed authentication token.");
@@ -405,7 +405,7 @@ namespace EpicGames.SmartlingLocalization
 				}
 				catch (Exception Ex)
 				{
-					Logs.Add($"[FAILED] Uploading: '{FileToUpload.FullName}' ({EpicLocale}) - {Ex}");
+					Logs.Add($"[WARNING] Failed in uploading: '{FileToUpload.FullName}' ({EpicLocale}) - {Ex}");
 					return Logs;
 				}
 			}
@@ -416,7 +416,7 @@ namespace EpicGames.SmartlingLocalization
 			}
 			else
 			{
-				Logs.Add($"[FAILED] Uploading: '{FileToUpload.FullName}' ({EpicLocale})");
+				Logs.Add($"[WARNING] Failed in uploading: '{FileToUpload.FullName}' ({EpicLocale})");
 				await AppendRequestErrorsToLog(UploadResponse, Logs);
 			}
 			return Logs;
