Changed a warning in SmartlingLocalizationProvider to avoid spamming warnings when uploading empty localization files to Smartling.

#Synced-CL 48764915

diff --git a/Engine/Source/Programs/AutomationTool/SmartlingLocalization/SmartlingLocalizationProvider.cs b/Engine/Source/Programs/AutomationTool/SmartlingLocalization/SmartlingLocalizationProvider.cs
--- a/Engine/Source/Programs/AutomationTool/SmartlingLocalization/SmartlingLocalizationProvider.cs
+++ b/Engine/Source/Programs/AutomationTool/SmartlingLocalization/SmartlingLocalizationProvider.cs
@@ -416,7 +416,7 @@ namespace EpicGames.SmartlingLocalization
 			}
 			else
 			{
-				Logs.Add($"[WARNING] Failed in uploading: '{FileToUpload.FullName}' ({EpicLocale})");
+				Logs.Add($"[FAILED] Failed in uploading: '{FileToUpload.FullName}' ({EpicLocale})");
 				await AppendRequestErrorsToLog(UploadResponse, Logs);
 			}
 			return Logs;
