Improve updating WhatsNew in App Store Connect

#rnx

diff --git a/Engine/Source/Programs/AutomationTool/Apple/ExportAppleBuild.Automation.cs b/Engine/Source/Programs/AutomationTool/Apple/ExportAppleBuild.Automation.cs
--- a/Engine/Source/Programs/AutomationTool/Apple/ExportAppleBuild.Automation.cs
+++ b/Engine/Source/Programs/AutomationTool/Apple/ExportAppleBuild.Automation.cs
@@ -399,8 +399,10 @@ class ExportAppleBuild : BuildCommand
 
 						Logger.LogInformation("{Output}", Output);
 
-						// update the XcArchive location to the resigned one
-						XcArchive = new DirectoryReference(FixedXcArchive);
+						// replace the xcarchive with the fixed version
+						Logger.LogInformation("Overwriting xcarchive: {XcArchive} with updated version", XcArchive.FullName);
+						Directory.Delete(XcArchive.FullName, true);
+						Directory.Move(FixedXcArchive, XcArchive.FullName);
 
 						// we know that this xcarchive has a version that is not in the appstore, so return true 
 					}
@@ -418,7 +420,6 @@ class ExportAppleBuild : BuildCommand
 					}
 
 				}
-
 				// if we get here, there is either no newer version, or we have successfully re-versioned the .app
 				return true;
 			}
