Re-enable appium for bypassing apple sign in screen
#Synced-CL 48668367

diff --git a/Engine/Source/Programs/AutomationTool/Gauntlet/Platform/IOS/Gauntlet.IOSAppium.cs b/Engine/Source/Programs/AutomationTool/Gauntlet/Platform/IOS/Gauntlet.IOSAppium.cs
--- a/Engine/Source/Programs/AutomationTool/Gauntlet/Platform/IOS/Gauntlet.IOSAppium.cs
+++ b/Engine/Source/Programs/AutomationTool/Gauntlet/Platform/IOS/Gauntlet.IOSAppium.cs
@@ -381,8 +381,7 @@ namespace Gauntlet
 
 		public void CancelAppleSignOnAlert(string appID, int maxRetries = 5)
 		{
-			var WebDriverTimeout = TimeSpan.FromMilliseconds(5000);
-			var WebDriverWait = new WebDriverWait(new SystemClock(), Driver, WebDriverTimeout, TimeSpan.FromMilliseconds(250));
+			var WebDriverWait = new WebDriverWait(new SystemClock(), Driver, TimeSpan.FromMilliseconds(5000), TimeSpan.FromMilliseconds(250));
 			for (int attempt = 0; attempt < maxRetries; attempt++)
 			{
 				Log.Verbose("starting alert attempt #{0}", attempt);
@@ -407,18 +406,16 @@ namespace Gauntlet
 					}
 
 					alert.Dismiss();
-					Log.Verbose("Alert Dismissed");
+					Log.Info("Appium has Dismissed an Alert.");
 					Driver.ActivateApp(appID);
 
 					return;
 				}
 				catch (WebDriverTimeoutException)
 				{
-					Log.Verbose("Did not find alert");
 				}
 				catch (UnhandledAlertException)
 				{
-					Log.Verbose("Did not find alert");
 				}
 				catch (Exception ex)
 				{
diff --git a/Engine/Source/Programs/AutomationTool/Gauntlet/Platform/IOS/Gauntlet.TargetDeviceIOS.cs b/Engine/Source/Programs/AutomationTool/Gauntlet/Platform/IOS/Gauntlet.TargetDeviceIOS.cs
--- a/Engine/Source/Programs/AutomationTool/Gauntlet/Platform/IOS/Gauntlet.TargetDeviceIOS.cs
+++ b/Engine/Source/Programs/AutomationTool/Gauntlet/Platform/IOS/Gauntlet.TargetDeviceIOS.cs
@@ -354,6 +354,12 @@ namespace Gauntlet
 
 		public void FullClean()
 		{
+			List<string> AppAllowList = new List<string>
+			{
+				"WebDriverAgent.xctrunner",
+				"com.mobilelabsinc.deviceControl.xctrunner",
+				"com.mobilelabsinc.deviceControl"
+			};
 			string AppOutput;
 			if (UseDeviceCtl)
 			{
@@ -380,7 +386,8 @@ namespace Gauntlet
 			IEnumerable<string> AppList = AppRegex.Matches(AppOutput)
 				.Where(Match => Match.Success)
 				.Select(Match => Match.Value)
-				.Where(Match => !Match.Contains("com.apple.", StringComparison.OrdinalIgnoreCase));
+				.Where(Match => !Match.Contains("com.apple.", StringComparison.OrdinalIgnoreCase))
+				.Where(Match => !AppAllowList.Contains(Match));
 
 			foreach (string App in AppList)
 			{
@@ -833,6 +840,16 @@ namespace Gauntlet
 				throw new DeviceException("Failed to launch on {0}. {1}", Name, AppProcess.Output);
 			}
 
+			if (AppiumContainer != null)
+			{
+				lock (Mutex)
+				{
+					// Making sure app is in foreground and no alerts present. 
+					AppiumContainer.VerifyAppInForeground(IOSInstall.PackageName);
+					AppiumContainer.CancelAppleSignOnAlert(IOSInstall.PackageName);
+				}
+			}
+
 			return new IOSAppInstance(IOSInstall, AppProcess, IOSInstall.CommandLine);
 		}
 
