PR #12187: Make UpdateLocalVersion commandlet write P4 CompatibleChangelist into Build.version

#Synced-CL 48862428

diff --git a/Engine/Source/Programs/AutomationTool/Scripts/UpdateLocalVersion.Automation.cs b/Engine/Source/Programs/AutomationTool/Scripts/UpdateLocalVersion.Automation.cs
--- a/Engine/Source/Programs/AutomationTool/Scripts/UpdateLocalVersion.Automation.cs
+++ b/Engine/Source/Programs/AutomationTool/Scripts/UpdateLocalVersion.Automation.cs
@@ -21,6 +21,11 @@ public class UpdateLocalVersion : BuildCommand
 		UnrealBuild UnrealBuild = new UnrealBuild(this);
 		int? ChangelistOverride = ParseParamNullableInt("cl");
 		int? CompatibleChangelistOverride = ParseParamNullableInt("compatiblecl");
+  		// If override is not provided, try to read compatible CL value from p4 if enabled
+		if (!CompatibleChangelistOverride.HasValue && P4Enabled)
+		{
+			CompatibleChangelistOverride = P4Env.CodeChangelist;
+		}
 		string Build = ParseParamValue("Build", null);
 		UnrealBuild.UpdateVersionFiles(ChangelistNumberOverride: ChangelistOverride, CompatibleChangelistNumberOverride: CompatibleChangelistOverride, Build: Build);
 	}
