Remove the unused third-party Newtonsoft.Json dependency from UBT. Fix up assemblies that were relying on a transitive dependency from UBT.

#Synced-CL 48900888

diff --git a/Engine/Source/Programs/AutomationTool/Android/AndroidFileClient.cs b/Engine/Source/Programs/AutomationTool/Android/AndroidFileClient.cs
--- a/Engine/Source/Programs/AutomationTool/Android/AndroidFileClient.cs
+++ b/Engine/Source/Programs/AutomationTool/Android/AndroidFileClient.cs
@@ -18,7 +18,6 @@ using AutomationScripts;
 using System.Net;
 using System.Net.Sockets;
 using System.IO.Compression;
-using Newtonsoft.Json.Linq;
 using Microsoft.Extensions.Logging;
 
 using static AutomationTool.CommandUtils;
diff --git a/Engine/Source/Programs/AutomationTool/Gauntlet/Gauntlet.Automation.csproj b/Engine/Source/Programs/AutomationTool/Gauntlet/Gauntlet.Automation.csproj
--- a/Engine/Source/Programs/AutomationTool/Gauntlet/Gauntlet.Automation.csproj
+++ b/Engine/Source/Programs/AutomationTool/Gauntlet/Gauntlet.Automation.csproj
@@ -39,10 +39,6 @@
       <HintPath>..\..\..\..\Binaries\DotNET\Ionic.Zip.Reduced.dll</HintPath>
       <Private>False</Private>
     </Reference>
-    <Reference Include="Newtonsoft.JSON">
-      <HintPath>..\..\..\..\Binaries\ThirdParty\Newtonsoft\Newtonsoft.Json.dll</HintPath>
-      <Private>True</Private>
-    </Reference>
   </ItemGroup>
 
     <!-- The following will glob all platform extensions' Gauntlet source
@@ -95,5 +91,6 @@
     <PackageReference Include="Magick.NET-Q16-HDRI-AnyCPU" Version="14.9.1" PrivateAssets="all" />
     <PackageReference Include="Microsoft.CSharp" Version="4.7.0" />
     <PackageReference Include="MySql.Data" Version="6.10.9" PrivateAssets="all" />
+    <PackageReference Include="Newtonsoft.Json" Version="13.0.4" PrivateAssets="all" />
   </ItemGroup>
 </Project>
\ No newline at end of file
diff --git a/Engine/Source/Programs/AutomationTool/Turnkey/Turnkey.Automation.csproj b/Engine/Source/Programs/AutomationTool/Turnkey/Turnkey.Automation.csproj
--- a/Engine/Source/Programs/AutomationTool/Turnkey/Turnkey.Automation.csproj
+++ b/Engine/Source/Programs/AutomationTool/Turnkey/Turnkey.Automation.csproj
@@ -55,8 +55,10 @@
     <ProjectReference Include="..\Gauntlet\Gauntlet.Automation.csproj" PrivateAssets="All"><Private>false</Private></ProjectReference>
   </ItemGroup>
   <ItemGroup>
+    <PackageReference Include="Google.Apis.Drive.v3" Version="1.33.1.1232" PrivateAssets="all" />
     <PackageReference Include="Microsoft.CSharp" Version="4.7.0" />
+    <!-- Unused, but required to resolve version conflicts from transitive dependencies. -->
+    <PackageReference Include="Newtonsoft.Json" Version="13.0.4" PrivateAssets="all" />
     <PackageReference Include="System.IO.FileSystem.AccessControl" Version="4.7.0" />
-    <PackageReference Include="Google.Apis.Drive.v3" Version="1.33.1.1232" PrivateAssets="all" />
   </ItemGroup>
-</Project>
\ No newline at end of file
+</Project>
diff --git a/Engine/Source/Programs/UnrealBuildTool/UnrealBuildTool.csproj b/Engine/Source/Programs/UnrealBuildTool/UnrealBuildTool.csproj
--- a/Engine/Source/Programs/UnrealBuildTool/UnrealBuildTool.csproj
+++ b/Engine/Source/Programs/UnrealBuildTool/UnrealBuildTool.csproj
@@ -109,10 +109,6 @@
     <!-- Required by Ionic -->
     <PackageReference Include="System.Security.Permissions" Version="8.0.0" />
     <PackageReference Include="System.Text.Encoding.CodePages" Version="8.0.0" />
-
-    <!-- Added as a temporary workaround to running from visual studio where the system version of dotnet is run. 
-         Adding a loader context resolves many of the issues, but it still fails when MSBuild tries to write the json -->
-    <PackageReference Include="Newtonsoft.Json" Version="13.0.4" />
   </ItemGroup>
 
   <!-- Various batch files and shell scripts that build or run UnrealBuildTool -->
