This change fixes the issue where dotnet FileInfo CopyTo() function does not overide a file with the latest version.
This change is to enable the same behaviour that robocopy on windows exhibits while running dotnet CopyTo() on other platforms.

#rnx

diff --git a/Engine/Source/Programs/AutomationTool/BuildGraph/TempStorage.cs b/Engine/Source/Programs/AutomationTool/BuildGraph/TempStorage.cs
--- a/Engine/Source/Programs/AutomationTool/BuildGraph/TempStorage.cs
+++ b/Engine/Source/Programs/AutomationTool/BuildGraph/TempStorage.cs
@@ -1428,7 +1428,7 @@ namespace AutomationTool
 			foreach (FileInfo file in dir.GetFiles())
 			{
 				string targetFilePath = Path.Combine(destinationDir.FullName, file.Name);
-				file.CopyTo(targetFilePath);
+				file.CopyTo(targetFilePath, true);
 			}
 
 			foreach (DirectoryInfo subDir in dirs)
