This change fixes the issue where SetupBuildEnvironment() is not assigning a value to DotnetMsbuildPath, causing a NullReferenceException when constructing the dotnet build command when running under wine.

This change enabled dotnet builds to succeed when running on Wine.

#rnx

diff --git a/Engine/Source/Programs/AutomationTool/AutomationUtils/CommandEnvironment.cs b/Engine/Source/Programs/AutomationTool/AutomationUtils/CommandEnvironment.cs
--- a/Engine/Source/Programs/AutomationTool/AutomationUtils/CommandEnvironment.cs
+++ b/Engine/Source/Programs/AutomationTool/AutomationUtils/CommandEnvironment.cs
@@ -207,13 +207,6 @@ namespace AutomationTool
 			// Assume we have the capability co compile.
 			HasCapabilityToCompile = true;
 
-			if (RuntimePlatform.IsRunningOnWine())
-			{
-				// Just set an empty path as we currently compile .NET Framework/Core dependencies outside Wine
-				FrameworkMsbuildPath = "";
-				return;
-			}
-
 			try
 			{
 				FrameworkMsbuildPath = HostPlatform.Current.GetFrameworkMsbuildExe();
