BuildAcquire UAT command for downloading builds from cloud storage and syncing UGS to the build's CL

#rnx
#Synced-CL 48943714

diff --git a/Engine/Source/Programs/AutomationTool/Scripts/BuildAcquire.Automation.cs b/Engine/Source/Programs/AutomationTool/Scripts/BuildAcquire.Automation.cs
new file mode 100644
--- /dev/null
+++ b/Engine/Source/Programs/AutomationTool/Scripts/BuildAcquire.Automation.cs
@@ -0,0 +1,338 @@
+// Copyright Epic Games, Inc. All Rights Reserved.
+
+using AutomationTool;
+using AutomationUtils;
+using EpicGames.Core;
+using EpicGames.ProjectStore;
+using EpicGames.Serialization;
+using Microsoft.Extensions.DependencyInjection;
+using Microsoft.Extensions.Logging;
+using System;
+using System.Collections.Generic;
+using System.ComponentModel;
+using System.Linq;
+using System.Reflection;
+using System.Diagnostics;
+using UnrealBuildBase;
+using EpicGames.UHT.Types;
+using System.Text.RegularExpressions;
+using Microsoft.Extensions.FileSystemGlobbing.Internal.PathSegments;
+
+class BuildInfo
+{
+	public string BuildVariant = null;
+	public string Platform = null;
+}
+
+public static class EnumExtensions
+{
+	public static EnumDescription ConvertEnumWithDescription<T>() where T : struct, IConvertible
+	{
+		if (!typeof(T).IsEnum)
+		{
+			throw new ArgumentException("Type T must be an Enum");
+		}
+
+		var enumType = typeof(T).Name;
+		var valueDescriptions = Enum.GetValues(typeof(T))
+			.Cast<Enum>()
+			.ToDictionary(Convert.ToInt32, GetEnumDescription);
+		var descriptionValues = Enum.GetValues(typeof(T))
+			.Cast<Enum>()
+			.Select(x => GetEnumDescription(x)).ToList();
+
+		return new EnumDescription
+		{
+			Type = enumType,
+			ValueDescriptions = valueDescriptions,
+			Descriptions = descriptionValues
+		};
+	}
+
+	public static string GetEnumDescription(Enum value)
+	{
+		FieldInfo fi = value.GetType().GetField(value.ToString());
+		DescriptionAttribute[] attributes = (DescriptionAttribute[])fi.GetCustomAttributes(typeof(DescriptionAttribute), false);
+
+		return attributes.Length > 0 ? attributes[0].Description : string.Empty;
+	}
+
+	public static string GetEnumNameOrDescription(Enum value)
+	{
+		var Description = GetEnumDescription(value);
+		return Description == string.Empty ? value.ToString() : Description;
+	}
+}
+
+public class EnumDescription
+{
+	public string Type { get; set; }
+	public IDictionary<int, string> ValueDescriptions { get; set; }
+	public IList<string> Descriptions { get; set; }
+}
+
+class BuildAcquire : BuildCommand, IProjectParamsHelpers
+{
+	public enum BuildVariant
+	{
+		[Description("Development")]
+		Development,
+		[Description("Debug")]
+		Debug,
+		[Description("Test")]
+		Test,
+		[Description("Shipping")]
+		Shipping
+	}
+
+	public enum BuildTypeEnum
+	{
+		[Description("staged-build")]
+		Staged,
+		[Description("packaged")]
+		Packaged
+	}
+
+	public BuildAcquire()
+	{
+		CloudStorageQueryBuilder query = new CloudStorageQueryBuilder();
+		query.Equals("BuildGroup", 1);
+		CbObject obj = query.Done();
+	}
+
+	public override void ExecuteBuild()
+	{
+		Initialize();
+
+		var Params = new ProjectParams
+			(
+				Command: this,
+				RawProjectPath: ProjectPath
+			);
+		DownloadBuilds(Params);
+
+		if (SyncUGS)
+		{
+			if (!IsEditorRunning(Params))
+			{
+				CommandUtils.Run(UGSExePath.ToString(), $"sync {Changelist}");
+			}
+			else
+			{
+				throw new AutomationException("The editor is running, cannot sync using UGS. Close the editor and try again");
+			}
+		}
+	}
+
+	private void Initialize()
+	{
+		ProjectPath = ParseProjectParam();
+		if (ProjectPath == null)
+		{
+			throw new AutomationException("No project file specified. Use -project=<project>.");
+		}
+
+		if (!Enum.TryParse(ParseOptionalStringParam("BuildConfigType"), out BuildConfigType))
+		{
+			BuildConfigType = BuildVariant.Development;
+		}
+
+		SyncUGS = ParseParam("SyncUGS");
+		DownloadClient = ParseParam("DownloadClient");
+		DownloadServer = ParseParam("DownloadServer");
+		DownloadGame = ParseParam("DownloadGame");
+
+		Stream = ParseParamValue("Stream");
+		if (Stream == null)
+		{
+			throw new AutomationException("No stream specified, cannot continue. Use -Stream=<Stream>");
+		}
+
+		Changelist = ParseParamValue("CL");
+		if (Changelist == null)
+		{
+			throw new AutomationException("No CL number specified. Use -CL=<changelist>");
+		}
+
+		if (!(DownloadClient || DownloadServer || DownloadGame) && !ParseParam("BuildPlatform"))
+		{
+			throw new AutomationException("No build variant provided. Use -DownloadGame|-DownloadServer|-DownloadClient or -BuildPlatform");
+		}
+
+		UGSExePath = FileReference.Combine(new DirectoryReference(Environment.GetFolderPath(Environment.SpecialFolder.LocalApplicationData)), "UnrealGameSync", "Latest", "ugs.exe");
+		if (SyncUGS)
+		{
+			if (!FileReference.Exists(UGSExePath))
+			{
+				throw new AutomationException("UGS not found at {}, cannot continue", UGSExePath.ToString());
+			}
+		}
+
+		CloudStorageService = CommandUtils.ServiceScope.ServiceProvider.GetService<ICloudStorage>();
+		if (CloudStorageService == null)
+		{
+			throw new NotImplementedException("Cannot retrieve cloud storage from service provider, unable to continue");
+		}
+
+		CloudStorageService.Reconfigure(options =>
+		{
+			options.ReadFromIni(ProjectPath);
+			options.EchoZenCliToStdout = true;
+		});
+
+		CloudStorageService.LoginAsync(unattended: true);
+	}
+
+	private List<FindBuildResponse> FilterBuilds(in List<FindBuildResponse> foundBuilds, string CookPlatform)
+	{
+		List<FindBuildResponse> Builds = foundBuilds
+				.Where(build => !build.Name.Contains("-zen-", StringComparison.OrdinalIgnoreCase))
+				.Where(build => CookPlatform.Equals(build.Metadata.Find("cookPlatform").AsString(), StringComparison.OrdinalIgnoreCase)).ToList();
+
+		return Builds;
+	}
+
+	private void DownloadBuilds(ProjectParams Params)
+	{
+		CloudStorageQueryBuilder query = new CloudStorageQueryBuilder();
+		query.Equals("type", BuildType);
+		query.Equals("commit", Changelist);
+		query.Equals("isPreflight", "false");
+
+		string[] BuildPlatforms = ParseOptionalStringParam("BuildPlatform")?.Split('+');
+
+		CbObject obj = query.Done();
+
+		List<FindBuildResponse> foundBuilds = CloudStorageService.FindBuilds(ProjectPath.GetFileNameWithoutExtension(), BuildType, Stream, obj).ToList();
+        List<FindBuildResponse> BuildsToDownload = new();
+		
+		if (foundBuilds.Count > 0)
+		{
+			if (BuildPlatforms != null && BuildPlatforms.Length > 0 && ParseParam("GetBuild"))
+			{
+				foreach (string BuildPlatform in BuildPlatforms)
+				{
+					BuildsToDownload.Add(foundBuilds.Where(Build => Build.Name.Equals(BuildPlatform, StringComparison.OrdinalIgnoreCase)).First());
+				}
+			}
+			else
+			{
+				var ServerPlatformStrings = Params.ServerTargetPlatformInstances.Select(platform => platform.GetCookPlatform(true, false)).ToList();
+				var ClientPlatformStrings = Params.ClientTargetPlatformInstances.Select(platform => platform.GetCookPlatform(false, true)).ToList();
+				var GamePlatformStrings = Params.ClientTargetPlatformInstances.Select(platform => platform.GetCookPlatform(false, false)).ToList();
+
+				if (DownloadServer)
+				{
+					ServerPlatformStrings.ForEach(platform => BuildsToDownload.AddRange(FilterBuilds(foundBuilds, CookPlatform: platform)));
+				}
+
+				if (DownloadClient)
+				{
+					ClientPlatformStrings.ForEach(platform => BuildsToDownload.AddRange(FilterBuilds(foundBuilds, CookPlatform: platform)));
+				}
+
+				if (DownloadGame)
+				{
+					GamePlatformStrings.ForEach(platform => BuildsToDownload.AddRange(FilterBuilds(foundBuilds, CookPlatform: platform)));
+				}
+			}
+
+			BuildsToDownload = BuildsToDownload.DistinctBy(Build => Build.Name).ToList();
+		}
+
+		if (BuildsToDownload.Count == 0)
+		{
+			Logger.LogWarning("No builds found at CL {} matching the parameters", Changelist);
+			return;
+		}
+
+        foreach (FindBuildResponse Build in BuildsToDownload)
+        {
+            Logger.LogInformation("Build to download {Name}:{BuildId}", Build.Name, Build.BuildId);
+        }
+
+        foreach (FindBuildResponse Build in BuildsToDownload)
+        {
+            DirectoryReference FinalDownloadDir = DirectoryReference.Combine(new DirectoryReference(Params.BaseStageDirectory), Build.Metadata.Find("cookPlatform").AsString());
+			CloudStorageService.DownloadBuild(Build.BucketId, Build.BuildId, FinalDownloadDir, downloadOptions: new CloudBuildDownloadOptions
+			{
+				AssumeHttp2 = true,
+				BoostWorkers = true
+			});
+        }
+    }
+
+	private bool IsEditorRunning(ProjectParams Params)
+	{
+		Process[] ProcessArr = Process.GetProcesses();
+		Regex EditorRegex = new Regex(@"UnrealEditor.*\.exe");
+		foreach (var process in ProcessArr)
+		{
+			// Skip the system and the idle process
+			if (process.Id == 0 || process.Id == 4)
+			{
+				continue;
+			}
+
+			try
+			{
+				FileReference ProcessFile = new FileReference(process.MainModule.FileName);
+				string ProcessFileName = ProcessFile.GetFileName();
+				if (EditorRegex.Match(ProcessFileName).Success)
+				{
+					return ProcessFile.IsUnderDirectory(Unreal.EngineDirectory);
+				}
+			}
+			catch (Win32Exception)
+			{
+				// A Win32 exception means that this is a system process without any Modules
+				continue;
+			}
+		}
+
+		return false;
+	}
+
+	private ICloudStorage CloudStorageService;
+
+	string _BuildType;
+	private string BuildType
+	{
+		get
+		{
+			if (_BuildType == null)
+			{
+                String CmdBuildType = ParseRequiredStringParam("BuildType");
+				if (!Enum.TryParse(CmdBuildType, true, out BuildTypeEnum BuildType))
+				{
+					EnumDescription BuildTypeDescriptions = EnumExtensions.ConvertEnumWithDescription<BuildTypeEnum>();
+					int BuildTypeIndex = BuildTypeDescriptions.Descriptions.IndexOf(CmdBuildType);
+					if (BuildTypeIndex < 0)
+					{
+						String PossibleValues = String.Join("', '", Enum.GetNames(typeof(BuildTypeEnum)));
+						throw new AutomationException("Unrecognized build type {0}. Possible build types are '{1}'.", CmdBuildType, PossibleValues);
+					}
+					else
+					{
+						_BuildType = BuildTypeDescriptions.Descriptions[BuildTypeIndex];
+					}
+				}
+				_BuildType = EnumExtensions.GetEnumNameOrDescription(BuildType);
+			}
+
+			return _BuildType;
+		}
+	}
+
+	string Stream;
+	string Changelist;
+	FileReference ProjectPath;
+	bool DownloadServer;
+	bool DownloadClient;
+	bool DownloadGame;
+	bool SyncUGS;
+
+	FileReference UGSExePath;
+	// TODO: Add support for build config once established how PL2 wants to pass in build configs. Currently it uses build platform which contains the full set: platform, config, executable type
+	BuildVariant BuildConfigType;
+}
