Add Map as session report metadata

#rnx
#Synced-CL 48947454

diff --git a/Engine/Source/Programs/AutomationTool/Gauntlet/Unreal/Base/Gauntlet.UnrealTestNode.cs b/Engine/Source/Programs/AutomationTool/Gauntlet/Unreal/Base/Gauntlet.UnrealTestNode.cs
--- a/Engine/Source/Programs/AutomationTool/Gauntlet/Unreal/Base/Gauntlet.UnrealTestNode.cs
+++ b/Engine/Source/Programs/AutomationTool/Gauntlet/Unreal/Base/Gauntlet.UnrealTestNode.cs
@@ -1916,6 +1916,15 @@ namespace Gauntlet
 					Report.SetMetadata(Entry[0], Entry[1]);
 				}
 			}
+
+			// Additional metadata based on test config
+			if (GetCachedConfiguration() is UnrealTestConfiguration TestConfig)
+			{
+				if (!string.IsNullOrEmpty(TestConfig.Map))
+				{
+					Report.SetMetadata("Map", TestConfig.Map);
+				}
+			}
 		}
 
 		void SubmitToHorde(ITestReport Report)
