Add the ability to override color of a step message, plus sending an interrupt will automatically set it to the "User1" color which can be user-defined.

diff --git a/Engine/Source/Programs/AutomationTool/Gauntlet/Framework/Horde/Gauntlet.Horde.cs b/Engine/Source/Programs/AutomationTool/Gauntlet/Framework/Horde/Gauntlet.Horde.cs
--- a/Engine/Source/Programs/AutomationTool/Gauntlet/Framework/Horde/Gauntlet.Horde.cs
+++ b/Engine/Source/Programs/AutomationTool/Gauntlet/Framework/Horde/Gauntlet.Horde.cs
@@ -109,6 +109,7 @@ namespace Gauntlet
 #nullable enable
 		private static string? curStepMessage { get; set; } = null;
 		private static string? curStepSeverity { get; set; } = null;
+		private static string? curStepColor { get; set; } = null;
 #nullable disable
 
 		/// <summary>
@@ -135,6 +136,8 @@ namespace Gauntlet
 		/// <param name="bGenerateSummary"> Whether or not we should immediately generate the summary log. If this is being called from outside of a test node, set to yes.</param>
 		public static void ReportInvalidatingIssue(HordeStepMessage InMessage, string IssueDescription, bool bGenerateSummary = true)
 		{
+			string InterruptColorOverride = "User1";
+			curStepColor = InterruptColorOverride;
 			string InvalidatingIssueString = $"**Invalidating issue found:** {InMessage.Message}\n";
 			OnStepSummary += () => InvalidatingIssueString;
 			if (!string.IsNullOrEmpty(IssueDescription))
@@ -169,6 +172,7 @@ namespace Gauntlet
 #nullable enable
 			public string? message { get; set; }
 			public string? severity { get; set; }
+			public string? color { get; set; }
 
 #nullable disable
 
@@ -274,7 +278,7 @@ namespace Gauntlet
 					{
 						DefaultIgnoreCondition = JsonIgnoreCondition.WhenWritingNull
 					};
-					JobStepReport report = new JobStepReport() { name = "Gauntlet Step Details", fileName = MarkdownFilename, message = curStepMessage, severity = curStepSeverity };
+					JobStepReport report = new JobStepReport() { name = "Gauntlet Step Details", fileName = MarkdownFilename, message = curStepMessage, severity = curStepSeverity, color = curStepColor };
 					File.WriteAllText(Path.Combine(LogFolder, "GauntletStepDetails.report.json"), JsonSerializer.Serialize(report, options));
 				}
 				bHasGeneratedStepReport = true;
