Remove spammy TestDataGenerated step message from UE BootTest

#Synced-CL 48839584

diff --git a/Engine/Source/Programs/AutomationTool/Gauntlet/Unreal/Automation/UE.BootTest.cs b/Engine/Source/Programs/AutomationTool/Gauntlet/Unreal/Automation/UE.BootTest.cs
--- a/Engine/Source/Programs/AutomationTool/Gauntlet/Unreal/Automation/UE.BootTest.cs
+++ b/Engine/Source/Programs/AutomationTool/Gauntlet/Unreal/Automation/UE.BootTest.cs
@@ -184,7 +184,6 @@ namespace UE
 
 					if (!MissingLogs)
 					{
-						Gauntlet.Horde.UpdateStepMessage(Gauntlet.Horde.GauntletHordeStepMessage.TestDataGenerated);
 						Log.Info("Found valid log artifacts for test");
 					}
 				}
