Change invocation of DerivedDataTool to include the CreateCache command that was added recently and is meant to replace the current "no command" behavior.  Ensured that the uproject file commandline argument comes first if it is provided.

#Synced-CL 48523383

diff --git a/Engine/Source/Programs/AutomationTool/BuildGraph/Tasks/DDCReplayLoadTask.cs b/Engine/Source/Programs/AutomationTool/BuildGraph/Tasks/DDCReplayLoadTask.cs
--- a/Engine/Source/Programs/AutomationTool/BuildGraph/Tasks/DDCReplayLoadTask.cs
+++ b/Engine/Source/Programs/AutomationTool/BuildGraph/Tasks/DDCReplayLoadTask.cs
@@ -2,10 +2,12 @@
 
 using System;
 using System.Collections.Generic;
+using System.Linq;
 using System.Threading.Tasks;
 using System.Xml;
 using EpicGames.Core;
 using UnrealBuildBase;
+using UnrealBuildTool;
 
 namespace AutomationTool.Tasks
 {
@@ -20,6 +22,12 @@ namespace AutomationTool.Tasks
 		[TaskParameter]
 		public string Exe { get; set; }
 
+		/// <summary>
+		/// The project to run the tool with.
+		/// </summary>
+		[TaskParameter(Optional = true, ValidationType = TaskParameterValidationType.FileSpec)]
+		public string Project { get; set; }
+
 		/// <summary>
 		/// The DDC replay load file or tag set to use.
 		/// </summary>
@@ -91,10 +99,29 @@ namespace AutomationTool.Tasks
 		/// <param name="tagNameToFileSet">Mapping from tag names to the set of files they include</param>
 		public override async Task ExecuteAsync(JobContext job, HashSet<FileReference> buildProducts, Dictionary<string, HashSet<FileReference>> tagNameToFileSet)
 		{
+			FileReference projectFile = null;
+			if (!String.IsNullOrEmpty(_parameters.Project))
+			{
+				if (_parameters.Project.EndsWith(".uproject", StringComparison.OrdinalIgnoreCase))
+				{
+					projectFile = ResolveFile(_parameters.Project);
+				}
+				else
+				{
+					projectFile = NativeProjects.EnumerateProjectFiles(Log.Logger).FirstOrDefault(x => x.GetFileNameWithoutExtension().Equals(_parameters.Project, StringComparison.OrdinalIgnoreCase));
+				}
+
+				if (projectFile == null || !FileReference.Exists(projectFile))
+				{
+					throw new BuildException("Unable to resolve project '{0}'", _parameters.Project);
+				}
+			}
+
+			string commonArgs = String.Format("{0}CreateCache {1}{2}", projectFile != null ? $"{projectFile.FullName} " : "", _parameters.Arguments, CommandUtils.IsBuildMachine ? " -buildmachine" : "");
+
 			if (_parameters.DDCReplayFile.StartsWith("#", StringComparison.OrdinalIgnoreCase))
 			{
 				HashSet<FileReference> files = ResolveFilespec(Unreal.RootDirectory, _parameters.DDCReplayFile, tagNameToFileSet);
-				string commonArgs = String.Format("{0}{1}", _parameters.Arguments, CommandUtils.IsBuildMachine ? " -buildmachine" : "");
 				foreach (FileReference file in files)
 				{
 					string taskArguments = $"{commonArgs} -DDC-ReplayLoad={file}";
@@ -103,7 +130,8 @@ namespace AutomationTool.Tasks
 			}
 			else
 			{
-				await ExecuteAsync(_parameters.Exe, _parameters.Arguments, _parameters.WorkingDir, envVars: ParseEnvVars(_parameters.Environment, _parameters.EnvironmentFile), logOutput: _parameters.LogOutput, errorLevel: _parameters.ErrorLevel);
+				string taskArguments = $"{commonArgs} -DDC-ReplayLoad={_parameters.DDCReplayFile}";
+				await ExecuteAsync(_parameters.Exe, taskArguments, _parameters.WorkingDir, envVars: ParseEnvVars(_parameters.Environment, _parameters.EnvironmentFile), logOutput: _parameters.LogOutput, errorLevel: _parameters.ErrorLevel);
 			}
 		}
 
