Remove use of dcpvar r.SkinCache.CompileShaders in automation test command lines.
It is also a cook time setting so when used it needs to be coordinated with the cook by the config files or the build graph.
#rnx
#Synced-CL 48982928

diff --git a/Engine/Source/Programs/AutomationTool/Gauntlet/Unreal/Automation/UE.Automation.cs b/Engine/Source/Programs/AutomationTool/Gauntlet/Unreal/Automation/UE.Automation.cs
--- a/Engine/Source/Programs/AutomationTool/Gauntlet/Unreal/Automation/UE.Automation.cs
+++ b/Engine/Source/Programs/AutomationTool/Gauntlet/Unreal/Automation/UE.Automation.cs
@@ -271,7 +271,7 @@ namespace UE
 				// These are flags that are required only on the main role that is going to execute the tests. ie raytracing is required only on the client or if it is an editor test.
 				if (RayTracing)
 				{
-					AppConfig.CommandLine += " -dpcvars=r.RayTracing=1,r.SkinCache.CompileShaders=1,r.Lumen.HardwareRayTracing=1";
+					AppConfig.CommandLine += " -dpcvars=r.RayTracing=1,r.Lumen.HardwareRayTracing=1";
 				}
 				else
 				{
