Use Get file name to set Map value metadata

#rnx

#Synced-CL 48996119

diff --git a/Engine/Source/Programs/AutomationTool/Gauntlet/Unreal/Base/Gauntlet.UnrealTestNode.cs b/Engine/Source/Programs/AutomationTool/Gauntlet/Unreal/Base/Gauntlet.UnrealTestNode.cs
--- a/Engine/Source/Programs/AutomationTool/Gauntlet/Unreal/Base/Gauntlet.UnrealTestNode.cs
+++ b/Engine/Source/Programs/AutomationTool/Gauntlet/Unreal/Base/Gauntlet.UnrealTestNode.cs
@@ -1922,7 +1922,7 @@ namespace Gauntlet
 			{
 				if (!string.IsNullOrEmpty(TestConfig.Map))
 				{
-					Report.SetMetadata("Map", TestConfig.Map);
+					Report.SetMetadata("Map", Path.GetFileName(TestConfig.Map));
 				}
 			}
 		}
