bump bpt version to live

diff --git a/Engine/Source/Programs/AutomationTool/AutomationUtils/MCPPublic.cs b/Engine/Source/Programs/AutomationTool/AutomationUtils/MCPPublic.cs
--- a/Engine/Source/Programs/AutomationTool/AutomationUtils/MCPPublic.cs
+++ b/Engine/Source/Programs/AutomationTool/AutomationUtils/MCPPublic.cs
@@ -666,7 +666,7 @@ namespace EpicGames.MCP.Automation
 			Online_v152,
 			Online_v160,
 
-			Online_Live = Online_v160
+			Online_Live = Live
 		}
 
 		/// <summary>
