[Gauntlet] Fix exception formatting in UnrealSession

diff --git a/Engine/Source/Programs/AutomationTool/Gauntlet/Unreal/Base/Gauntlet.UnrealSession.cs b/Engine/Source/Programs/AutomationTool/Gauntlet/Unreal/Base/Gauntlet.UnrealSession.cs
--- a/Engine/Source/Programs/AutomationTool/Gauntlet/Unreal/Base/Gauntlet.UnrealSession.cs
+++ b/Engine/Source/Programs/AutomationTool/Gauntlet/Unreal/Base/Gauntlet.UnrealSession.cs
@@ -1916,7 +1916,7 @@ namespace Gauntlet
 
 					if (AppInstance == null)
 					{
-						throw new AutomationException("Failed to create an IAppInstance after attempting Run() on {Install}", Install);
+						throw new AutomationException("Failed to create an IAppInstance after attempting Run() on {0}", Install);
 					}
 					else
 					{
