Linux/MacOS - Ensure that the generated symbolizer shell script for Android tombstones is actually syntacticly correct so it can run by users without modifications.

#Synced-CL 48848898

diff --git a/Engine/Source/Programs/AutomationTool/Android/AndroidPlatform.Automation.cs b/Engine/Source/Programs/AutomationTool/Android/AndroidPlatform.Automation.cs
--- a/Engine/Source/Programs/AutomationTool/Android/AndroidPlatform.Automation.cs
+++ b/Engine/Source/Programs/AutomationTool/Android/AndroidPlatform.Automation.cs
@@ -2015,14 +2015,16 @@ public class AndroidPlatform : Platform
 			Logger.LogInformation("Writing shell script for symbolize with {Arg0}", "data in APK" );
 			BatchLines = new string[] {
 				"#!/bin/sh",
-				"if [ $? -ne 0]; then",
+				"if [ $# -eq 0 ]; then",
 				 "echo \"Required argument missing, pass a dump of adb crash log.\"",
 				 "exit 1",
 				"fi",
+				"if [ ! -d \"$NDKROOT\" ]; then",
+				 "echo Unable to locate local NDK location.",
+				 "exit 1",
+				"fi",
 				"cd \"`dirname \"$0\"`\"",
-				"NDKSTACK=",
-				"if [ \"$ANDROID_NDK_ROOT\" != \"\" ]; then NDKSTACK=$%ANDROID_NDK_ROOT/ndk-stack; else ADB=" + Environment.GetEnvironmentVariable("ANDROID_NDK_ROOT") + "/ndk-stack; fi",
-				"$NDKSTACK -sym " + GetFinalSymbolizedSODirectory(ApkName, SC, Architecture) + " -dump \"%1\" > " + Params.ShortProjectName + "_SymbolizedCallStackOutput.txt",
+				"$NDKROOT/ndk-stack -sym " + GetFinalSymbolizedSODirectory(ApkName, SC, Architecture) + " -dump \"$1\" > " + Params.ShortProjectName + "_SymbolizedCallStackOutput.txt",
 				"exit 0",
 				};
 		}
@@ -2032,12 +2034,10 @@ public class AndroidPlatform : Platform
 			BatchLines = new string[] {
 						"@echo off",
 						"IF %1.==. GOTO NoArgs",
+						"if \"%NDKROOT%\"==\"\" GOTO NoNDK",
 						"setlocal",
-						"set NDK_ROOT=%ANDROID_NDK_ROOT%",
-						"if \"%ANDROID_NDK_ROOT%\"==\"\" set NDK_ROOT=\""+Environment.GetEnvironmentVariable("ANDROID_NDK_ROOT")+"\"",
-						"set NDKSTACK=%NDK_ROOT%\\ndk-stack.cmd",
 						"",
-						"%NDKSTACK% -sym "+GetFinalSymbolizedSODirectory(ApkName, SC, Architecture)+" -dump \"%1\" > "+ Params.ShortProjectName+"_SymbolizedCallStackOutput.txt",
+						"%NDKROOT%\\ndk-stack.cmd -sym "+GetFinalSymbolizedSODirectory(ApkName, SC, Architecture)+" -dump \"%1\" > "+ Params.ShortProjectName+"_SymbolizedCallStackOutput.txt",
 						"",
 						"goto:eof",
 						"",
@@ -2045,6 +2045,13 @@ public class AndroidPlatform : Platform
 						":NoArgs",
 						"echo.",
 						"echo Required argument missing, pass a dump of adb crash log. (SymboliseCallStackDump C:\\adbcrashlog.txt)",
+						"pause",
+						"",
+						"goto:eof",
+						"",
+						":NoNDK",
+						"echo.",
+						"echo Unable to locate local NDK location.",
 						"pause"
 					};
 		}
