[Automation] Automation Platforms now define a function to strip all symbols. The default behavior calls into the existing strip symbols function
Android implements strip all symbols to run with --strip-unneeded
StripTask now supplies an option to call the strip all variant
#Synced-CL 48971518

diff --git a/Engine/Source/Programs/AutomationTool/Android/AndroidPlatform.Automation.cs b/Engine/Source/Programs/AutomationTool/Android/AndroidPlatform.Automation.cs
--- a/Engine/Source/Programs/AutomationTool/Android/AndroidPlatform.Automation.cs
+++ b/Engine/Source/Programs/AutomationTool/Android/AndroidPlatform.Automation.cs
@@ -4507,6 +4507,11 @@ public class AndroidPlatform : Platform
 	{
 		AndroidExports.StripSymbols(SourceFile, TargetFile, Log.Logger);
 	}
+
+	public override void StripAllSymbols(FileReference SourceFile, FileReference TargetFile)
+	{
+		AndroidExports.StripSymbols(SourceFile, TargetFile, Log.Logger, true);
+	}
 }
 
 public class AndroidPlatformMulti : AndroidPlatform
diff --git a/Engine/Source/Programs/AutomationTool/AutomationUtils/Platform.cs b/Engine/Source/Programs/AutomationTool/AutomationUtils/Platform.cs
--- a/Engine/Source/Programs/AutomationTool/AutomationUtils/Platform.cs
+++ b/Engine/Source/Programs/AutomationTool/AutomationUtils/Platform.cs
@@ -832,6 +832,11 @@ namespace AutomationTool
 			}
 		}
 
+		public virtual void StripAllSymbols(FileReference SourceFile, FileReference TargetFile)
+		{
+			StripSymbols(SourceFile, TargetFile);
+		}
+
 		public virtual bool PublishSymbols(DirectoryReference SymbolStoreDirectory, List<FileReference> Files,
 			bool bIndexSources, List<FileReference> SourceFiles,
 			string Product, string Branch, int Change, string BuildVersion = null, string VFSMapping = null)
diff --git a/Engine/Source/Programs/AutomationTool/BuildGraph/Tasks/StripTask.cs b/Engine/Source/Programs/AutomationTool/BuildGraph/Tasks/StripTask.cs
--- a/Engine/Source/Programs/AutomationTool/BuildGraph/Tasks/StripTask.cs
+++ b/Engine/Source/Programs/AutomationTool/BuildGraph/Tasks/StripTask.cs
@@ -45,6 +45,12 @@ namespace AutomationTool.Tasks
 		/// </summary>
 		[TaskParameter(Optional = true, ValidationType = TaskParameterValidationType.TagList)]
 		public string Tag { get; set; }
+
+		/// <summary>
+		/// Whether all symbols should be stripped. Not relevant to all platforms.
+		/// </summary>
+		[TaskParameter(Optional = true)]
+		public bool StripAll { get; set; }
 	}
 
 	/// <summary>
@@ -97,7 +103,15 @@ namespace AutomationTool.Tasks
 				{
 					Logger.LogInformation("Stripping symbols: {Arg0} -> {Arg1}", sourceFiles[idx].FullName, targetFiles[idx].FullName);
 				}
-				targetPlatform.StripSymbols(sourceFiles[idx], targetFiles[idx]);
+
+				if (_parameters.StripAll)
+				{
+					targetPlatform.StripAllSymbols(sourceFiles[idx], targetFiles[idx]);
+				}
+				else
+				{
+					targetPlatform.StripSymbols(sourceFiles[idx], targetFiles[idx]);
+				}
 			}
 
 			// Apply the optional tag to the build products
