[Gauntlet] Default LongProcessResult to use an argument array for better compliance with unix systems
IOS app launch now uses arg list
#Synced-CL 48971518

diff --git a/Engine/Source/Programs/AutomationTool/Gauntlet/Framework/Gauntlet.Utils.cs b/Engine/Source/Programs/AutomationTool/Gauntlet/Framework/Gauntlet.Utils.cs
--- a/Engine/Source/Programs/AutomationTool/Gauntlet/Framework/Gauntlet.Utils.cs
+++ b/Engine/Source/Programs/AutomationTool/Gauntlet/Framework/Gauntlet.Utils.cs
@@ -1,25 +1,25 @@
 // Copyright Epic Games, Inc. All Rights Reserved.
 
+using AutomationTool;
+using EpicGames.Core;
+using ImageMagick;
+using Microsoft.Extensions.Logging;
 using System;
 using System.Collections.Generic;
+using System.CommandLine.Parsing;
+using System.Diagnostics;
 using System.IO;
 using System.Linq;
-using System.Diagnostics;
-using AutomationTool;
-using System.Threading;
-using System.Text.RegularExpressions;
-using System.Threading.Tasks;
 using System.Reflection;
-using ImageMagick;
-using UnrealBuildTool;
-using EpicGames.Core;
 using System.Security.Cryptography;
 using System.Text;
-using Microsoft.Extensions.Logging;
-using Logging = Microsoft.Extensions.Logging;
-
-using static AutomationTool.CommandUtils;
+using System.Text.RegularExpressions;
+using System.Threading;
+using System.Threading.Tasks;
 using UnrealBuildBase;
+using UnrealBuildTool;
+using static AutomationTool.CommandUtils;
+using Logging = Microsoft.Extensions.Logging;
 
 namespace Gauntlet
 {
@@ -2831,10 +2831,20 @@ namespace Gauntlet
 		private OutputFilterCallbackType OutputFilterCallback;
 
 		public LongProcessResult(string Command, string Args, ERunOptions Options, OutputFilterCallbackType OutputCallback = null, string WorkingDir = null, string LocalCache = null)
+		{
+			Create(Command, CommandLineParser.SplitCommandLine(Args), Options, OutputCallback, WorkingDir, LocalCache);
+		}
+
+		public LongProcessResult(string Command, IEnumerable<string> Args, ERunOptions Options, OutputFilterCallbackType OutputCallback = null, string WorkingDir = null, string LocalCache = null)
+		{
+			Create(Command, Args, Options, OutputCallback, WorkingDir, LocalCache);
+		}
+
+		protected void Create(string Command, IEnumerable<string> Args, ERunOptions Options, OutputFilterCallbackType OutputCallback = null, string WorkingDir = null, string LocalCache = null)
 		{
 			ExitCode = -1;
 			Logger = null;
-			LocalCachePath = string.IsNullOrEmpty(LocalCache)? ProcessUtils.LocalLogsPath : Path.Combine(LocalCache, "Logs");
+			LocalCachePath = string.IsNullOrEmpty(LocalCache) ? ProcessUtils.LocalLogsPath : Path.Combine(LocalCache, "Logs");
 			ProcSyncObject = new object();
 			OutputSyncObject = new object();
 			OutputFilterCallback = OutputCallback;
@@ -2844,7 +2854,7 @@ namespace Gauntlet
 				StartInfo.WorkingDirectory = WorkingDir;
 			}
 			AppName = Path.GetFileNameWithoutExtension(StartInfo.FileName);
-			CommandLine = $"{Path.GetFileName(StartInfo.FileName)} {Args}";
+			CommandLine = $"{Path.GetFileName(StartInfo.FileName)} {string.Join(' ', Args)}";
 			if (!Options.HasFlag(ERunOptions.NoLoggingOfRunCommand))
 			{
 				Log.Info("Running: {CommandLine}", CommandLine);
@@ -2880,6 +2890,16 @@ namespace Gauntlet
 			RegisterProcessExit();
 		}
 
+		private void ResolveArguments(string[] Args)
+		{
+
+		}
+
+		private void ResolveArguments(string Args)
+		{
+
+		}
+
 		private void TryCreateLogFile()
 		{
 			string LogFilename = ProcessUtils.GetLogFilePath(AppName, LocalCachePath);
diff --git a/Engine/Source/Programs/AutomationTool/Gauntlet/Gauntlet.Automation.csproj b/Engine/Source/Programs/AutomationTool/Gauntlet/Gauntlet.Automation.csproj
--- a/Engine/Source/Programs/AutomationTool/Gauntlet/Gauntlet.Automation.csproj
+++ b/Engine/Source/Programs/AutomationTool/Gauntlet/Gauntlet.Automation.csproj
@@ -91,6 +91,7 @@
     <PackageReference Include="Magick.NET-Q16-HDRI-AnyCPU" Version="14.9.1" PrivateAssets="all" />
     <PackageReference Include="Microsoft.CSharp" Version="4.7.0" />
     <PackageReference Include="MySql.Data" Version="6.10.9" PrivateAssets="all" />
+    <PackageReference Include="System.CommandLine" Version="2.0.1" />
     <PackageReference Include="Newtonsoft.Json" Version="13.0.4" PrivateAssets="all" />
   </ItemGroup>
 </Project>
\ No newline at end of file
diff --git a/Engine/Source/Programs/AutomationTool/Gauntlet/Platform/IOS/Gauntlet.TargetDeviceIOS.cs b/Engine/Source/Programs/AutomationTool/Gauntlet/Platform/IOS/Gauntlet.TargetDeviceIOS.cs
--- a/Engine/Source/Programs/AutomationTool/Gauntlet/Platform/IOS/Gauntlet.TargetDeviceIOS.cs
+++ b/Engine/Source/Programs/AutomationTool/Gauntlet/Platform/IOS/Gauntlet.TargetDeviceIOS.cs
@@ -13,6 +13,7 @@ using AutomationTool;
 using UnrealBuildTool;
 using static AutomationTool.CommandUtils;
 using System.Threading;
+using System.CommandLine.Parsing;
 
 /** IOS General Notes:
  *
@@ -809,14 +810,26 @@ namespace Gauntlet
 			Log.Info("Launching {0} on {1}", Install.Name, this);
 			Log.Verbose(IOSInstall.CommandLine);
 
+			string[] LaunchCommand;
 			ILongProcessResult AppProcess;
 			if (UseDeviceCtl)
 			{
 				// device process launch	base command
 				// --terminate-existing		exit any existing forms of the app
+				// --environment-variables	json blob of envvars to set
 				// --console				pipe process stdout to devicectl, also holds onto the process until exit instead of immediate exit
 				// --device					device to run on
-				string LaunchCommand = string.Format("device process launch --terminate-existing --console --device {0} {1}", UUID, IOSInstall.PackageName);
+				LaunchCommand =
+				[
+					"device",
+					"process",
+					"launch",
+					"--terminate-existing",
+					"--environment-variables", $"{{\"OS_ACTIVITY_DT_MODE\":\"enable\"}}",
+					"--console",
+					"--device", UUID,
+					IOSInstall.PackageName
+				];
 				AppProcess = ExecuteDevicectlCommandNoWait(LaunchCommand, AdditionalOptions: ERunOptions.SpewIsVerbose, LocalCache: IOSInstall.Device.LocalCachePath);
 			}
 			else
@@ -831,7 +844,7 @@ namespace Gauntlet
 				// -m		Skip installation
 				// -I		Non-interactive mode, just launch and end launch process when app process ends
 				// -b		path to bundle to run
-				string LaunchCommand = string.Format("-m -I -b \"{0}\"", IOSInstall.AppPath);
+				LaunchCommand = ["-m", "-I", "-b", IOSInstall.AppPath];
 				AppProcess = ExecuteIOSDeployCommandNoWait(LaunchCommand, LocalCache: IOSInstall.Device.LocalCachePath);
 			}
 
@@ -865,6 +878,11 @@ namespace Gauntlet
 		}
 
 		public static ILongProcessResult ExecuteIOSDeployCommandNoWait(string CommandLine, ERunOptions AdditionalOptions = ERunOptions.None, string LocalCache = null, bool bQuiet = false)
+		{
+			return ExecuteIOSDeployCommandNoWait(CommandLineParser.SplitCommandLine(CommandLine), AdditionalOptions, LocalCache, bQuiet);
+		}
+
+		public static ILongProcessResult ExecuteIOSDeployCommandNoWait(IEnumerable<string> CommandLine, ERunOptions AdditionalOptions = ERunOptions.None, string LocalCache = null, bool bQuiet = false)
 		{
 			string IOSDeploy = Path.Combine(Globals.UnrealRootDir, "Engine/Extras/ThirdPartyNotUE/ios-deploy/bin/ios-deploy");
 			return ExecuteDeploymentCommandNoWait(IOSDeploy, CommandLine, null, AdditionalOptions, LocalCache, bQuiet);
@@ -902,6 +920,13 @@ namespace Gauntlet
 			return ExecuteDeploymentCommandNoWait(XCRun, CommandLine, AdditionalOptions: AdditionalOptions, LocalCache: LocalCache, bQuiet: bQuiet);
 		}
 
+		public static ILongProcessResult ExecuteDevicectlCommandNoWait(string[] CommandLine, ERunOptions AdditionalOptions = ERunOptions.None, string LocalCache = null, bool bQuiet = false)
+		{
+			CommandLine = ["devicectl", ..CommandLine];
+			string XCRun = "/usr/bin/xcrun";
+			return ExecuteDeploymentCommandNoWait(XCRun, CommandLine, AdditionalOptions: AdditionalOptions, LocalCache: LocalCache, bQuiet: bQuiet);
+		}
+
 		public IProcessResult ExecuteDevicectlCommand(string CommandLine, int WaitTime = 60, bool WarnOnTimeout = true, bool UseDeviceID = true, ERunOptions AdditionalOptions = ERunOptions.None, bool bQuiet = false)
 		{
 			if (UseDeviceID)
@@ -973,6 +998,11 @@ namespace Gauntlet
 		}
 
 		public static ILongProcessResult ExecuteDeploymentCommandNoWait(string Executable, string CommandLine, string WorkingDir = null, ERunOptions AdditionalOptions = ERunOptions.None, string LocalCache = null, bool bQuiet = false)
+		{
+			return ExecuteDeploymentCommandNoWait(Executable, CommandLineParser.SplitCommandLine(CommandLine), WorkingDir, AdditionalOptions, LocalCache, bQuiet);
+		}
+
+		public static ILongProcessResult ExecuteDeploymentCommandNoWait(string Executable, IEnumerable<string> CommandLine, string WorkingDir = null, ERunOptions AdditionalOptions = ERunOptions.None, string LocalCache = null, bool bQuiet = false)
 		{
 			if (!AdditionalOptions.HasFlag(ERunOptions.UseShellExecute) && !File.Exists(Executable))
 			{
@@ -981,7 +1011,7 @@ namespace Gauntlet
 
 			if (!bQuiet)
 			{
-				Log.Info("{DeploymentExecutable} executing '{Command}'", Executable, CommandLine);
+				Log.Info("{DeploymentExecutable} executing '{Command}'", Executable, string.Join(' ', CommandLine));
 			}
 
 			ERunOptions RunOptions = ERunOptions.NoWaitForExit | AdditionalOptions;
