[Gauntlet] Adds ability to run -dev binaries on non-bulk Android builds

#Synced-CL 49130326

diff --git a/Engine/Source/Programs/AutomationTool/Android/Scripts/AndroidAutomationUtils.cs b/Engine/Source/Programs/AutomationTool/Android/Scripts/AndroidAutomationUtils.cs
new file mode 100644
--- /dev/null
+++ b/Engine/Source/Programs/AutomationTool/Android/Scripts/AndroidAutomationUtils.cs
@@ -0,0 +1,49 @@
+// Copyright Epic Games, Inc. All Rights Reserved.
+
+using AutomationTool;
+using EpicGames.Core;
+using Microsoft.Extensions.Logging;
+using System;
+using System.IO;
+using System.IO.Compression;
+using UnrealBuildBase;
+using UnrealBuildTool;
+
+public class ReplaceAndroidExecutable : BuildCommand
+{
+	public override ExitCode Execute()
+	{
+		FileReference ProjectFile = new (ParseParamValue("ProjectFile", null));
+		FileReference Executable = new (ParseParamValue("Executable", null));
+		FileReference SourceApk = new (ParseParamValue("SourceApk", null));
+		FileReference TargetApk = new (ParseParamValue("TargetApk", null));
+		string Architecture = ParseParamValue("Architecture", null);
+
+		if (!FileReference.Exists(ProjectFile))
+		{
+			throw new AutomationException("The Project file '{0}' does not exist", ProjectFile);
+		}
+		if (!FileReference.Exists(Executable))
+		{
+			throw new AutomationException("The executable '{0}' to use as a replacement does not exist", Executable);
+		}
+		if (!FileReference.Exists(SourceApk))
+		{
+			throw new AutomationException("The source apk '{0}' does not exist", SourceApk);
+		}
+		if (string.IsNullOrEmpty(Architecture))
+		{
+			throw new AutomationException("The architecture was null or empty");
+		}
+
+		string RelativeSO = $"lib/{Architecture}/libUnreal.so";
+
+		Logger.LogInformation("Replacing shared object file...");
+		AndroidExports.ReplaceSO(SourceApk, TargetApk, Executable, RelativeSO, Logger);
+
+		Logger.LogInformation("Signing apk...");
+		AndroidExports.SignDebugApk(ProjectFile, TargetApk, TargetApk, Logger);
+
+		return ExitCode.Success;
+	}
+}
\ No newline at end of file
diff --git a/Engine/Source/Programs/AutomationTool/Gauntlet/Platform/Android/Gauntlet.AndroidBuildSource.cs b/Engine/Source/Programs/AutomationTool/Gauntlet/Platform/Android/Gauntlet.AndroidBuildSource.cs
--- a/Engine/Source/Programs/AutomationTool/Gauntlet/Platform/Android/Gauntlet.AndroidBuildSource.cs
+++ b/Engine/Source/Programs/AutomationTool/Gauntlet/Platform/Android/Gauntlet.AndroidBuildSource.cs
@@ -158,14 +158,8 @@ namespace Gauntlet
 					continue;
 				}
 
-				// Android builds are always packaged, and we can always replace the command line
-				BuildFlags Flags = BuildFlags.Packaged | BuildFlags.CanReplaceCommandLine;
-
-				// if there's data then the pak files are in an obb and we can sub in a new exe
-				if (FilesToInstall.Count() > 0)
-				{
-					Flags |= BuildFlags.CanReplaceExecutable;
-				}
+				// Android builds are always packaged, and we can always replace the command line and executable
+				BuildFlags Flags = BuildFlags.Packaged | BuildFlags.CanReplaceCommandLine | BuildFlags.CanReplaceExecutable;
 				if (AbsPath.Contains("Bulk"))
 				{
 					Flags |= BuildFlags.Bulk;
