[Gauntlet] Adds -UseAgentTag flag. When enabled, this appends the name of the agent running the UAT instance as a horde device reservation required tag. This allows agents to selectively choose devices they have a more direct relation too, i.e. devices in a larger pool they are explicitly wired to.

#Synced-CL 49183475

diff --git a/Engine/Source/Programs/AutomationTool/AutomationUtils/DeviceReservation.cs b/Engine/Source/Programs/AutomationTool/AutomationUtils/DeviceReservation.cs
--- a/Engine/Source/Programs/AutomationTool/AutomationUtils/DeviceReservation.cs
+++ b/Engine/Source/Programs/AutomationTool/AutomationUtils/DeviceReservation.cs
@@ -264,7 +264,7 @@ namespace AutomationTool.DeviceReservation
 			public Tags Tags { get; set; }
 		}
 
-		public class Tags
+		public class Tags : ICloneable
 		{
 			public enum Type
 			{
@@ -383,6 +383,17 @@ namespace AutomationTool.DeviceReservation
 				return ToString().GetHashCode();
 			}
 
+			object ICloneable.Clone() => Clone();
+			public Tags Clone()
+			{
+				Tags Other = (Tags)MemberwiseClone();
+				Other.Required = new List<string>(Required);
+				Other.Preferred = new List<string>(Preferred);
+				Other.Undesired = new List<string>(Undesired);
+				Other.Blocked = new List<string>(Blocked);
+				return Other;
+			}
+
 			/// <summary>
 			/// Does this Tag request have any tags associated with it
 			/// </summary>
diff --git a/Engine/Source/Programs/AutomationTool/Gauntlet/Framework/Devices/Gauntlet.DeviceReservationService.cs b/Engine/Source/Programs/AutomationTool/Gauntlet/Framework/Devices/Gauntlet.DeviceReservationService.cs
--- a/Engine/Source/Programs/AutomationTool/Gauntlet/Framework/Devices/Gauntlet.DeviceReservationService.cs
+++ b/Engine/Source/Programs/AutomationTool/Gauntlet/Framework/Devices/Gauntlet.DeviceReservationService.cs
@@ -12,6 +12,8 @@ using System.Data;
 using System.Linq;
 using UnrealBuildTool;
 
+#nullable enable
+
 namespace Gauntlet
 {
 	/// <summary>
@@ -49,13 +51,35 @@ namespace Gauntlet
 		/// The base URL of the horde service to request devices from
 		/// </summary>
 		[AutoParam("")]
-		public string DeviceURL { get; protected set; }
+		public string DeviceURL { get; protected set; } = string.Empty;
 
 		/// <summary>
 		/// Name of the horde device pool to request devices from
 		/// </summary>
 		[AutoParamWithNames(Default: "", "DevicePool")]
-		public string DevicePoolID { get; protected set; }
+		public string DevicePoolID { get; protected set; } = string.Empty;
+
+		/// <summary>
+		/// If true, the horde agent will have it's name appended as a required tag
+		/// This is useful in instances where you want an agent to use a limited number of devices in a larger pool
+		/// </summary>
+		[AutoParam(false)]
+		public bool UseAgentTag;
+
+		/// <summary>
+		/// HTTP client used to collect horde information
+		/// </summary>
+		private static HordeHttpClient? HordeClient = null;
+
+		/// <summary>
+		/// ID of the actively running horde job
+		/// </summary>
+		private static JobId HordeJobId = default;
+
+		/// <summary>
+		/// ID of the actively running horde step
+		/// </summary>
+		private static JobStepId HordeStepId = default;
 
 		/// <summary>
 		/// Endpoint for reservations
@@ -90,15 +114,16 @@ namespace Gauntlet
 		{
 			if (Enabled)
 			{
-				JobId JobId = default;
-				JobStepId StepId = default;
+				HordeClient = CommandUtils.ServiceProvider.GetRequiredService<IHordeClient>().CreateHttpClient();
+
 				try
 				{
-					string JobIdEnvVar = Environment.GetEnvironmentVariable("UE_HORDE_JOBID");
-					JobId = JobId.Parse(JobIdEnvVar);
+					// Allow these to throw null
+					string JobIdEnvVar = Environment.GetEnvironmentVariable("UE_HORDE_JOBID")!;
+					HordeJobId = JobId.Parse(JobIdEnvVar);
 
-					string StepIdEnvVar = Environment.GetEnvironmentVariable("UE_HORDE_STEPID");
-					StepId = JobStepId.Parse(StepIdEnvVar);
+					string StepIdEnvVar = Environment.GetEnvironmentVariable("UE_HORDE_STEPID")!;
+					HordeStepId = JobStepId.Parse(StepIdEnvVar);
 				}
 				catch
 				{
@@ -108,17 +133,9 @@ namespace Gauntlet
 
 				try
 				{
-					// Locate the batch and step
-					HordeHttpClient HordeClient = CommandUtils.ServiceProvider.GetRequiredService<IHordeClient>().CreateHttpClient();
-					GetJobResponse Job = HordeClient.GetJobAsync(JobId).Result;
-					GetJobBatchResponse JobBatch = Job.Batches
-						.Where(Batch => Batch.Steps
-						.Where(Step => Step.Id == StepId)
-						.FirstOrDefault()?.Id == StepId)
-						.First();
-					GetJobStepResponse Step = JobBatch.Steps.Where(Step => Step.Id == StepId).First();
-
-					if (Step.Annotations.TryGetValue("DeviceReserve", out string Value))
+					GetJobBatchResponse Batch = GetBatchInfo(GetJobInfo());
+					GetJobStepResponse Step = GetStepInfo(Batch);
+					if (Step.Annotations != null && Step.Annotations.TryGetValue("DeviceReserve", out string? Value))
 					{
 						bool InstallStep = false;
 						if (Value.Equals("Begin", StringComparison.InvariantCultureIgnoreCase))
@@ -142,14 +159,14 @@ namespace Gauntlet
 					}
 					else
 					{
-						foreach (GetJobStepResponse PreviousStep in JobBatch.Steps)
+						foreach (GetJobStepResponse PreviousStep in Batch.Steps)
 						{
 							if (PreviousStep.Id == Step.Id)
 							{
 								// Ignore steps that occur after this one
 								break;
 							}
-							else if (PreviousStep.Annotations.TryGetValue("DeviceReserve", out Value) && Value.Equals("Begin", StringComparison.InvariantCultureIgnoreCase))
+							else if (PreviousStep.Annotations != null && PreviousStep.Annotations.TryGetValue("DeviceReserve", out Value) && Value.Equals("Begin", StringComparison.InvariantCultureIgnoreCase))
 							{
 								Log.Verbose("Detected running step is within a reservation block. Any devices used will skip installations");
 								DevicePool.IsInstallStep = false;
@@ -174,7 +191,7 @@ namespace Gauntlet
 		{
 			AutoParam.ApplyParamsAndDefaults(this, Globals.Params.AllArguments);
 
-			if (!Uri.TryCreate(DeviceURL, UriKind.Absolute, out ReservationServerUri))
+			if (!Uri.TryCreate(DeviceURL, UriKind.Absolute, out ReservationServerUri!))
 			{
 				throw new AutomationException("Failed to resolve \"{0}\" as a valid URI", DeviceURL);
 			}
@@ -257,7 +274,20 @@ namespace Gauntlet
 			List<ITargetDevice> ScopeReservedDevices = new List<ITargetDevice>();
 			foreach (UnrealDeviceTargetConstraint Constraint in RequestedConstraints)
 			{
-				Reservation NewReservation = Reservation.Create(ReservationServerUri, [Constraint.FormatWithIdentifier()], TimeSpan.FromMinutes(10), RetryMax: 0, PoolID: DevicePoolID, DeviceName: Constraint.DeviceName);
+				Reservation.Tags Tags = Tags = Constraint.Tags;
+				if (UseAgentTag)
+				{
+					// If someone is trying to repro locally without the proper horde envvars, this will throw
+					// We can't guarantee to get the right device without an agent name, so this is preferred over silently reserving a default device
+					string? AgentName = GetBatchInfo(GetJobInfo()).AgentId;
+					if (!string.IsNullOrEmpty(AgentName))
+					{
+						Tags = Tags.Clone();
+						Tags.AddTag(AgentName, Reservation.Tags.Type.Required);
+					}
+				}
+
+				Reservation NewReservation = Reservation.Create(ReservationServerUri, [Constraint.FormatWithIdentifier()], TimeSpan.FromMinutes(10), RetryMax: 0, PoolID: DevicePoolID, DeviceName: Constraint.DeviceName, Tags);
 				DeviceReservationAutoRenew DeviceReservation = new DeviceReservationAutoRenew(DeviceURL, NewReservation);
 
 				if (DeviceReservation == null || DeviceReservation.Devices.Count != 1)
@@ -280,7 +310,7 @@ namespace Gauntlet
 				{
 					Address = Device.IPOrHostName,
 					Name = Device.Name,
-					Platform = UnrealTargetPlatform.Parse(UnrealTargetPlatform.GetValidPlatformNames().FirstOrDefault(Entry => Entry == Device.Type.Replace("-DevKit", "", StringComparison.OrdinalIgnoreCase))),
+					Platform = UnrealTargetPlatform.Parse(UnrealTargetPlatform.GetValidPlatformNames().First(Entry => Entry == Device.Type.Replace("-DevKit", "", StringComparison.OrdinalIgnoreCase))),
 					DeviceData = Device.DeviceData,
 					Model = Device.Model,
 					Tags = Device.Tags
@@ -355,7 +385,7 @@ namespace Gauntlet
 				}
 
 				// Unregister device
-				if (ServiceReservations.TryGetValue(Device, out DeviceReservationAutoRenew Reservation))
+				if (ServiceReservations.TryGetValue(Device, out DeviceReservationAutoRenew? Reservation))
 				{
 					bool DisposeReservation = ServiceReservations.Count(Entry => Entry.Value == Reservation) == 1;
 
@@ -390,7 +420,7 @@ namespace Gauntlet
 		public virtual void ReportDeviceError(string DeviceName, string ErrorMessage)
 		{
 			// TargetDevice name is not always DeviceData name... need to try to resolve target device names
-			ITargetDevice MatchingDevice = null;
+			ITargetDevice? MatchingDevice = null;
 			foreach (ITargetDevice Device in ReservedDevices)
 			{
 				if (Device.Name.Equals(DeviceName, StringComparison.OrdinalIgnoreCase))
@@ -418,7 +448,7 @@ namespace Gauntlet
 		/// <returns>An array containg all tags on the device</returns>
 		public List<string> GetDeviceTags(ITargetDevice Device)
 		{
-			if (ServiceDeviceInfo.TryGetValue(Device, out DeviceDefinition Definition))
+			if (ServiceDeviceInfo.TryGetValue(Device, out DeviceDefinition? Definition))
 			{
 				return Definition.Tags ?? [];
 			}
@@ -436,5 +466,52 @@ namespace Gauntlet
 		{
 			return GetDeviceTags(Device).Contains(Tag, StringComparer.OrdinalIgnoreCase);
 		}
+
+		/// <summary>
+		/// Fetches information about the actively running job
+		/// </summary>
+		/// <returns>Information of the actively running job</returns>
+		/// <exception cref="AutomationException">Throws when calling without the horde environment configured</exception>
+		private static GetJobResponse GetJobInfo()
+		{
+			if (HordeClient == null)
+			{
+				throw new AutomationException("Failed to retrieve Job info because the horde client was not initialized." +
+					"Ensure you're running in a horde environment before calling this function.");
+			}
+
+			return HordeClient.GetJobAsync(HordeJobId).Result;
+		}
+
+		/// <summary>
+		/// Fetches information about the actively running job batch
+		/// </summary>
+		/// <param name="Job">Job the batch is contained in</param>
+		/// <returns>Information of the actively running job batch</returns>
+		/// <exception cref="AutomationException">Throws when the Job doesn't contain batches</exception>
+		private static GetJobBatchResponse GetBatchInfo(GetJobResponse Job)
+		{
+			if (Job.Batches == null)
+			{
+				throw new AutomationException("Failed to retrieve the job's batch info because the batch info was null." +
+					"Ensure the UE_HORDE_JOBID is set to an existing job's id that contains relevant batch info");
+			}
+
+			return Job.Batches
+				.Where(Batch => Batch.Steps
+				.Where(Step => Step.Id == HordeStepId)
+				.FirstOrDefault()?.Id == HordeStepId)
+				.First();
+		}
+
+		/// <summary>
+		/// Fetches information about the actively running job step
+		/// </summary>
+		/// <param name="Batch">Batch the step is contained in</param>
+		/// <returns>Information of the actively running job step</returns>
+		private static GetJobStepResponse GetStepInfo(GetJobBatchResponse Batch)
+		{
+			return Batch.Steps.Where(Step => Step.Id == HordeStepId).First();
+		}
 	}
 }
\ No newline at end of file
