Implementation of utilizing appium on ios tests to take screenshots at a specified interval

#Synced-CL 49230356

diff --git a/Engine/Source/Programs/AutomationTool/Gauntlet/Platform/IOS/Gauntlet.IOSAppium.cs b/Engine/Source/Programs/AutomationTool/Gauntlet/Platform/IOS/Gauntlet.IOSAppium.cs
--- a/Engine/Source/Programs/AutomationTool/Gauntlet/Platform/IOS/Gauntlet.IOSAppium.cs
+++ b/Engine/Source/Programs/AutomationTool/Gauntlet/Platform/IOS/Gauntlet.IOSAppium.cs
@@ -8,9 +8,10 @@ using OpenQA.Selenium.Appium.iOS;
 using OpenQA.Selenium.Appium.Service;
 using OpenQA.Selenium.Support.UI;
 using System;
-using System.Collections.Generic;
+using System.IO;
 using System.Text.RegularExpressions;
 using System.Threading;
+using UnrealBuildBase;
 using static AutomationTool.CommandUtils;
 
 namespace Gauntlet
@@ -42,7 +43,7 @@ namespace Gauntlet
 	/// </summary>
 	public class AppiumContainer : IDisposable
 	{
-		class Config
+		public class Config
 		{
 			/// <summary>
 			/// Name of the automation driver to use. Defaults to XCUITest if not specified
@@ -75,6 +76,16 @@ namespace Gauntlet
 			/// </summary>
 			public int WdaLaunchTimeoutMs { get; set; } = 60000;
 
+			/// <summary>
+			/// How often to take screenshots in seconds. Disabled if <= 0
+			/// </summary>
+			public float ScreenshotPeriod { get; set; } = 0;
+
+			/// <summary>
+			/// Path screenshots are output to
+			/// </summary>
+			public string ScreenshotDirectory { get; set; }
+
 			/// <summary>
 			/// Optional - Allow you to override the location of the appium executable.
 			/// Useful if you opt not to install appium to the global npm root
@@ -132,7 +143,7 @@ namespace Gauntlet
 		/// <summary>
 		/// Config file used for the Driver
 		/// </summary>
-		private static Config AppiumConfig = null;
+		public static Config AppiumConfig { get; private set; } = null;
 
 		/// <summary>
 		/// Driver handle
@@ -144,6 +155,10 @@ namespace Gauntlet
 		/// </summary>
 		private AppiumLocalService Server = null;
 
+		private Thread ScreenshotThread = null;
+
+		private AutoResetEvent ScreenshotTrigger = null;
+
 		/// <summary>
 		/// UUID of the device being tests
 		/// </summary>
@@ -269,6 +284,12 @@ namespace Gauntlet
 					throw new DeviceException($"Exception during WebdriverAgent Creation {Ex.Message}");
 				}
 			}
+
+			if (!Globals.Params.ParseParam("NoAppiumScreenshots") && AppiumConfig.ScreenshotPeriod > 0)
+			{
+				ScreenshotThread = new Thread(StartScreenshotLoop);
+				ScreenshotThread.Start();
+			}
 		}
 
 		public void Stop()
@@ -285,6 +306,11 @@ namespace Gauntlet
 				Server = null;
 			}
 
+			if (ScreenshotThread != null)
+			{
+				StopScreenshotLoop();
+			}
+
 			// TODO: Terminate WDA app - parse devicectl for pid, then kill?
 		}
 
@@ -389,13 +415,13 @@ namespace Gauntlet
 				{
 					var alert = WebDriverWait.Until(drv =>
 					{
-						try 
-						{ 
-							return drv.SwitchTo().Alert(); 
+						try
+						{
+							return drv.SwitchTo().Alert();
 						}
-						catch (NoAlertPresentException) 
-						{ 
-							return null; 
+						catch (NoAlertPresentException)
+						{
+							return null;
 						}
 					});
 
@@ -425,5 +451,102 @@ namespace Gauntlet
 				Thread.Sleep(300);
 			}
 		}
+
+		/// <summary>
+		/// Uses appium WebDriver to take a screenshot. 
+		/// Filename is saved as png in the format of: yy-MM-dd-HH-mm-ss-ff
+		/// </summary>
+		/// <param name="OutputDirectory">Folder to copy screenshots to</param>
+		public void TakeScreenshot(string OutputDirectory)
+		{
+			if (string.IsNullOrEmpty(OutputDirectory))
+			{
+				throw new AutomationException("Screenshot output directory must not be null or empty!");
+			}
+
+			if (!DirectoryExists(OutputDirectory))
+			{
+				CreateDirectory(OutputDirectory);
+			}
+
+			string FileName = $"{DateTime.Now.ToString("yy-MM-dd-HH-mm-ss-ff")}.png";
+			string FilePath = Path.Combine(OutputDirectory, FileName);
+			try
+			{
+				(Driver as ITakesScreenshot).GetScreenshot().SaveAsFile(FilePath);
+				Log.Verbose("Saved screenshot to: {ScreenshotFile}", FilePath);
+			}
+			catch (Exception Ex)
+			{
+				Log.Info("Encountered an {ExceptionType} when trying to capture a screenshot: {ExceptionMessage}", Ex.GetType().Name, Ex.Message);
+			}
+		}
+
+		private void StartScreenshotLoop()
+		{
+			TimeSpan Period = TimeSpan.FromSeconds(AppiumConfig.ScreenshotPeriod);
+			ScreenshotTrigger = new AutoResetEvent(false);
+
+			string ScreenshotDirectory = GetScreenshotDirectory();
+
+			if (DirectoryExists(ScreenshotDirectory))
+			{
+				DeleteDirectory_NoExceptions(ScreenshotDirectory);
+			}
+
+			while (!ScreenshotTrigger.WaitOne(Period))
+			{
+				TakeScreenshot(ScreenshotDirectory);
+			}
+		}
+
+		/// <summary>
+		/// Stops the capturing of periodic screenshots, optionally moving all captured images to an output directory
+		/// </summary>
+		/// <param name="OutputDirectory">Folder to copy screenshots to</param>
+		public void StopScreenshotLoop(string OutputDirectory = "")
+		{
+			if (ScreenshotThread == null || ScreenshotTrigger == null)
+			{
+				Log.Verbose("Attempted to end the screenshot loop before it started. Ignoring.");
+				return;
+			}
+
+			ScreenshotTrigger.Set();
+			ScreenshotThread.Join();
+			ScreenshotThread = null;
+			ScreenshotTrigger = null;
+
+			if (!string.IsNullOrEmpty(OutputDirectory))
+			{
+				try
+				{
+					if (!DirectoryExists(OutputDirectory))
+					{
+						CreateDirectory(OutputDirectory);
+					}
+
+					string ScreenshotDirectory = GetScreenshotDirectory();
+
+					Log.Verbose("Copying {SourceScreenshotDirectory} to {DestinationScreenshotDirectory}...", ScreenshotDirectory, OutputDirectory);
+					CopyDirectory_NoExceptions(ScreenshotDirectory, OutputDirectory);
+				}
+				catch (Exception Ex)
+				{
+					Log.Info("Encountered an {ExceptionType} when trying to copy screenshots to requested output directory. " +
+						"Screenshots will not be available: {ExceptionMessage}", Ex.GetType().Name, Ex.Message);
+				}
+			}
+		}
+
+		private string GetScreenshotDirectory()
+		{
+			string ScreenshotDirectory = AppiumConfig.ScreenshotDirectory;
+			if (string.IsNullOrEmpty(ScreenshotDirectory))
+			{
+				ScreenshotDirectory = Path.Combine(Unreal.RootDirectory.FullName, "GauntletTemp", "IOSScreenshots");
+			}
+			return ScreenshotDirectory;
+		}
 	}
-}
\ No newline at end of file
+}
diff --git a/Engine/Source/Programs/AutomationTool/Gauntlet/Platform/IOS/Gauntlet.TargetDeviceIOS.cs b/Engine/Source/Programs/AutomationTool/Gauntlet/Platform/IOS/Gauntlet.TargetDeviceIOS.cs
--- a/Engine/Source/Programs/AutomationTool/Gauntlet/Platform/IOS/Gauntlet.TargetDeviceIOS.cs
+++ b/Engine/Source/Programs/AutomationTool/Gauntlet/Platform/IOS/Gauntlet.TargetDeviceIOS.cs
@@ -146,7 +146,7 @@ namespace Gauntlet
 		/// </summary>
 		private IOSAppInstall InstallCache = null;
 
-		private AppiumContainer AppiumContainer = null;
+		public AppiumContainer AppiumContainer { get; private set; } = null;
 
 		public TargetDeviceIOS(string InName, string InCachePath = null)
 		{
@@ -1373,6 +1373,12 @@ namespace Gauntlet
 			{
 				Log.Error(KnownLogEvents.Gauntlet_DeviceEvent, "Failed to move artifacts out of temp directory:\n{Exception}", Ex.Message);
 			}
+
+			AppiumContainer.Config AppiumConfig = AppiumContainer.AppiumConfig;
+			if (!Globals.Params.ParseParam("NoAppiumScreenshots") && AppiumConfig != null && AppiumConfig.ScreenshotPeriod > 0)
+			{
+				IOSDevice.AppiumContainer.StopScreenshotLoop(Path.Combine(DestinationArtifactPath, "Screenshots"));
+			}
 		}
 	}
 
