Allow excluding the evaluation of Gauntlet plugin source files when evaluating Gauntlet with MSBuild, as currently no such plugin source files exist, and the search takes 8-12 seconds extra. Use this exclusion for UBT's project evaluation.

#Synced-CL 49405020

diff --git a/Engine/Source/Programs/AutomationTool/Gauntlet/Gauntlet.Automation.csproj b/Engine/Source/Programs/AutomationTool/Gauntlet/Gauntlet.Automation.csproj
--- a/Engine/Source/Programs/AutomationTool/Gauntlet/Gauntlet.Automation.csproj
+++ b/Engine/Source/Programs/AutomationTool/Gauntlet/Gauntlet.Automation.csproj
@@ -57,13 +57,22 @@
     </Compile>
   </ItemGroup>
 
-    <!-- The following will glob all UAT plugin's Gauntlet source
-  files and include them in the build. They will also appear in VisualStudio -->
-  <ItemGroup>
+  <!-- The following two ItemGroups will glob all UAT plugin's Gauntlet source
+    files and include them in the build. They will also appear in VisualStudio -->
+
+  <!-- Allow skipping including Gauntlet plugins behind a specific MSBuild property, as:
+    1. Currently there are no files that actually get picked up by this include.
+    2. The extensive disk scan causes MSBuild project evaluation to take between 8-12 seconds extra for this alone.
+    3. UBT project parsing support exposes, but nothing accesses, the collection of compiled items, so there's no actual benefit to having the items.
+    This is defaulted to still include items for now, in case licensees include plugins that depend on this, but should be considered obsolete as of UE 5.8. -->
+  <ItemGroup Condition="'$(EpicGamesIncludeGauntletPlugins)' != 'false'">
     <Compile Include="../../../../Plugins/**/Build/Scripts/Gauntlet/**/*.cs">
       <!-- Shorten the path so that everything is under a Plugins folder but does not include the Build/Scripts/Gauntlet part -->
-      <Link>$([System.Text.RegularExpressions.Regex]::Replace(%(FullPath), '^.*?(?=Plugins[\\/])|Build[\\/]Scripts[\\/]Gauntlet[\\/]',''))</Link>      
-    </Compile> 
+      <Link>$([System.Text.RegularExpressions.Regex]::Replace(%(FullPath), '^.*?(?=Plugins[\\/])|Build[\\/]Scripts[\\/]Gauntlet[\\/]',''))</Link>
+    </Compile>
+  </ItemGroup>
+
+  <ItemGroup>
     <Compile Include="../../../../Platforms/*/Plugins/**/Build/Scripts/Gauntlet/**/*.cs">
       <!-- RecursiveDir metadata expands to [PlatformName]/Source/Programs/[etc]. The Replace() will
       replace everything from the first slash to the end leaving just the platform name. Should the
diff --git a/Engine/Source/Programs/Shared/EpicGames.Build/DotnetProjectInfo.cs b/Engine/Source/Programs/Shared/EpicGames.Build/DotnetProjectInfo.cs
--- a/Engine/Source/Programs/Shared/EpicGames.Build/DotnetProjectInfo.cs
+++ b/Engine/Source/Programs/Shared/EpicGames.Build/DotnetProjectInfo.cs
@@ -119,10 +119,19 @@ public sealed record class DotnetProjectInfo : IDotnetProjectInfo
 	private static DotnetProjectInfo ReadImpl(FileReference file, IReadOnlyDictionary<string, string> globalProperties)
 	{
 		using ProjectCollection projectCollection = new();
+
+		// See Gauntlet.Automation.cs for context.
+		const string includeGauntletPluginsKey = "EpicGamesIncludeGauntletPlugins";
+		Dictionary<string, string> globalPropertiesWithDefaults = new(globalProperties);
+		if (!globalProperties.ContainsKey(includeGauntletPluginsKey))
+		{
+			globalPropertiesWithDefaults[includeGauntletPluginsKey] = "false";
+		}
+
 		ProjectOptions projectOptions = new()
 		{
-			GlobalProperties = globalProperties.ToDictionary(),
-			LoadSettings = ProjectLoadSettings.IgnoreMissingImports | ProjectLoadSettings.IgnoreInvalidImports,
+			GlobalProperties = globalPropertiesWithDefaults,
+			LoadSettings = ProjectLoadSettings.IgnoreMissingImports | ProjectLoadSettings.IgnoreInvalidImports | ProjectLoadSettings.DoNotEvaluateElementsWithFalseCondition,
 			ProjectCollection = projectCollection,
 		};
 		Project project = Project.FromFile(file.FullName, projectOptions);
