Properly escape character hte quotes in the JSON for environment variables set in xcrun to avoid error

dataCorrupted(Swift.DecodingError.Context(codingPath: [], debugDescription: "The given data was not valid JSON."

when being launched as C# process as the lack of escape causes the " to get omitted on iOS quick launch.
Additionally combine environment variables into single argument.

#Synced-CL 49329359

diff --git a/Engine/Source/Programs/AutomationTool/IOS/IOSPlatform.Automation.cs b/Engine/Source/Programs/AutomationTool/IOS/IOSPlatform.Automation.cs
--- a/Engine/Source/Programs/AutomationTool/IOS/IOSPlatform.Automation.cs
+++ b/Engine/Source/Programs/AutomationTool/IOS/IOSPlatform.Automation.cs
@@ -1236,10 +1236,9 @@ public class IOSPlatform : ApplePlatform
 				// For now, update the app launch implementation to use devicectl. Windows may need to switch to a different 3rd party implementation.
 				Program = "xcrun";
 				Arguments = "devicectl device process launch";
-				Arguments += " -e '{\\\"APP_DISTRIBUTOR_ID_OVERRIDE\\\": \\\"com.apple.AppStore\\\"}'";
 				Arguments += " --terminate-existing";	// if it's already running on device, kill it first
 				Arguments += " --console";	// attach to the console so we can get log output
-				Arguments += " --environment-variables '{\"OS_ACTIVITY_DT_MODE\":\"enable\"}'";	// ask os_log to go into stdout/stderr
+				Arguments += " --environment-variables '{\\\"APP_DISTRIBUTOR_ID_OVERRIDE\\\":\\\"com.apple.AppStore\\\",\\\"OS_ACTIVITY_DT_MODE\\\":\\\"enable\\\"}'";	// ask os_log to go into stdout/stderr
 				Arguments += " --device " + Params.DeviceNames[0];
 				Arguments += " \"" + BundleIdentifier + "\" --";
 			}
