Expose HordeCreateReport to BGV2

#rnx

diff --git a/Engine/Source/Programs/AutomationTool/BuildGraph/Tasks/HordeCreateReportTask.cs b/Engine/Source/Programs/AutomationTool/BuildGraph/Tasks/HordeCreateReportTask.cs
--- a/Engine/Source/Programs/AutomationTool/BuildGraph/Tasks/HordeCreateReportTask.cs
+++ b/Engine/Source/Programs/AutomationTool/BuildGraph/Tasks/HordeCreateReportTask.cs
@@ -6,6 +6,7 @@ using System.Linq;
 using System.Text.Json;
 using System.Threading.Tasks;
 using System.Xml;
+using AutomationTool;
 using EpicGames.Core;
 using Microsoft.Extensions.Logging;
 
@@ -109,4 +110,27 @@ namespace AutomationTool.Tasks
 		/// <returns>The tag names which are modified by this task</returns>
 		public override IEnumerable<string> FindProducedTagNames() => Enumerable.Empty<string>();
 	}
-}
+
+	public static partial class StandardTasks
+	{
+		/// <summary>
+		/// Compiles a target
+		/// </summary>
+		/// <param name="name">Name for the report</param>
+		/// <param name="scope"> Where to display the report</param>
+		/// <param name="placement">Where to show the report</param>
+		/// <param name="text">Text to be displayed</param>
+		/// <returns>Creates a Horde report file, which will be displayed on the dashboard with any job running this task.</returns>
+		public static async Task<FileSet> HordeCreateReportAsync(string name, string scope, string placement, string text)
+		{
+			HordeCreateReportTaskParameters parameters = new()
+			{
+				Name = name,
+				Scope = scope,
+				Placement = placement,
+				Text = text
+			};
+			return await ExecuteAsync(new HordeCreateReportTask(parameters));
+		}
+	}
+}
\ No newline at end of file
