[Gauntlet] Restore iOS OS_Activity mode environment variable. This is added to ensure os_log properly outputs to process stdout

diff --git a/Engine/Source/Programs/AutomationTool/Gauntlet/Platform/IOS/Gauntlet.TargetDeviceIOS.cs b/Engine/Source/Programs/AutomationTool/Gauntlet/Platform/IOS/Gauntlet.TargetDeviceIOS.cs
--- a/Engine/Source/Programs/AutomationTool/Gauntlet/Platform/IOS/Gauntlet.TargetDeviceIOS.cs
+++ b/Engine/Source/Programs/AutomationTool/Gauntlet/Platform/IOS/Gauntlet.TargetDeviceIOS.cs
@@ -816,7 +816,7 @@ namespace Gauntlet
 				// --terminate-existing		exit any existing forms of the app
 				// --console				pipe process stdout to devicectl, also holds onto the process until exit instead of immediate exit
 				// --device					device to run on
-				string LaunchCommand = string.Format("device process launch --terminate-existing --console --device {0} {1}", UUID, IOSInstall.PackageName);
+				string LaunchCommand = $"device process launch --terminate-existing --console --environment-variables {{\"OS_ACTIVITY_DT_MODE\":\"enable\"}} --device {UUID} {IOSInstall.PackageName}";
 				AppProcess = ExecuteDevicectlCommandNoWait(LaunchCommand, AdditionalOptions: ERunOptions.SpewIsVerbose, LocalCache: IOSInstall.Device.LocalCachePath);
 			}
 			else
