Fix nullability attribute on FileReference.FromString. Change the condition to `IsNullOrWhiteSpace` in order to make sure the construction does not fail. Ensure that all uses of this method are correct. For uses where we know the passed-in value is never null or whitespace, switch to using the constructor directly.

#Synced-CL 49646571

diff --git a/Engine/Source/Programs/AutomationTool/AutomationUtils/ScriptManager.cs b/Engine/Source/Programs/AutomationTool/AutomationUtils/ScriptManager.cs
--- a/Engine/Source/Programs/AutomationTool/AutomationUtils/ScriptManager.cs
+++ b/Engine/Source/Programs/AutomationTool/AutomationUtils/ScriptManager.cs
@@ -82,7 +82,12 @@ namespace AutomationTool
 
 			foreach (Assembly CompiledAssembly in AllScriptAssemblies)
 			{
-				DirectoryReference AssemblyDirectory = FileReference.FromString(CompiledAssembly.Location).Directory;
+				if (string.IsNullOrWhiteSpace(CompiledAssembly.Location))
+				{
+					continue;
+				}
+
+				DirectoryReference AssemblyDirectory = new FileReference(CompiledAssembly.Location).Directory;
 				AssemblyDirectory = DirectoryReference.FindCorrectCase(AssemblyDirectory);
 
 				if (OutputDirs.Add(AssemblyDirectory))
diff --git a/Engine/Source/Programs/AutomationTool/AutomationUtils/UnrealBuild.cs b/Engine/Source/Programs/AutomationTool/AutomationUtils/UnrealBuild.cs
--- a/Engine/Source/Programs/AutomationTool/AutomationUtils/UnrealBuild.cs
+++ b/Engine/Source/Programs/AutomationTool/AutomationUtils/UnrealBuild.cs
@@ -766,7 +766,7 @@ namespace AutomationTool
 		}
 
 		public static FileReference UnrealBuildToolDll => 
-			FileReference.FromString(CommandUtils.CombinePaths(CommandUtils.CmdEnv.LocalRoot, "Engine/Binaries/DotNET/UnrealBuildTool/UnrealBuildTool.dll"));
+			new(CommandUtils.CombinePaths(CommandUtils.CmdEnv.LocalRoot, "Engine/Binaries/DotNET/UnrealBuildTool/UnrealBuildTool.dll"));
 
 		// List of everything we built so far
 		public readonly HashSet<string> BuildProductFiles = new HashSet<string>(StringComparer.InvariantCultureIgnoreCase);
diff --git a/Engine/Source/Programs/AutomationTool/AutomationUtils/WindowsHostPlatform.cs b/Engine/Source/Programs/AutomationTool/AutomationUtils/WindowsHostPlatform.cs
--- a/Engine/Source/Programs/AutomationTool/AutomationUtils/WindowsHostPlatform.cs
+++ b/Engine/Source/Programs/AutomationTool/AutomationUtils/WindowsHostPlatform.cs
@@ -22,7 +22,7 @@ namespace AutomationTool
 				try
 				{
 					// Look for visual studio msbuild
-					FileReference msbuild = FileReference.FromString(WindowsExports.GetMSBuildToolPath());
+					FileReference msbuild = new(WindowsExports.GetMSBuildToolPath());
 					if (msbuild != null && FileReference.Exists(msbuild))
 					{
 						CachedFrameworkMsbuildTool = msbuild.FullName;
diff --git a/Engine/Source/Programs/AutomationTool/BuildGraph/BuildGraph.cs b/Engine/Source/Programs/AutomationTool/BuildGraph/BuildGraph.cs
--- a/Engine/Source/Programs/AutomationTool/BuildGraph/BuildGraph.cs
+++ b/Engine/Source/Programs/AutomationTool/BuildGraph/BuildGraph.cs
@@ -469,9 +469,9 @@ namespace AutomationTool
 			{
 				// Import schema if one is passed in
 				BgScriptSchema? schema;
-				if (importSchemaFileName != null)
+				if (!String.IsNullOrWhiteSpace(importSchemaFileName))
 				{
-					schema = BgScriptSchema.Import(FileReference.FromString(importSchemaFileName));
+					schema = BgScriptSchema.Import(new(importSchemaFileName));
 				}
 				else
 				{
diff --git a/Engine/Source/Programs/AutomationTool/IOS/IOSPlatform.Automation.cs b/Engine/Source/Programs/AutomationTool/IOS/IOSPlatform.Automation.cs
--- a/Engine/Source/Programs/AutomationTool/IOS/IOSPlatform.Automation.cs
+++ b/Engine/Source/Programs/AutomationTool/IOS/IOSPlatform.Automation.cs
@@ -1024,11 +1024,12 @@ public class IOSPlatform : ApplePlatform
 
 		if (Params.DeviceNames[0].Contains(AppleToolChainSettings.LocalMacUUID) || Params.bMakeMacNative)
 		{
+			FileReference AppToDeployFileReference = new(AppToDeploy);
 			// Deploy on local mac, not on iOS device
 			String MakeMacNativeSHPath = FileReference.Combine(Unreal.EngineDirectory, "Build/BatchFiles/Mac/MakeMacNative.sh").FullName;
-			String Arg1 = FileReference.FromString(AppToDeploy).Directory.FullName;														// Path to staged ios folder
-			String Arg2 = FileReference.Combine(FileReference.FromString(AppToDeploy).Directory.ParentDirectory, "iOSonMac").FullName;	// Path to destination folder
-			String Arg3 = FileReference.FromString(AppToDeploy).GetFileName();															// App name
+			String Arg1 = AppToDeployFileReference.Directory.FullName;														// Path to staged ios folder
+			String Arg2 = FileReference.Combine(AppToDeployFileReference.Directory.ParentDirectory, "iOSonMac").FullName;	// Path to destination folder
+			String Arg3 = AppToDeployFileReference.GetFileName();															// App name
 			Logger.LogInformation("Generating {1}/{2} to be run natively on an Apple Silicon Mac", Arg2, Arg3);
 			String Result = Utils.RunLocalProcessAndReturnStdOut("/bin/sh", $"\"{MakeMacNativeSHPath}\" \"{Arg1}\" \"{Arg2}\" \"{Arg3}\"");
 			if (Result.Length > 0)
diff --git a/Engine/Source/Programs/AutomationTool/Scripts/BuildPluginCommand.Automation.cs b/Engine/Source/Programs/AutomationTool/Scripts/BuildPluginCommand.Automation.cs
--- a/Engine/Source/Programs/AutomationTool/Scripts/BuildPluginCommand.Automation.cs
+++ b/Engine/Source/Programs/AutomationTool/Scripts/BuildPluginCommand.Automation.cs
@@ -38,7 +38,7 @@ public sealed class BuildPlugin : BuildCommand
 	{
 		// See if an engine dir was specified, fall back on default if not.
 		DirectoryReference EngineDirParam = ParseOptionalDirectoryReferenceParam("EngineDir");
-		UnrealBuildToolDll = EngineDirParam == null ? UnrealBuild.UnrealBuildToolDll : FileReference.FromString(CommandUtils.CombinePaths(EngineDirParam.ToString(), UnrealBuildToolDllRelativePath));
+		UnrealBuildToolDll = EngineDirParam == null ? UnrealBuild.UnrealBuildToolDll : new(CommandUtils.CombinePaths(EngineDirParam.ToString(), UnrealBuildToolDllRelativePath));
 
 		// Get the plugin filename
 		string PluginParam = ParseParamValue("Plugin");
diff --git a/Engine/Source/Programs/AutomationTool/Scripts/CookDiffReportHelper.cs b/Engine/Source/Programs/AutomationTool/Scripts/CookDiffReportHelper.cs
--- a/Engine/Source/Programs/AutomationTool/Scripts/CookDiffReportHelper.cs
+++ b/Engine/Source/Programs/AutomationTool/Scripts/CookDiffReportHelper.cs
@@ -99,7 +99,7 @@ namespace AutomationTool
 		private IEnumerable<FileReference> FindSourceJsons(string Source)
 		{
 			// source is a folder
-			DirectoryReference SourceDirectory = DirectoryReference.FromString(Source);
+			DirectoryReference SourceDirectory = new(Source);
 			if (DirectoryReference.Exists(SourceDirectory))
 			{
 				return DirectoryReference.EnumerateFiles( SourceDirectory, "*.json" );
@@ -111,7 +111,7 @@ namespace AutomationTool
 
 
 			// source is a file
-			FileReference JsonFile = FileReference.FromString(Source);
+			FileReference JsonFile = new(Source);
 
 			// for multiprocess cook, additional json files may have been created alongside this one, myfile-1.json, myfile-2.json etc
 			string BaseFileName = JsonFile.GetFileNameWithoutExtension();
diff --git a/Engine/Source/Programs/Shared/EpicGames.Build/System/CompileScriptModules.cs b/Engine/Source/Programs/Shared/EpicGames.Build/System/CompileScriptModules.cs
--- a/Engine/Source/Programs/Shared/EpicGames.Build/System/CompileScriptModules.cs
+++ b/Engine/Source/Programs/Shared/EpicGames.Build/System/CompileScriptModules.cs
@@ -442,7 +442,7 @@ namespace UnrealBuildBase
 					if (BuildRecord != null && BuildRecord.ProjectPath != null)
 					{
 						CsProjBuildRecordEntry Entry = new CsProjBuildRecordEntry(
-							FileReference.FromString(Path.GetFullPath(BuildRecord.ProjectPath, JsonFile.Directory.FullName)),
+							new(Path.GetFullPath(BuildRecord.ProjectPath, JsonFile.Directory.FullName)),
 							JsonFile,
 							BuildRecord);
 						LoadedBuildRecords.Add(Entry.ProjectFile, Entry);
@@ -605,7 +605,7 @@ namespace UnrealBuildBase
 			// Are all referenced build records valid?
 			foreach (CsProjBuildRecordRef ReferencedProject in Entry.BuildRecord.ProjectReferencesAndTimes)
 			{
-				FileReference FullProjectPath = FileReference.FromString(Path.GetFullPath(ReferencedProject.ProjectPath, ProjectPath.Directory.FullName));
+				FileReference FullProjectPath = new(Path.GetFullPath(ReferencedProject.ProjectPath, ProjectPath.Directory.FullName));
 				ValidateBuildRecordRecursively(BuildRecords, FullProjectPath, Hook, Logger);
 
 				if (!BuildRecords.TryGetValue(FullProjectPath, out CsProjBuildRecordEntry? RefEntry))
diff --git a/Engine/Source/Programs/Shared/EpicGames.Build/System/FileHasher.cs b/Engine/Source/Programs/Shared/EpicGames.Build/System/FileHasher.cs
--- a/Engine/Source/Programs/Shared/EpicGames.Build/System/FileHasher.cs
+++ b/Engine/Source/Programs/Shared/EpicGames.Build/System/FileHasher.cs
@@ -76,7 +76,7 @@ namespace UnrealBuildBase
 		public void PurgeStale()
 		{
 			int count = CachedDigests.Count;
-			IEnumerable<KeyValuePair<string, CachedDigest>> valid = CachedDigests.Where((x) => x.Value.UpToDate(FileItem.GetItemByFileReference(FileReference.FromString(x.Key))));
+			IEnumerable<KeyValuePair<string, CachedDigest>> valid = CachedDigests.Where((x) => x.Value.UpToDate(FileItem.GetItemByFileReference(new(x.Key))));
 			CachedDigests = new(valid);
 			logger?.LogDebug("Purged {Entries} stale cached digest entries", count - CachedDigests.Count);
 		}
diff --git a/Engine/Source/Programs/Shared/EpicGames.Core/DirectoryReference.cs b/Engine/Source/Programs/Shared/EpicGames.Core/DirectoryReference.cs
--- a/Engine/Source/Programs/Shared/EpicGames.Core/DirectoryReference.cs
+++ b/Engine/Source/Programs/Shared/EpicGames.Core/DirectoryReference.cs
@@ -288,11 +288,11 @@ namespace EpicGames.Core
 		public static DirectoryReference? FromFile(FileReference? file) => file?.Directory;
 
 		/// <summary>
-		/// Create a DirectoryReference from a string. If the string is null, returns a null DirectoryReference.
+		/// Create a DirectoryReference from a string. If the string is null or whitespace, returns a null DirectoryReference.
 		/// </summary>
 		/// <param name="directoryName">Path for the new object</param>
 		/// <returns>Returns a FileReference representing the given string, or null.</returns>
-		public static DirectoryReference? FromString(string? directoryName) => String.IsNullOrEmpty(directoryName) ? null : new(directoryName);
+		public static DirectoryReference? FromString(string? directoryName) => String.IsNullOrWhiteSpace(directoryName) ? null : new(directoryName);
 
 		/// <summary>
 		/// Finds the correct case to match the location of this file on disk. Uses the given case for parts of the path that do not exist.
diff --git a/Engine/Source/Programs/Shared/EpicGames.Core/FileReference.cs b/Engine/Source/Programs/Shared/EpicGames.Core/FileReference.cs
--- a/Engine/Source/Programs/Shared/EpicGames.Core/FileReference.cs
+++ b/Engine/Source/Programs/Shared/EpicGames.Core/FileReference.cs
@@ -68,12 +68,11 @@ namespace EpicGames.Core
 		}
 
 		/// <summary>
-		/// Create a FileReference from a string. If the string is null, returns a null FileReference.
+		/// Create a FileReference from a string. If the string is null or whitespace, returns a null FileReference.
 		/// </summary>
 		/// <param name="fileName">FileName for the string</param>
 		/// <returns>Returns a FileReference representing the given string, or null.</returns>
-		[return: NotNullIfNotNull(nameof(fileName))]
-		public static FileReference? FromString(string? fileName) => String.IsNullOrEmpty(fileName) ? null : new(fileName);
+		public static FileReference? FromString(string? fileName) => String.IsNullOrWhiteSpace(fileName) ? null : new(fileName);
 
 		/// <summary>
 		/// Gets the file name without path information
diff --git a/Engine/Source/Programs/Shared/EpicGames.MsBuild/CsProjBuilder.cs b/Engine/Source/Programs/Shared/EpicGames.MsBuild/CsProjBuilder.cs
--- a/Engine/Source/Programs/Shared/EpicGames.MsBuild/CsProjBuilder.cs
+++ b/Engine/Source/Programs/Shared/EpicGames.MsBuild/CsProjBuilder.cs
@@ -173,7 +173,7 @@ namespace EpicGames.MsBuild
 					projects.TryAdd(foundProject, project);
 					foreach (ProjectItem reference in project.GetItems("ProjectReference"))
 					{
-						referencedProjects.Add(FileReference.FromString(Path.GetFullPath(reference.EvaluatedInclude, foundProject.Directory.FullName)));
+						referencedProjects.Add(new(Path.GetFullPath(reference.EvaluatedInclude, foundProject.Directory.FullName)));
 					}
 				});
 
@@ -327,7 +327,7 @@ namespace EpicGames.MsBuild
 			{
 				foreach (CsProjBuildRecordRef referencedProject in entry.BuildRecord.ProjectReferencesAndTimes)
 				{
-					FileReference refProjectPath = FileReference.FromString(Path.GetFullPath(referencedProject.ProjectPath, entry.ProjectFile.Directory.FullName));
+					FileReference refProjectPath = new(Path.GetFullPath(referencedProject.ProjectPath, entry.ProjectFile.Directory.FullName));
 					if (buildRecords.TryGetValue(refProjectPath, out CsProjBuildRecordEntry refEntry))
 					{
 						referencedProject.TargetBuildTime = refEntry.BuildRecord.TargetBuildTime;
diff --git a/Engine/Source/Programs/TextureBuildWorker/TextureBuildWorker.Build.cs b/Engine/Source/Programs/TextureBuildWorker/TextureBuildWorker.Build.cs
--- a/Engine/Source/Programs/TextureBuildWorker/TextureBuildWorker.Build.cs
+++ b/Engine/Source/Programs/TextureBuildWorker/TextureBuildWorker.Build.cs
@@ -38,11 +38,11 @@ public abstract class TextureBuildWorkerTarget : DerivedDataBuildWorkerTarget
 		string OodleTextureFormatModule = "TextureFormatOodle";
 		if (bProjectOodle)
 		{
-			if (Target.ProjectFile != null && !String.IsNullOrEmpty(ProjectOodlePlugin))
+			if (Target.ProjectFile != null && !String.IsNullOrWhiteSpace(ProjectOodlePlugin))
 			{
 				ProjectOodlePlugin = Path.Combine(Target.ProjectFile.Directory.ToString(), ProjectOodlePlugin);
 
-				OodleUPluginFile = FileReference.FromString(ProjectOodlePlugin);
+				OodleUPluginFile = new(ProjectOodlePlugin);
 				OodlePluginType = PluginType.Project;
 			}
 			else
diff --git a/Engine/Source/Programs/UnrealBuildTool/Configuration/UEBuildModule.cs b/Engine/Source/Programs/UnrealBuildTool/Configuration/UEBuildModule.cs
--- a/Engine/Source/Programs/UnrealBuildTool/Configuration/UEBuildModule.cs
+++ b/Engine/Source/Programs/UnrealBuildTool/Configuration/UEBuildModule.cs
@@ -474,7 +474,7 @@ namespace UnrealBuildTool
 		{
 			_ = InEnumerableStrings.TryGetNonEnumeratedCount(out int estimatedCount);
 			HashSet<string> hashSet = new(estimatedCount);
-			foreach (string s in InEnumerableStrings)
+			foreach (string s in InEnumerableStrings.Where(str => !string.IsNullOrWhiteSpace(str)))
 			{
 				hashSet.Add(ExpandPathVariables(s, null, null));
 			}
diff --git a/Engine/Source/Programs/UnrealBuildTool/Configuration/Xml/XmlConfig.cs b/Engine/Source/Programs/UnrealBuildTool/Configuration/Xml/XmlConfig.cs
--- a/Engine/Source/Programs/UnrealBuildTool/Configuration/Xml/XmlConfig.cs
+++ b/Engine/Source/Programs/UnrealBuildTool/Configuration/Xml/XmlConfig.cs
@@ -910,7 +910,7 @@ namespace UnrealBuildTool
 			}
 			else if (fieldType == typeof(FileReference))
 			{
-				return FileReference.FromString(text);
+				return new FileReference(text);
 			}
 			else
 			{
diff --git a/Engine/Source/Programs/UnrealBuildTool/Modes/BuildMode.cs b/Engine/Source/Programs/UnrealBuildTool/Modes/BuildMode.cs
--- a/Engine/Source/Programs/UnrealBuildTool/Modes/BuildMode.cs
+++ b/Engine/Source/Programs/UnrealBuildTool/Modes/BuildMode.cs
@@ -138,7 +138,7 @@ namespace UnrealBuildTool
 			// Create the log file, and flush the startup listener to it
 			if (!Arguments.HasOption("-NoLog") && !Log.HasFileWriter())
 			{
-				Log.AddFileWriter("DefaultLogTraceListener", FileReference.FromString(BaseLogFileName));
+				Log.AddFileWriter("DefaultLogTraceListener", new(BaseLogFileName));
 			}
 			else
 			{
diff --git a/Engine/Source/Programs/UnrealBuildTool/Modes/FixIncludePathsMode.cs b/Engine/Source/Programs/UnrealBuildTool/Modes/FixIncludePathsMode.cs
--- a/Engine/Source/Programs/UnrealBuildTool/Modes/FixIncludePathsMode.cs
+++ b/Engine/Source/Programs/UnrealBuildTool/Modes/FixIncludePathsMode.cs
@@ -206,7 +206,7 @@ namespace UnrealBuildTool
 												string Path = System.IO.Path.GetFullPath(System.IO.Path.Combine(IncludePath.FullName, Include));
 												if (System.IO.File.Exists(Path))
 												{
-													FoundIncludeFile = FileReference.FromString(Path);
+													FoundIncludeFile = new(Path);
 													FoundIncludePath = IncludePath;
 													break;
 												}
@@ -264,7 +264,7 @@ namespace UnrealBuildTool
 															PreferredInclude = null;
 
 															string ModulePath = FullPath;
-															FileReference IncludeFileReference = FileReference.FromString(FullPath);
+															FileReference IncludeFileReference = new(FullPath);
 															DirectoryReference? TempDirectory = IncludeFileReference.Directory;
 															DirectoryReference? FoundDirectory = null;
 															// find the module this include is part of
diff --git a/Engine/Source/Programs/UnrealBuildTool/Modes/IWYUMode.cs b/Engine/Source/Programs/UnrealBuildTool/Modes/IWYUMode.cs
--- a/Engine/Source/Programs/UnrealBuildTool/Modes/IWYUMode.cs
+++ b/Engine/Source/Programs/UnrealBuildTool/Modes/IWYUMode.cs
@@ -388,7 +388,7 @@ namespace UnrealBuildTool
 			// Create the log file, and flush the startup listener to it
 			if (!Arguments.HasOption("-NoLog") && !Log.HasFileWriter())
 			{
-				Log.AddFileWriter("DefaultLogTraceListener", FileReference.FromString(IWYUBaseLogFileName));
+				Log.AddFileWriter("DefaultLogTraceListener", new(IWYUBaseLogFileName));
 			}
 			else
 			{
@@ -1375,7 +1375,7 @@ namespace UnrealBuildTool
 
 			if (!Nested)
 			{
-				FileReference FileReference = FileReference.Combine(FileReference.FromString(From).Directory, Include);
+				FileReference FileReference = FileReference.Combine(new FileReference(From).Directory, Include);
 				FileItem FileItem = FileItem.GetItemByFileReference(FileReference);
 				if (FileItem.Exists)
 				{
diff --git a/Engine/Source/Programs/UnrealBuildTool/Modes/PipInstallMode.cs b/Engine/Source/Programs/UnrealBuildTool/Modes/PipInstallMode.cs
--- a/Engine/Source/Programs/UnrealBuildTool/Modes/PipInstallMode.cs
+++ b/Engine/Source/Programs/UnrealBuildTool/Modes/PipInstallMode.cs
@@ -599,8 +599,8 @@ namespace UnrealBuildTool.Modes
 			string[] PythonPlugins = FileReference.ReadAllLines(PluginsListingFile);
 			foreach (string PluginPath in PythonPlugins)
 			{
-				FileReference PluginFile = FileReference.FromString(PluginPath);
-				if (!FileReference.Exists(PluginFile))
+				FileReference? PluginFile = FileReference.FromString(PluginPath);
+				if (PluginFile is null || !FileReference.Exists(PluginFile))
 				{
 					continue;
 				}
diff --git a/Engine/Source/Programs/UnrealBuildTool/Platform/Android/AndroidToolChain.cs b/Engine/Source/Programs/UnrealBuildTool/Platform/Android/AndroidToolChain.cs
--- a/Engine/Source/Programs/UnrealBuildTool/Platform/Android/AndroidToolChain.cs
+++ b/Engine/Source/Programs/UnrealBuildTool/Platform/Android/AndroidToolChain.cs
@@ -270,7 +270,7 @@ namespace UnrealBuildTool
 
 		protected override ClangToolChainInfo GetToolChainInfo()
 		{
-			return new ClangToolChainInfo(DirectoryReference.FromString(AndroidPlatformSDK.GetNDKRoot(Logger)), FileReference.FromString(ClangPath)!, FileReference.FromString(ArPathArm64)!, Logger);
+			return new ClangToolChainInfo(DirectoryReference.FromString(AndroidPlatformSDK.GetNDKRoot(Logger)), new(ClangPath!), new(ArPathArm64!), Logger);
 		}
 
 		public static string GetGLESVersion(bool bBuildForES31)
diff --git a/Engine/Source/Programs/UnrealBuildTool/Platform/Clang/ClangSpecificFileAction.cs b/Engine/Source/Programs/UnrealBuildTool/Platform/Clang/ClangSpecificFileAction.cs
--- a/Engine/Source/Programs/UnrealBuildTool/Platform/Clang/ClangSpecificFileAction.cs
+++ b/Engine/Source/Programs/UnrealBuildTool/Platform/Clang/ClangSpecificFileAction.cs
@@ -122,7 +122,7 @@ namespace UnrealBuildTool
 				Action.ProducedItems.Add(PreprocessedItem);
 			}
 
-			FileReference ResponseFile = RootPaths.GetLocalPath(FileReference.FromString(Action.CommandArguments.Substring(1).Trim('"')));
+			FileReference ResponseFile = RootPaths.GetLocalPath(new(Action.CommandArguments.Substring(1).Trim('"')));
 			File.WriteAllLines(ResponseFile.FullName, NewRspLines);
 
 			return Action;
diff --git a/Engine/Source/Programs/UnrealBuildTool/Platform/IOS/IOSToolChain.cs b/Engine/Source/Programs/UnrealBuildTool/Platform/IOS/IOSToolChain.cs
--- a/Engine/Source/Programs/UnrealBuildTool/Platform/IOS/IOSToolChain.cs
+++ b/Engine/Source/Programs/UnrealBuildTool/Platform/IOS/IOSToolChain.cs
@@ -382,7 +382,7 @@ namespace UnrealBuildTool
 				{
 					if (AdditionalLibrary.Contains(".framework/", StringComparison.Ordinal))
 					{
-						LinkArguments.Add($"-l\"{NormalizeCommandLinePath(FileReference.FromString(AdditionalLibrary), LinkEnvironment.RootPaths)}\"");
+						LinkArguments.Add($"-l\"{NormalizeCommandLinePath(new FileReference(AdditionalLibrary), LinkEnvironment.RootPaths)}\"");
 					}
 					else if (String.IsNullOrEmpty(Path.GetDirectoryName(AdditionalLibrary)) && String.IsNullOrEmpty(Path.GetExtension(AdditionalLibrary)))
 					{
@@ -390,9 +390,9 @@ namespace UnrealBuildTool
 					}
 					else
 					{
-                        FileReference LibraryPath = FileReference.FromString(AdditionalLibrary);
+						FileReference LibraryPath = new(AdditionalLibrary);
 						LinkArguments.Add($"-l\"{NormalizeCommandLinePath(LibraryPath, LinkEnvironment.RootPaths)}\"");
-                    }					
+					}
 				}
 
 				foreach (FileReference LibraryPath in LinkEnvironment.Libraries)
@@ -460,7 +460,7 @@ namespace UnrealBuildTool
 			FileReference ResponseFileName;
 			if (bBuildImportLibraryOnly)
 			{
-				ResponseFileName = FileReference.FromString(OutputFile.FullName + ResponseExt);
+				ResponseFileName = new(OutputFile.FullName + ResponseExt);
 			}
 			else
 			{
diff --git a/Engine/Source/Programs/UnrealBuildTool/Platform/Mac/MacToolChain.cs b/Engine/Source/Programs/UnrealBuildTool/Platform/Mac/MacToolChain.cs
--- a/Engine/Source/Programs/UnrealBuildTool/Platform/Mac/MacToolChain.cs
+++ b/Engine/Source/Programs/UnrealBuildTool/Platform/Mac/MacToolChain.cs
@@ -766,7 +766,7 @@ namespace UnrealBuildTool
 			FileReference ResponseFileName;
 			if (bBuildImportLibraryOnly)
 			{
-				ResponseFileName = FileReference.FromString(OutputFile.FullName + ResponseExt);
+				ResponseFileName = new(OutputFile.FullName + ResponseExt);
 			}
 			else
 			{
diff --git a/Engine/Source/Programs/UnrealBuildTool/Platform/Windows/PVSToolChain.cs b/Engine/Source/Programs/UnrealBuildTool/Platform/Windows/PVSToolChain.cs
--- a/Engine/Source/Programs/UnrealBuildTool/Platform/Windows/PVSToolChain.cs
+++ b/Engine/Source/Programs/UnrealBuildTool/Platform/Windows/PVSToolChain.cs
@@ -521,7 +521,7 @@ namespace UnrealBuildTool
 
 		public FileReference? UpdateFilePath(CppRootPaths rootPaths)
 		{
-			FileReference? fileReference = !String.IsNullOrWhiteSpace(File) ? FileReference.FromString(File) : null;
+			FileReference? fileReference = FileReference.FromString(File);
 			if (fileReference != null && rootPaths.bUseVfs)
 			{
 				fileReference = rootPaths.GetLocalPath(fileReference);
diff --git a/Engine/Source/Programs/UnrealBuildTool/Platform/Windows/VCToolChain.cs b/Engine/Source/Programs/UnrealBuildTool/Platform/Windows/VCToolChain.cs
--- a/Engine/Source/Programs/UnrealBuildTool/Platform/Windows/VCToolChain.cs
+++ b/Engine/Source/Programs/UnrealBuildTool/Platform/Windows/VCToolChain.cs
@@ -2558,7 +2558,7 @@ namespace UnrealBuildTool
 						FileReference ManifestFile = FileReference.Combine(Makefile.ProjectIntermediateDirectory, $"{Target.Name}TimingManifest.txt");
 						File.WriteAllLines(ManifestFile.FullName, TimingJsonFiles.Select(f => f.AbsolutePath));
 
-						FileReference ExpectedCompileTimeFile = FileReference.FromString(Path.Combine(Makefile.ProjectIntermediateDirectory.FullName, $"{Target.Name}.json"));
+						FileReference ExpectedCompileTimeFile = new(Path.Combine(Makefile.ProjectIntermediateDirectory.FullName, $"{Target.Name}.json"));
 						List<string> ActionArgs = new List<string>()
 						{
 							$"-Name={Target.Name}",
diff --git a/Engine/Source/Programs/UnrealBuildTool/ProjectFiles/VisualStudioCode/VSCodeProjectFileGenerator.cs b/Engine/Source/Programs/UnrealBuildTool/ProjectFiles/VisualStudioCode/VSCodeProjectFileGenerator.cs
--- a/Engine/Source/Programs/UnrealBuildTool/ProjectFiles/VisualStudioCode/VSCodeProjectFileGenerator.cs
+++ b/Engine/Source/Programs/UnrealBuildTool/ProjectFiles/VisualStudioCode/VSCodeProjectFileGenerator.cs
@@ -351,7 +351,7 @@ namespace UnrealBuildTool
 			if (OperatingSystem.IsWindows())
 			{
 				VCEnvironment Environment = VCEnvironment.Create(WindowsPlatform.GetDefaultCompiler(null, Target.Rules.WindowsPlatform.Architecture, Logger, true), WindowsCompiler.Default, Target.Platform, Target.Rules.WindowsPlatform.Architecture, null, null, Target.Rules.WindowsPlatform.WindowsSdkVersion, null, Target.Rules.WindowsPlatform.bUseCPPWinRT, Target.Rules.WindowsPlatform.bAllowClangLinker, Target.Rules.WindowsPlatform.bAllowRadLinker, Logger);
-				CompilerPath = FileReference.FromString(Environment.CompilerPath.FullName);
+				CompilerPath = Environment.CompilerPath;
 				UsingClang = false;
 			}
 			else if (OperatingSystem.IsLinux())
diff --git a/Engine/Source/Programs/UnrealBuildTool/System/CppRootPaths.cs b/Engine/Source/Programs/UnrealBuildTool/System/CppRootPaths.cs
--- a/Engine/Source/Programs/UnrealBuildTool/System/CppRootPaths.cs
+++ b/Engine/Source/Programs/UnrealBuildTool/System/CppRootPaths.cs
@@ -316,7 +316,7 @@ namespace UnrealBuildTool
 			{
 				if (reference.IsUnderDirectory(vfs))
 				{
-					return FileReference.FromString(local.FullName + reference.FullName.Substring(vfs.FullName.Length));
+					return new(local.FullName + reference.FullName.Substring(vfs.FullName.Length));
 				}
 			}
 
@@ -324,7 +324,7 @@ namespace UnrealBuildTool
 			{
 				if (reference.IsUnderDirectory(vfs))
 				{
-					return FileReference.FromString(local.FullName + reference.FullName.Substring(vfs.FullName.Length));
+					return new(local.FullName + reference.FullName.Substring(vfs.FullName.Length));
 				}
 			}
 
diff --git a/Engine/Source/Programs/UnrealBuildTool/System/HotReload.cs b/Engine/Source/Programs/UnrealBuildTool/System/HotReload.cs
--- a/Engine/Source/Programs/UnrealBuildTool/System/HotReload.cs
+++ b/Engine/Source/Programs/UnrealBuildTool/System/HotReload.cs
@@ -1038,7 +1038,7 @@ namespace UnrealBuildTool
 								FileItem OldSourceDependencyFileItem = Action.ProducedItems.First(x => x.Location == OldSourceDependencyFile);
 								NewAction.ProducedItems.Remove(OldSourceDependencyFileItem);
 
-								FileReference NewSourceDependencyFile = FileReference.FromString(OldSourceDependencyFile.FullName.Substring(0, OldSourceDependencyFile.FullName.Length - extension.Length) + extraExt + extension);
+								FileReference NewSourceDependencyFile = new(OldSourceDependencyFile.FullName.Substring(0, OldSourceDependencyFile.FullName.Length - extension.Length) + extraExt + extension);
 								FileItem NewSourceDependencyFileItem = FileItem.GetItemByFileReference(NewSourceDependencyFile);
 								NewAction.ProducedItems.Add(NewSourceDependencyFileItem);
 								NewAction.DependencyListFile = NewSourceDependencyFileItem;
diff --git a/Engine/Source/Programs/UnrealBuildTool/System/TargetMakefile.cs b/Engine/Source/Programs/UnrealBuildTool/System/TargetMakefile.cs
--- a/Engine/Source/Programs/UnrealBuildTool/System/TargetMakefile.cs
+++ b/Engine/Source/Programs/UnrealBuildTool/System/TargetMakefile.cs
@@ -810,7 +810,7 @@ namespace UnrealBuildTool
 				foreach (FileReference Manifest in Makefile.AdditionalArguments!
 					.Where(x => x.StartsWith("-Manifest=", StringComparison.OrdinalIgnoreCase))
 					.Select(x => FileReference.FromString(x.Substring("-Manifest=".Length).Trim('"').Trim('\'')))
-					.Where(x => x != null))
+					.OfType<FileReference>())
 				{
 					if (!FileReference.Exists(Manifest))
 					{
diff --git a/Engine/Source/Programs/UnrealBuildTool/ToolChain/ISPCToolChain.cs b/Engine/Source/Programs/UnrealBuildTool/ToolChain/ISPCToolChain.cs
--- a/Engine/Source/Programs/UnrealBuildTool/ToolChain/ISPCToolChain.cs
+++ b/Engine/Source/Programs/UnrealBuildTool/ToolChain/ISPCToolChain.cs
@@ -848,7 +848,7 @@ namespace UnrealBuildTool
 				// Generate the timing info
 				if (CompileEnvironment.bPrintTimingInfo)
 				{
-					FileItem TraceFile = FileItem.GetItemByFileReference(FileReference.FromString($"{CompiledISPCObjFileNoISA}.json"));
+					FileItem TraceFile = FileItem.GetItemByFileReference(new($"{CompiledISPCObjFileNoISA}.json"));
 					Arguments.Add("--time-trace");
 					CompileAction.ProducedItems.Add(TraceFile);
 				}
