BuildGraph: DotNetTask should only use specific version of dotnet

#rnx
#Synced-CL 49686086

diff --git a/Engine/Source/Programs/AutomationTool/BuildGraph/Tasks/DotNetTask.cs b/Engine/Source/Programs/AutomationTool/BuildGraph/Tasks/DotNetTask.cs
--- a/Engine/Source/Programs/AutomationTool/BuildGraph/Tasks/DotNetTask.cs
+++ b/Engine/Source/Programs/AutomationTool/BuildGraph/Tasks/DotNetTask.cs
@@ -1,9 +1,11 @@
 // Copyright Epic Games, Inc. All Rights Reserved.
 
 using System.Collections.Generic;
+using System.Runtime.Intrinsics.X86;
 using System.Threading.Tasks;
 using System.Xml;
 using EpicGames.Core;
+using Ionic.Zlib;
 using Microsoft.Extensions.Logging;
 using UnrealBuildBase;
 
@@ -76,13 +78,17 @@ namespace AutomationTool.Tasks
 		/// <param name="tagNameToFileSet">Mapping from tag names to the set of files they include</param>
 		public override async Task ExecuteAsync(JobContext job, HashSet<FileReference> buildProducts, Dictionary<string, HashSet<FileReference>> tagNameToFileSet)
 		{
-			FileReference dotNetFile = _parameters.DotNetPath == null ? Unreal.DotnetPath : _parameters.DotNetPath;
+			FileReference dotNetFile = _parameters.DotNetPath ?? Unreal.DotnetPath;
 			if (!FileReference.Exists(dotNetFile))
 			{
 				throw new AutomationException("DotNet is missing from {0}", dotNetFile);
 			}
 
-			IProcessResult result = await ExecuteAsync(dotNetFile.FullName, _parameters.Arguments, workingDir: _parameters.BaseDir, envVars: ParseEnvVars(_parameters.Environment, _parameters.EnvironmentFile));
+			Dictionary<string, string> envVars = ParseEnvVars(_parameters.Environment, _parameters.EnvironmentFile);
+			envVars["DOTNET_MULTILEVEL_LOOKUP"] = "0";
+			envVars["DOTNET_ROLL_FORWARD"] = "LatestMajor";
+
+			IProcessResult result = await ExecuteAsync(dotNetFile.FullName, _parameters.Arguments, workingDir: _parameters.BaseDir, envVars: envVars);
 			if (result.ExitCode < 0 || result.ExitCode >= _parameters.ErrorLevel)
 			{
 				Logger.LogError(KnownLogEvents.ExitCode, "Docker terminated with an exit code indicating an error ({ExitCode})", result.ExitCode);
