[Gauntlet] Trim quotes from LongProcessResult argument lists on unix platforms because the Process.ArgumentList tokenizer appends it's own

#Synced-CL 49698028

diff --git a/Engine/Source/Programs/AutomationTool/Gauntlet/Unreal/Base/Gauntlet.UnrealTestConfiguration.cs b/Engine/Source/Programs/AutomationTool/Gauntlet/Unreal/Base/Gauntlet.UnrealTestConfiguration.cs
--- a/Engine/Source/Programs/AutomationTool/Gauntlet/Unreal/Base/Gauntlet.UnrealTestConfiguration.cs
+++ b/Engine/Source/Programs/AutomationTool/Gauntlet/Unreal/Base/Gauntlet.UnrealTestConfiguration.cs
@@ -328,7 +328,7 @@ namespace Gauntlet
 
 		public static List<string> GetArgumentList(string Args)
 		{
-			return CommandLineRegex.Matches(Args).Select(Match => Match.Value).ToList();
+			return CommandLineRegex.Matches(Args).Select(Match => Match.Value.Replace("\"", string.Empty)).ToList();
 		}
 	}
 
