Initialization should fail when no test is run

#rnx

#Synced-CL 49727639

diff --git a/Engine/Source/Programs/AutomationTool/Gauntlet/Unreal/Base/Gauntlet.UnrealTestNode.cs b/Engine/Source/Programs/AutomationTool/Gauntlet/Unreal/Base/Gauntlet.UnrealTestNode.cs
--- a/Engine/Source/Programs/AutomationTool/Gauntlet/Unreal/Base/Gauntlet.UnrealTestNode.cs
+++ b/Engine/Source/Programs/AutomationTool/Gauntlet/Unreal/Base/Gauntlet.UnrealTestNode.cs
@@ -1833,8 +1833,8 @@ namespace Gauntlet
 				{
 					var InitPhase = SessionReport.GetPhase("Initialization");
 					InitPhase.GetStream().AddWarning("Could not find Unreal Automated test report at {FilePath}. No test report details available.", JsonReportPath);
-					// Check if there is any Engine Test Error or Fatal Error found to add more context
-					UnrealRoleResult RelevantRoleResult = RoleResults?.FirstOrDefault(I => I.ProcessResult == UnrealProcessResult.EngineTestError || I.ProcessResult == UnrealProcessResult.EncounteredFatalError);
+					// Check if there is any non-OK exit to add more context
+					UnrealRoleResult RelevantRoleResult = RoleResults?.FirstOrDefault(I => I.ProcessResult != UnrealProcessResult.ExitOk && I.ProcessResult != UnrealProcessResult.EncounteredEnsure);
 					if (RelevantRoleResult != null)
 					{
 						string RoleType = RelevantRoleResult.Artifacts.SessionRole.RoleType.ToString();
@@ -1845,8 +1845,8 @@ namespace Gauntlet
 						}
 						Stream.AddError("{Role}: {Summary}", RoleType, RelevantRoleResult.Summary);
 
-						// In case of a fatal error, it is an interrupted outcome
-						if (RelevantRoleResult.ProcessResult == UnrealProcessResult.EncounteredFatalError)
+						// In case of a fatal error or initialization failure, it is an interrupted outcome
+						if (RelevantRoleResult.ProcessResult == UnrealProcessResult.EncounteredFatalError || RelevantRoleResult.ProcessResult == UnrealProcessResult.InitializationFailure)
 						{
 							InitPhase.SetOutcome(TestPhaseOutcome.Interrupted);
 						}
