BuildGraph: Explicitly set .NET path envvars in DotNetTask

#rnx

#Synced-CL 49739777

diff --git a/Engine/Source/Programs/AutomationTool/AutomationUtils/ProcessUtils.cs b/Engine/Source/Programs/AutomationTool/AutomationUtils/ProcessUtils.cs
--- a/Engine/Source/Programs/AutomationTool/AutomationUtils/ProcessUtils.cs
+++ b/Engine/Source/Programs/AutomationTool/AutomationUtils/ProcessUtils.cs
@@ -1,5 +1,6 @@
 // Copyright Epic Games, Inc. All Rights Reserved.
 using System;
+using System.Collections;
 using System.Collections.Generic;
 using System.Linq;
 using System.Text;
@@ -987,18 +988,30 @@ namespace AutomationTool
 			if (!IsEngineExecutable(App))
 			{
 				Dictionary<string, string> FinalEnv = (Env != null) ? new Dictionary<string, string>(Env) : new Dictionary<string, string>();
-			
-				// We're clearing DOTNET_* for any external programs to avoid enforcing our version of .NET
-				// on applications that may not be compatible with it e.g. requiring a newer version that may be
-				// installed on the system as part of their installer.
-				FinalEnv["DOTNET_ROOT"] = "";
-				FinalEnv["DOTNET_MULTILEVEL_LOOKUP"] = "";
-				FinalEnv["DOTNET_ROLL_FORWARD"] = "";
 
-				// Clearing our variable as well to make sure that if the spawned tool ends up calling our .bat files, we will reinitialize the envvars.
-				FinalEnv["UE_DOTNET_VERSION"] = "";
+				// Environment variables to preserve if explicitly set, otherwise clear
+				string[] dotnetVars = [
+					"DOTNET_ROOT",
+					"DOTNET_MULTILEVEL_LOOKUP",
+					"DOTNET_ROLL_FORWARD",
+					"DOTNET_HOST_PATH",
+					"DOTNET_MSBUILD_SDK_RESOLVER_CLI_DIR",
+					"DOTNET_MSBUILD_SDK_RESOLVER_SDKS_DIR",
+					"DOTNET_MSBUILD_SDK_RESOLVER_SDKS_VER",
+				];
+
+				foreach (string varName in dotnetVars)
+				{
+					if (Env == null || !Env.ContainsKey(varName))
+					{
+						FinalEnv[varName] = String.Empty;
+					}
+				}
+
+				// Clearing our variable as well
+				FinalEnv["UE_DOTNET_VERSION"] = String.Empty;
 
-				Env = FinalEnv;					
+				Env = FinalEnv;
 			}
 
 			bool bUseShellExecute = Options.HasFlag(ERunOptions.UseShellExecute);
@@ -1053,6 +1066,7 @@ namespace AutomationTool
 				{
 					Proc.StartInfo.StandardOutputEncoding = new System.Text.UnicodeEncoding(false, false, false);
 				}
+
 				Proc.Start();
 
 				if (bRedirectStdOut)
diff --git a/Engine/Source/Programs/AutomationTool/BuildGraph/Tasks/DotNetTask.cs b/Engine/Source/Programs/AutomationTool/BuildGraph/Tasks/DotNetTask.cs
--- a/Engine/Source/Programs/AutomationTool/BuildGraph/Tasks/DotNetTask.cs
+++ b/Engine/Source/Programs/AutomationTool/BuildGraph/Tasks/DotNetTask.cs
@@ -1,6 +1,9 @@
 // Copyright Epic Games, Inc. All Rights Reserved.
 
+using System;
 using System.Collections.Generic;
+using System.IO;
+using System.Linq;
 using System.Threading.Tasks;
 using System.Xml;
 using EpicGames.Core;
@@ -82,7 +85,59 @@ namespace AutomationTool.Tasks
 				throw new AutomationException("DotNet is missing from {0}", dotNetFile);
 			}
 
-			IProcessResult result = await ExecuteAsync(dotNetFile.FullName, _parameters.Arguments, workingDir: _parameters.BaseDir, envVars: ParseEnvVars(_parameters.Environment, _parameters.EnvironmentFile));
+			string dotNetRoot = dotNetFile.Directory.FullName;
+
+			Dictionary<string, string> envVars = ParseEnvVars(_parameters.Environment, _parameters.EnvironmentFile);
+			
+			// Modify PATH:
+			// Remove any dotnet in PATH
+			string currentPath = Environment.GetEnvironmentVariable("PATH") ?? "";
+			string cleanPath = String.Join(Path.PathSeparator.ToString(),
+				currentPath
+				.Split(Path.PathSeparator)
+				.Where(p => !p.Contains("dotnet", StringComparison.OrdinalIgnoreCase)));
+			// Add requested dotnet back to PAth
+			envVars["PATH"] = $"{dotNetRoot}{Path.PathSeparator}{cleanPath}";
+
+			// Set DOTNET_ROOT to ensure SDK resolution uses the correct .NET installation
+			envVars["DOTNET_ROOT"] = dotNetRoot;
+
+			// Disable multilevel lookup to prevent searching system-wide .NET installations
+			envVars["DOTNET_MULTILEVEL_LOOKUP"] = "0";
+
+			// Explicitly set the dotnet host path for MSBuild SDK resolver
+			envVars["DOTNET_HOST_PATH"] = dotNetFile.FullName;
+
+			// Point the SDK resolver CLI directory to our dotnet installation
+			envVars["DOTNET_MSBUILD_SDK_RESOLVER_CLI_DIR"] = dotNetRoot;
+
+			// Try to find and set the MSBuild SDKs path
+			DirectoryReference sdkDir = DirectoryReference.Combine(dotNetFile.Directory, "sdk");
+			if (DirectoryReference.Exists(sdkDir))
+			{
+				foreach (DirectoryReference versionDir in DirectoryReference.EnumerateDirectories(sdkDir).Reverse())
+				{
+					DirectoryReference sdksPath = DirectoryReference.Combine(versionDir, "Sdks");
+					if (DirectoryReference.Exists(sdksPath))
+					{
+						envVars["DOTNET_MSBUILD_SDK_RESOLVER_SDKS_DIR"] = sdksPath.FullName;
+						envVars["MSBuildSDKsPath"] = sdksPath.FullName;
+						envVars["MSBuildExtensionsPath"] = versionDir.FullName;
+						envVars["MSBUILD_EXE_PATH"] = FileReference.Combine(versionDir, "MSBuild.dll").FullName;
+
+						// Also set the SDK version explicitly
+						envVars["DOTNET_MSBUILD_SDK_RESOLVER_SDKS_VER"] = versionDir.GetDirectoryName();
+
+						break;
+					}
+				}
+			}
+
+			// CRITICAL: Disable the workload resolver to prevent it from searching other .NET locations
+			// The error originates from Microsoft.NET.Sdk.ImportWorkloads.props
+			envVars["MSBuildEnableWorkloadResolver"] = "false";
+			
+			IProcessResult result = await ExecuteAsync(dotNetFile.FullName, _parameters.Arguments, workingDir: _parameters.BaseDir, envVars: envVars);
 			if (result.ExitCode < 0 || result.ExitCode >= _parameters.ErrorLevel)
 			{
 				Logger.LogError(KnownLogEvents.ExitCode, "Docker terminated with an exit code indicating an error ({ExitCode})", result.ExitCode);
