Fix of Issue 1190259 TestTag does not containt definition for RPC

#rnx
#Synced-CL 49725733

diff --git a/Engine/Source/Programs/AutomationTool/Gauntlet/Framework/RpcFramework/Gauntlet.RpcTimedSections.cs b/Engine/Source/Programs/AutomationTool/Gauntlet/Framework/RpcFramework/Gauntlet.RpcTimedSections.cs
--- a/Engine/Source/Programs/AutomationTool/Gauntlet/Framework/RpcFramework/Gauntlet.RpcTimedSections.cs
+++ b/Engine/Source/Programs/AutomationTool/Gauntlet/Framework/RpcFramework/Gauntlet.RpcTimedSections.cs
@@ -93,8 +93,6 @@ namespace Gauntlet
 			/// <param name="ReportedPhase"></param>
 			public void OnAddToReport(TestPhase ReportedPhase)
 			{
-				ReportedPhase.AddTags([TestTag.RPC]);
-
 				if (ReportedPhase.outcome == TestPhaseOutcome.Interrupted
 					&& !ReportedPhase.GetStream().GetEvents()
 						.Where(E => E.Level == Logging.LogLevel.Error || E.Level == Logging.LogLevel.Critical).Any())
