Allow rolling forward major dotnet versions on Linux and Mac when running through shell scripts. We already do this for Windows, so this brings the host platforms into alignment. Allow rolling forward for non-engine processes launched by DotNetTask and any automation processes too.

#Synced-CL 49760059

diff --git a/Engine/Build/BatchFiles/Linux/SetupDotnet.sh b/Engine/Build/BatchFiles/Linux/SetupDotnet.sh
--- a/Engine/Build/BatchFiles/Linux/SetupDotnet.sh
+++ b/Engine/Build/BatchFiles/Linux/SetupDotnet.sh
@@ -9,6 +9,9 @@ fi
 START_DIR=`pwd`
 cd "$1"
 
+# for dotnet applications that require older dotnet runtimes, this will allow them to use our bundled dotnet runtime
+export DOTNET_ROLL_FORWARD=Major
+
 IS_DOTNET_INSTALLED=0
 DOTNET_VERSION_PATH=$(command -v dotnet) || true
 
diff --git a/Engine/Build/BatchFiles/Mac/SetupDotnet.sh b/Engine/Build/BatchFiles/Mac/SetupDotnet.sh
--- a/Engine/Build/BatchFiles/Mac/SetupDotnet.sh
+++ b/Engine/Build/BatchFiles/Mac/SetupDotnet.sh
@@ -9,6 +9,9 @@ fi
 START_DIR=`pwd`
 cd "$1"
 
+# for dotnet applications that require older dotnet runtimes, this will allow them to use our bundled dotnet runtime
+export DOTNET_ROLL_FORWARD=Major
+
 IS_DOTNET_INSTALLED=0
 DOTNET_VERSION_PATH=$(command -v dotnet) || true
 
diff --git a/Engine/Source/Programs/AutomationTool/AutomationUtils/ProcessUtils.cs b/Engine/Source/Programs/AutomationTool/AutomationUtils/ProcessUtils.cs
--- a/Engine/Source/Programs/AutomationTool/AutomationUtils/ProcessUtils.cs
+++ b/Engine/Source/Programs/AutomationTool/AutomationUtils/ProcessUtils.cs
@@ -66,13 +66,13 @@ namespace AutomationTool
 			{
 				foreach (var EnvPair in Env)
 				{
-					if (NewProcess.StartInfo.EnvironmentVariables.ContainsKey(EnvPair.Key))
+					if (NewProcess.StartInfo.Environment.ContainsKey(EnvPair.Key))
 					{
-						NewProcess.StartInfo.EnvironmentVariables.Remove(EnvPair.Key);
+						NewProcess.StartInfo.Environment.Remove(EnvPair.Key);
 					}
 					if (!String.IsNullOrEmpty(EnvPair.Value))
 					{
-						NewProcess.StartInfo.EnvironmentVariables.Add(EnvPair.Key, EnvPair.Value);
+						NewProcess.StartInfo.Environment.Add(EnvPair.Key, EnvPair.Value);
 					}
 				}
 			}
@@ -1008,6 +1008,12 @@ namespace AutomationTool
 					}
 				}
 
+				// Allow rolling forward to a later dotnet runtime if it's not already specified
+				if (Env?.ContainsKey("DOTNET_ROLL_FORWARD") != true)
+				{
+					FinalEnv["DOTNET_ROLL_FORWARD"] = "Major";
+				}
+
 				// Clearing our variable as well
 				FinalEnv["UE_DOTNET_VERSION"] = String.Empty;
 
@@ -1332,4 +1338,4 @@ namespace AutomationTool
 		}
 
 	}
-}
+}
\ No newline at end of file
diff --git a/Engine/Source/Programs/AutomationTool/BuildGraph/Tasks/DotNetTask.cs b/Engine/Source/Programs/AutomationTool/BuildGraph/Tasks/DotNetTask.cs
--- a/Engine/Source/Programs/AutomationTool/BuildGraph/Tasks/DotNetTask.cs
+++ b/Engine/Source/Programs/AutomationTool/BuildGraph/Tasks/DotNetTask.cs
@@ -85,6 +85,7 @@ namespace AutomationTool.Tasks
 				throw new AutomationException("DotNet is missing from {0}", dotNetFile);
 			}
 
+			// Set DOTNET_ROOT to ensure SDK resolution uses the correct .NET installation
 			string dotNetRoot = dotNetFile.Directory.FullName;
 
 			Dictionary<string, string> envVars = ParseEnvVars(_parameters.Environment, _parameters.EnvironmentFile);
@@ -111,6 +112,9 @@ namespace AutomationTool.Tasks
 			// Point the SDK resolver CLI directory to our dotnet installation
 			envVars["DOTNET_MSBUILD_SDK_RESOLVER_CLI_DIR"] = dotNetRoot;
 
+			// Allow rolling forward major dotnet versions as a concession to dotnet version upgrades
+			envVars["DOTNET_ROLL_FORWARD"] = "Major";
+
 			// Try to find and set the MSBuild SDKs path
 			DirectoryReference sdkDir = DirectoryReference.Combine(dotNetFile.Directory, "sdk");
 			if (DirectoryReference.Exists(sdkDir))
