BuildGraph: Fix compiler warning

#Synced-CL 0

diff --git a/Engine/Source/Programs/AutomationTool/BuildGraph/Tasks/HordeCreateReportTask.cs b/Engine/Source/Programs/AutomationTool/BuildGraph/Tasks/HordeCreateReportTask.cs
--- a/Engine/Source/Programs/AutomationTool/BuildGraph/Tasks/HordeCreateReportTask.cs
+++ b/Engine/Source/Programs/AutomationTool/BuildGraph/Tasks/HordeCreateReportTask.cs
@@ -6,7 +6,6 @@ using System.Linq;
 using System.Text.Json;
 using System.Threading.Tasks;
 using System.Xml;
-using AutomationTool;
 using EpicGames.Core;
 using Microsoft.Extensions.Logging;
 
