[Gauntlet] Force clear horde device reservations where appropriate
#Synced-CL 49698028

diff --git a/Engine/Source/Programs/AutomationTool/AutomationUtils/DeviceReservation.cs b/Engine/Source/Programs/AutomationTool/AutomationUtils/DeviceReservation.cs
--- a/Engine/Source/Programs/AutomationTool/AutomationUtils/DeviceReservation.cs
+++ b/Engine/Source/Programs/AutomationTool/AutomationUtils/DeviceReservation.cs
@@ -546,7 +546,7 @@ namespace AutomationTool.DeviceReservation
 			Utils.InvokeAPI(BaseUri.AppendPath("api/v1/reservations/" + Guid.ToString()), "PUT", NewDuration);
 		}
 
-		public void Delete(Uri BaseUri, bool Force = true)
+		public void Delete(Uri BaseUri, bool Force = false)
 		{
 			string Path = $"api/v1/reservations/{Guid}";
 			Uri DeleteUri = BaseUri.AppendPath(Path);
diff --git a/Engine/Source/Programs/AutomationTool/Gauntlet/Framework/Devices/Gauntlet.DevicePool.cs b/Engine/Source/Programs/AutomationTool/Gauntlet/Framework/Devices/Gauntlet.DevicePool.cs
--- a/Engine/Source/Programs/AutomationTool/Gauntlet/Framework/Devices/Gauntlet.DevicePool.cs
+++ b/Engine/Source/Programs/AutomationTool/Gauntlet/Framework/Devices/Gauntlet.DevicePool.cs
@@ -848,7 +848,7 @@ namespace Gauntlet
 			return true;
 		}
 
-		public bool ReserveDevicesFromService(Dictionary<UnrealDeviceTargetConstraint, int> DeviceTypes)
+		public bool ReserveDevicesFromService(Dictionary<UnrealDeviceTargetConstraint, int> DeviceTypes, IDeviceReservationService.Options Options = null)
 		{
 			// Flatten the required devices into an enumerable
 			List<UnrealDeviceTargetConstraint> RequiredDeviceConstraints = DeviceTypes
@@ -891,7 +891,7 @@ namespace Gauntlet
 
 			foreach (var ReservationPair in ServiceMapping)
 			{
-				if (!ReservationPair.Key.ReserveDevicesFromService(ReservationPair.Value))
+				if (!ReservationPair.Key.ReserveDevicesFromService(ReservationPair.Value, Options))
 				{
 					Log.Info("Failed to reserve all devices from service. See above log for details");
 					return false;
@@ -919,8 +919,9 @@ namespace Gauntlet
 		/// <summary>
 		/// Release all devices in the provided list from our reserved list
 		/// </summary>
-		/// <param name="DeviceList"></param>
-		public void ReleaseDevices(IEnumerable<ITargetDevice> DeviceList)
+		/// <param name="DeviceList">Devices to release</param>
+		/// <param name="ServiceOptions">Options for releasing devices to a reservation service (where appropriate)</param>
+		public void ReleaseDevices(IEnumerable<ITargetDevice> DeviceList, IDeviceReservationService.Options ServiceOptions = null)
 		{
 			if (DeviceList == null || !DeviceList.Any())
 			{
@@ -939,7 +940,7 @@ namespace Gauntlet
 					}
 					else if (TryGetDevicesReservationService(Device, out IDeviceReservationService Service))
 					{
-						Service.ReleaseDevices([Device]);
+						Service.ReleaseDevices([Device], ServiceOptions);
 					}
 					else
 					{
diff --git a/Engine/Source/Programs/AutomationTool/Gauntlet/Framework/Devices/Gauntlet.DeviceReservationService.cs b/Engine/Source/Programs/AutomationTool/Gauntlet/Framework/Devices/Gauntlet.DeviceReservationService.cs
--- a/Engine/Source/Programs/AutomationTool/Gauntlet/Framework/Devices/Gauntlet.DeviceReservationService.cs
+++ b/Engine/Source/Programs/AutomationTool/Gauntlet/Framework/Devices/Gauntlet.DeviceReservationService.cs
@@ -21,11 +21,16 @@ namespace Gauntlet
 	/// </summary>
 	public interface IDeviceReservationService : IDisposable
 	{
+		public class Options
+		{ }
+
 		static virtual bool Enabled { get; }
 		List<ITargetDevice> ReservedDevices { get; }
 		bool CanSupportDeviceConstraint(UnrealDeviceTargetConstraint Constraint);
-		bool ReserveDevicesFromService(IEnumerable<UnrealDeviceTargetConstraint> RequestedDevices);
-		void ReleaseDevices(IEnumerable<ITargetDevice> Devices);
+		bool ReserveDevicesFromService(IEnumerable<UnrealDeviceTargetConstraint> RequestedDevices) => ReserveDevicesFromService(RequestedDevices, null);
+		bool ReserveDevicesFromService(IEnumerable<UnrealDeviceTargetConstraint> RequestedDevices, Options? Options);
+		void ReleaseDevices(IEnumerable<ITargetDevice> Devices) => ReleaseDevices(Devices, null);
+		void ReleaseDevices(IEnumerable<ITargetDevice> Devices, Options? Options);
 		void ReportDeviceError(string DeviceName, string ErrorMessage);
 	}
 
@@ -34,6 +39,15 @@ namespace Gauntlet
 	/// </summary>
 	public class HordeDeviceReservationService : IDeviceReservationService
 	{
+		public class Options : IDeviceReservationService.Options
+		{
+			/// <summary>
+			/// If true, reservations will be force deleted.
+			/// Force deleted reservations will also clean any associated reservation blocks regardless of future steps
+			/// </summary>
+			public bool bForceDelete { get; set; } = true;
+		}
+
 		/// <summary>
 		/// Whether or not this reservation service is enabled
 		/// </summary>
@@ -253,7 +267,10 @@ namespace Gauntlet
 			return true;
 		}
 
-		public virtual bool ReserveDevicesFromService(IEnumerable<UnrealDeviceTargetConstraint> RequestedConstraints)
+		bool IDeviceReservationService.ReserveDevicesFromService(IEnumerable<UnrealDeviceTargetConstraint> RequestedConstraints, IDeviceReservationService.Options? Options) 
+			=> ReserveDevicesFromService(RequestedConstraints, Options as Options);
+
+		public virtual bool ReserveDevicesFromService(IEnumerable<UnrealDeviceTargetConstraint> RequestedConstraints, Options? Options = null)
 		{
 			// Ensure no duplicate requests of an explicit device
 			HashSet<string> DeviceNames = new HashSet<string>();
@@ -359,11 +376,14 @@ namespace Gauntlet
 			return false;
 		}
 
+		void IDeviceReservationService.ReleaseDevices(IEnumerable<ITargetDevice> Devices, IDeviceReservationService.Options? Options)
+			=> ReleaseDevices(Devices, Options as Options);
+
 		/// <summary>
 		/// Release all devices in the provided list from our reserved list
 		/// </summary>
 		/// <param name="DeviceList"></param>
-		public void ReleaseDevices(IEnumerable<ITargetDevice> DeviceList)
+		public void ReleaseDevices(IEnumerable<ITargetDevice> DeviceList, Options? Options = null)
 		{
 			if (!DeviceList.Any())
 			{
@@ -395,6 +415,11 @@ namespace Gauntlet
 					ServiceDeviceInfo.Remove(Device);
 					InitialConnectionState.Remove(Device);
 
+					if (Options != null && Options.bForceDelete)
+					{
+						Reservation.ForceDelete = true;
+					}
+
 					Device.Dispose();
 					Reservation.Dispose();
 				}
diff --git a/Engine/Source/Programs/AutomationTool/Gauntlet/Framework/Devices/Gauntlet.UnrealDeviceReservation.cs b/Engine/Source/Programs/AutomationTool/Gauntlet/Framework/Devices/Gauntlet.UnrealDeviceReservation.cs
--- a/Engine/Source/Programs/AutomationTool/Gauntlet/Framework/Devices/Gauntlet.UnrealDeviceReservation.cs
+++ b/Engine/Source/Programs/AutomationTool/Gauntlet/Framework/Devices/Gauntlet.UnrealDeviceReservation.cs
@@ -15,6 +15,11 @@ namespace Gauntlet
 		protected List<ProblemDevice> ProblemDevices { get; private set; } = new List<ProblemDevice>();
 		private static object LockObject = new object();
 
+		private static readonly HordeDeviceReservationService.Options ForceReleaseOption = new HordeDeviceReservationService.Options()
+		{
+			bForceDelete = true
+		};
+
 		public bool TryReserveDevices(Dictionary<UnrealDeviceTargetConstraint, int> RequiredDeviceTypes, int ExpectedNumberOfDevices, bool bAllowPartialReservations = false)
 		{
 			// First, determine if this is a partial reservation. In the event a device is considered a problem device, don't discard all devices.
@@ -97,7 +102,7 @@ namespace Gauntlet
 				}
 
 				// Release any devices that were skipped
-				DevicePool.Instance.ReleaseDevices(SkippedDevices);
+				DevicePool.Instance.ReleaseDevices(SkippedDevices, ForceReleaseOption);
 
 				// Obtain connections to any required devices.
 				EstablishDeviceConnections(AcquiredDevices);
@@ -124,7 +129,7 @@ namespace Gauntlet
 					}
 					else
 					{
-						DevicePool.Instance.ReleaseDevices(AcquiredDevices);
+						DevicePool.Instance.ReleaseDevices(AcquiredDevices, ForceReleaseOption);
 					}
 
 					return false;
@@ -142,7 +147,7 @@ namespace Gauntlet
 			return true;
 		}
 
-		public void ReleaseDevices()
+		public void ReleaseDevices(bool bForceClearServiceReservations = false)
 		{
 			if ((ReservedDevices != null) && (ReservedDevices.Count() > 0))
 			{
@@ -150,7 +155,12 @@ namespace Gauntlet
 				{
 					IDeviceUsageReporter.RecordEnd(Device.Name, Device.Platform, IDeviceUsageReporter.EventType.Device);
 				}
-				DevicePool.Instance.ReleaseDevices(ReservedDevices);
+
+				IDeviceReservationService.Options ServiceOptions = bForceClearServiceReservations && HordeDeviceReservationService.Enabled
+					? ForceReleaseOption
+					: null;
+				DevicePool.Instance.ReleaseDevices(ReservedDevices, ServiceOptions);
+
 				ReservedDevices.Clear();
 			}
 		}
diff --git a/Engine/Source/Programs/AutomationTool/Gauntlet/Unreal/Base/Gauntlet.UnrealSession.cs b/Engine/Source/Programs/AutomationTool/Gauntlet/Unreal/Base/Gauntlet.UnrealSession.cs
--- a/Engine/Source/Programs/AutomationTool/Gauntlet/Unreal/Base/Gauntlet.UnrealSession.cs
+++ b/Engine/Source/Programs/AutomationTool/Gauntlet/Unreal/Base/Gauntlet.UnrealSession.cs
@@ -837,7 +837,7 @@ namespace Gauntlet
 			{
 				if(Globals.CancelSignalled)
 				{
-					ReleaseSessionDevices();
+					ReleaseSessionDevices(true);
 					return false;
 				}
 
@@ -852,7 +852,7 @@ namespace Gauntlet
 			}
 
 			Log.Error("Failed to reserve devices after {Attempts} attempts", Attempts);
-			ReleaseSessionDevices();
+			ReleaseSessionDevices(true);
 			return false;
 		}
 
@@ -917,7 +917,7 @@ namespace Gauntlet
 				if(!TryAssignDevicesToRoles())
 				{
 					// If device assignment fails, reservations were likely deleted at an unexpected time.
-					ReleaseSessionDevices();
+					ReleaseSessionDevices(true);
 					continue;
 				}
 
@@ -932,7 +932,7 @@ namespace Gauntlet
 					if(IsOutOfSpaceException(Ex) || IsOverlayException(Ex) || Ex is not DeviceException)
 					{
 						RemainingAttempts = 0;
-						ReleaseSessionDevices();
+						ReleaseSessionDevices(true);
 						continue;
 					}
 					else if (RemainingAttempts > 1)
@@ -946,7 +946,7 @@ namespace Gauntlet
 
 				if (Globals.CancelSignalled)
 				{
-					ReleaseSessionDevices();
+					ReleaseSessionDevices(true);
 					return null;
 				}
 
@@ -969,7 +969,7 @@ namespace Gauntlet
 
 				if (Globals.CancelSignalled)
 				{
-					ReleaseSessionDevices();
+					ReleaseSessionDevices(true);
 					return null;
 				}
 
@@ -1018,7 +1018,7 @@ namespace Gauntlet
 			}
 		}
 
-		public void ReleaseSessionDevices()
+		public void ReleaseSessionDevices(bool bForce = false)
 		{
 			// Terminate any running apps
 			if (SessionInstance != null && SessionInstance.RunningRoles != null)
@@ -1036,7 +1036,7 @@ namespace Gauntlet
 				return;
 			}
 
-			UnrealDeviceReservation.ReleaseDevices();
+			UnrealDeviceReservation.ReleaseDevices(bForce);
 
 			RolesToDevices.Clear();
 			RolesToConfigs.Clear();
@@ -1575,7 +1575,7 @@ namespace Gauntlet
 
 				if (InstallSuccess == false)
 				{
-					ReleaseSessionDevices();
+					ReleaseSessionDevices(true);
 
 					if (SessionRetries == 0)
 					{
@@ -1641,7 +1641,7 @@ namespace Gauntlet
 								RunningRole.AppInstance.Device.Disconnect();
 							}
 
-							ReleaseSessionDevices();
+							ReleaseSessionDevices(true);
 
 							if (SessionRetries == 0)
 							{
