Fix RPC InputKey not processing Tapped inputs

The Tapped parameter was previously omitted from the RPC arguments dictionary.

Created conversion method "AsRpcArgs()" to automatically include all properties decorated with JsonPropertyName, eliminating the possibility of forgetting parameters.

#Synced-CL 0

#Synced-CL 0

diff --git a/Engine/Source/Programs/AutomationTool/Gauntlet/Framework/Utils/Gauntlet.InputActions.cs b/Engine/Source/Programs/AutomationTool/Gauntlet/Framework/Utils/Gauntlet.InputActions.cs
--- a/Engine/Source/Programs/AutomationTool/Gauntlet/Framework/Utils/Gauntlet.InputActions.cs
+++ b/Engine/Source/Programs/AutomationTool/Gauntlet/Framework/Utils/Gauntlet.InputActions.cs
@@ -1,7 +1,9 @@
 ﻿// Copyright Epic Games, Inc. All Rights Reserved.
 
-using EpicGames.Core;
-using System;
+using System.Collections.Generic;
+using System.Reflection;
+using System.Text.Json;
+using System.Text.Json.Serialization;
 using UnrealBuildTool;
 
 namespace Gauntlet.Utils
@@ -559,22 +561,48 @@ namespace Gauntlet.Utils
 
 	}
 
-	public class InputAction
+	public class InputAction(FKey InKey, EInputEvent InInputEvent, double InXDelta = 0, double InYDelta = 0, bool InTapped = false)
 	{
-		public FKey Key;
-		public EInputEvent KeyAction;
-		public double XDelta;
-		public double YDelta;
-		public bool Tapped;
-
-		public InputAction(FKey Inkey, EInputEvent InInputEvent, double InXDelta = 0, double InYDelta = 0, bool InTapped = false)
+		[JsonPropertyName("inputkey")]
+		public FKey Key { get; } = InKey;
+
+		[JsonPropertyName("inputevent")]
+		public EInputEvent KeyAction { get; } = InInputEvent;
+
+		[JsonPropertyName("xdelta")]
+		public double XDelta { get; } = InXDelta;
+
+		[JsonPropertyName("ydelta")]
+		public double YDelta { get;  } = InYDelta;
+
+		[JsonPropertyName("tapped")]
+		public bool Tapped { get;  } = InTapped;
+
+		/// <summary>
+		/// Converts the instance's properties into a dictionary suitable for RPC method arguments.
+		/// Properties decorated with <see cref="JsonPropertyNameAttribute"/> are serialized and added to the dictionary.
+		/// </summary>
+		/// <returns>A dictionary with JSON property names as keys and serialized property values as JSON strings.</returns>
+		/// <remarks>
+		/// If this pattern is needed across multiple classes, consider creating an extension method
+		/// on <see cref="object"/> to avoid code duplication.
+		/// </remarks>
+		public Dictionary<string, object> AsRpcArgs()
 		{
-			this.Key = Inkey;
-			this.KeyAction = InInputEvent;
-			this.XDelta = InXDelta;
-			this.YDelta = InYDelta;
-			Tapped = InTapped;
-		}
+			Dictionary<string, object> Args = [];
+			PropertyInfo[] Properties = GetType().GetProperties();
+
+			foreach (var Property in Properties)
+			{
+				JsonPropertyNameAttribute JsonAttribute = Property.GetCustomAttribute<JsonPropertyNameAttribute>();
+				if (JsonAttribute != null)
+				{
+					object Value = Property.GetValue(this);
+					Args[JsonAttribute.Name] = JsonSerializer.Serialize(Value);
+				}
+			}
 
+			return Args;
+		}
 	}
 }
\ No newline at end of file
