CopyBuildLocally - facility to access builds from the P drive, UBS, depending on what's available
Set up patch diff jobs to use it so they can seamlessly run from all branches after Main builds are no longer stored on the P drive.

#rnx
#Synced-CL 49773231

diff --git a/Engine/Source/Programs/AutomationTool/AutomationUtils/ProjectParams.cs b/Engine/Source/Programs/AutomationTool/AutomationUtils/ProjectParams.cs
--- a/Engine/Source/Programs/AutomationTool/AutomationUtils/ProjectParams.cs
+++ b/Engine/Source/Programs/AutomationTool/AutomationUtils/ProjectParams.cs
@@ -1953,7 +1953,7 @@ namespace AutomationTool
         public bool AddPatchLevel;
 
 		/// <summary>
-		/// Name of dlc to cook and package (if this paramter is supplied cooks the dlc and packages it into the dlc directory)
+		/// Name of dlc to cook and package (if this parameter is supplied cooks the dlc and packages it into the dlc directory)
 		/// </summary>
 		public FileReference DLCFile;
 
@@ -2557,7 +2557,7 @@ namespace AutomationTool
 			{
 				if (!AvailableTargets.Contains(DefaultTarget))
 				{
-					throw new AutomationException(string.Format($"A default {Type} target '{DefaultTarget}' was specified in engine.ini but does not exist"));
+					throw new AutomationException(string.Format($"A default {Type} target '{DefaultTarget}' was specified in engine.ini but does not exist. Available targets: {string.Join(", ", AvailableTargets)}"));
 				}
 
 				Target = DefaultTarget;
