Removing appium foreground check as it breaks stdout console

diff --git a/Engine/Source/Programs/AutomationTool/Gauntlet/Platform/IOS/Gauntlet.TargetDeviceIOS.cs b/Engine/Source/Programs/AutomationTool/Gauntlet/Platform/IOS/Gauntlet.TargetDeviceIOS.cs
--- a/Engine/Source/Programs/AutomationTool/Gauntlet/Platform/IOS/Gauntlet.TargetDeviceIOS.cs
+++ b/Engine/Source/Programs/AutomationTool/Gauntlet/Platform/IOS/Gauntlet.TargetDeviceIOS.cs
@@ -844,8 +844,6 @@ namespace Gauntlet
 			{
 				lock (Mutex)
 				{
-					// Making sure app is in foreground and no alerts present. 
-					AppiumContainer.VerifyAppInForeground(IOSInstall.PackageName);
 					AppiumContainer.CancelAppleSignOnAlert(IOSInstall.PackageName);
 				}
 			}
@@ -1595,4 +1593,4 @@ namespace Gauntlet
 		public override UnrealTargetPlatform? Platform => UnrealTargetPlatform.IOS;
 		public override bool IsHostMountingSupported() => false;
 	}
-}
\ No newline at end of file
+}
