Implementing Android Screenshot thread loop.

#Synced-CL 49866952

diff --git a/Engine/Source/Programs/AutomationTool/Gauntlet/Platform/Android/Gauntlet.TargetDeviceAndroid.cs b/Engine/Source/Programs/AutomationTool/Gauntlet/Platform/Android/Gauntlet.TargetDeviceAndroid.cs
--- a/Engine/Source/Programs/AutomationTool/Gauntlet/Platform/Android/Gauntlet.TargetDeviceAndroid.cs
+++ b/Engine/Source/Programs/AutomationTool/Gauntlet/Platform/Android/Gauntlet.TargetDeviceAndroid.cs
@@ -1,17 +1,21 @@
 // Copyright Epic Games, Inc. All Rights Reserved.
 
+using AutomationTool;
+using EpicGames.Core;
+using Microsoft.Build.Framework;
+using OpenQA.Selenium;
 using System;
 using System.Collections.Generic;
+using System.Diagnostics;
 using System.IO;
-using AutomationTool;
-using UnrealBuildTool;
-using System.Text.RegularExpressions;
 using System.Linq;
 using System.Text;
+using System.Text.RegularExpressions;
 using System.Threading;
-using EpicGames.Core;
-using System.Diagnostics;
+using UnrealBuildBase;
+using UnrealBuildTool;
 using static AutomationTool.CommandUtils;
+using static System.Runtime.InteropServices.JavaScript.JSType;
 
 namespace Gauntlet
 {
@@ -117,6 +121,10 @@ namespace Gauntlet
 
 		private bool bUseAFSWithoutToken;
 
+		private Thread ScreenshotThread = null;
+
+		private AutoResetEvent ScreenshotTrigger = null;
+
 		public TargetDeviceAndroid(string InDeviceName = "", AndroidDeviceData DeviceData = null, string InCachePath = null)
 		{
 			AutoParam.ApplyParamsAndDefaults(this, Globals.Params.AllArguments);
@@ -245,6 +253,14 @@ namespace Gauntlet
 
 						Log.Info("Disconnected {0}", DeviceName);
 					}
+
+					if (ScreenshotThread != null)
+					{
+						ScreenshotTrigger.Set();
+						ScreenshotThread.Join();
+						ScreenshotThread = null;
+						ScreenshotTrigger = null;
+					}
 				}
 				catch (Exception Ex)
 				{
@@ -414,7 +430,7 @@ namespace Gauntlet
 
 			// Copy obbs from bulk builds
 			bool bSkipOBBInstall = Globals.Params.ParseParam("SkipOBBCopy"); // useful when iterating on dev executables
-			if(Build.FilesToInstall != null && Build.FilesToInstall.Any() && !bSkipOBBInstall)
+			if (Build.FilesToInstall != null && Build.FilesToInstall.Any() && !bSkipOBBInstall)
 			{
 				CopyOBBFiles(AppConfig, Build.FilesToInstall, Package);
 			}
@@ -422,7 +438,7 @@ namespace Gauntlet
 
 		public IAppInstall CreateAppInstall(UnrealAppConfig AppConfig)
 		{
-			if(AppConfig.Build is not AndroidBuild Build)
+			if (AppConfig.Build is not AndroidBuild Build)
 			{
 				throw new AutomationException("Unsupported build type {0} for Android!", AppConfig.Build.GetType());
 			}
@@ -478,7 +494,7 @@ namespace Gauntlet
 				throw new AutomationException("AppInstall is of incorrect type {0}! Must be of type AndroidAppInstall", App.GetType().Name);
 			}
 
-			if(!IsOn)
+			if (!IsOn)
 			{
 				PowerOn();
 			}
@@ -510,6 +526,12 @@ namespace Gauntlet
 				RunAdbDeviceCommand(string.Format("shell mkdir -p {0}/", DeviceExternalFilesSavedPath));
 			}
 
+			if (!Globals.Params.ParseParam("NoAndroidScreenshots"))
+			{
+				ScreenshotThread = new Thread(StartScreenshotLoop);
+				ScreenshotThread.Start();
+			}
+
 			if (!Globals.IsRunningDev && UseAFS(Install.AppConfig))
 			{
 				AndroidBuild Build = Install.AppConfig.Build as AndroidBuild;
@@ -562,7 +584,7 @@ namespace Gauntlet
 				{
 					ProjectDir += '/';
 				}
-				
+
 				LocalDirectoryMappings.Add(EIntendedBaseCopyDirectory.Build, ProjectDir + "Build");
 				LocalDirectoryMappings.Add(EIntendedBaseCopyDirectory.Binaries, ProjectDir + "Binaries");
 				LocalDirectoryMappings.Add(EIntendedBaseCopyDirectory.Config, ProjectDir + "Saved/Config");
@@ -585,7 +607,7 @@ namespace Gauntlet
 				bAllowExternalStartInShipping = false;
 				return true;
 			}
-			if(RawProjectPath == null || !FileReference.Exists(RawProjectPath) || TargetConfiguration == UnrealTargetConfiguration.Unknown)
+			if (RawProjectPath == null || !FileReference.Exists(RawProjectPath) || TargetConfiguration == UnrealTargetConfiguration.Unknown)
 			{
 				bEnablePlugin = false;
 				AFSToken = string.Empty;
@@ -699,7 +721,7 @@ namespace Gauntlet
 		public static IProcessResult RunAdbGlobalCommand(string Args, bool Wait = true, bool bShouldLogCommand = false, bool bPauseErrorParsing = false, int TimeoutSec = 60 * 15)
 		{
 			CommandUtils.ERunOptions RunOptions = CommandUtils.ERunOptions.AppMustExist | CommandUtils.ERunOptions.NoWaitForExit | CommandUtils.ERunOptions.SpewIsVerbose;
-			RunOptions |= Log.IsVeryVerbose? ERunOptions.AllowSpew : ERunOptions.NoLoggingOfRunCommand;
+			RunOptions |= Log.IsVeryVerbose ? ERunOptions.AllowSpew : ERunOptions.NoLoggingOfRunCommand;
 
 			if (bShouldLogCommand)
 			{
@@ -747,7 +769,7 @@ namespace Gauntlet
 			return new LongProcessResult(AdbCommand, Args, RunOptions, LocalCache: LocalCache);
 		}
 
-		public static IProcessResult RunAFSGlobalCommand(string Args, bool Wait = true, bool bShouldLogCommand = false, bool bPauseErrorParsing = false, int TimeoutSec = 60 * 15 )
+		public static IProcessResult RunAFSGlobalCommand(string Args, bool Wait = true, bool bShouldLogCommand = false, bool bPauseErrorParsing = false, int TimeoutSec = 60 * 15)
 		{
 			CommandUtils.ERunOptions RunOptions = CommandUtils.ERunOptions.AppMustExist | CommandUtils.ERunOptions.NoWaitForExit | CommandUtils.ERunOptions.SpewIsVerbose;
 
@@ -884,7 +906,7 @@ namespace Gauntlet
 			// we write this out to a dependency file in the demote dir and check it each time.
 			// current file time
 			DateTime LocalModifiedTime = File.GetLastWriteTime(SourcePath);
-			
+
 			// If not using AFS and the APK starts with AFS_, rename to use the normal one
 			if (!Globals.IsRunningDev && IsAPK && !UseAFS)
 			{
@@ -1138,7 +1160,7 @@ namespace Gauntlet
 			{
 				RunAFSDeviceCommand(string.Format("-p {0} -k {1} rm {2}", Package, AFSToken, "^mainobb"));
 				RunAFSDeviceCommand(string.Format("-p {0} -k {1} rm {2}", Package, AFSToken, "^patchobb"));
-				for(int i = 1; i <= 32; i++)
+				for (int i = 1; i <= 32; i++)
 				{
 					string overflowobb = string.Format("^overflow{0}obb", i);
 					RunAFSDeviceCommand(string.Format("-p {0} -k {1} rm {2}", Package, AFSToken, overflowobb));
@@ -1261,6 +1283,78 @@ namespace Gauntlet
 			return AllDevices.Keys.Where(D => AllDevices[D] == "device");
 		}
 
+		private void StartScreenshotLoop()
+		{
+			TimeSpan Period = TimeSpan.FromSeconds(30);
+			ScreenshotTrigger = new AutoResetEvent(false);
+
+			while (!ScreenshotTrigger.WaitOne(Period))
+			{
+				TakeScreenshot();
+			}
+		}
+		/// <summary>
+		/// Stops Screenshot thread, pulls screenshots from sdcard, and cleans up temp dir on the phone. 
+		/// </summary>
+		/// <param name="OutputDirectory"></param>
+		public void StopScreenshotLoop(string OutputDirectory = "")
+		{
+			if (ScreenshotThread == null || ScreenshotTrigger == null)
+			{
+				Log.Verbose("Attempted to end the screenshot loop before it started. Ignoring.");
+				return;
+			}
+
+			ScreenshotTrigger.Set();
+			ScreenshotThread.Join();
+			ScreenshotThread = null;
+			ScreenshotTrigger = null;
+
+			if (!string.IsNullOrEmpty(OutputDirectory))
+			{
+				try
+				{
+					if (!DirectoryExists(OutputDirectory))
+					{
+						CreateDirectory(OutputDirectory);
+					}
+
+					RunAdbDeviceCommand(string.Format($"pull /sdcard/Screenshots {OutputDirectory}"));
+					var files = Directory.GetFiles(OutputDirectory);
+					foreach (var file in files)
+					{
+						Log.Info("{0}", file);
+					}
+					Log.Info("Pulling Screenshots from phone to {OutputDirectory}", OutputDirectory);
+					RunAdbDeviceCommand("shell rm -r /sdcard/Screenshots");
+
+				}
+				catch (Exception Ex)
+				{
+					Log.Info("Encountered an {ExceptionType} when trying to copy screenshots to requested output directory. " +
+						"Screenshots will not be available: {ExceptionMessage}", Ex.GetType().Name, Ex.Message);
+				}
+			}
+		}
+		/// <summary>
+		/// Takes a screenshot using adb and saves to the phone's sdcard in a temp directory. 
+		/// </summary>
+		public void TakeScreenshot()
+		{
+			string FileName = $"{DateTime.Now.ToString("yy-MM-dd-HH-mm-ss-ff")}.png";
+			try
+			{
+				RunAdbDeviceCommand("shell mkdir -p /sdcard/Screenshots/");
+				RunAdbDeviceCommand(string.Format($"shell screencap -p /sdcard/Screenshots/{FileName}"));
+
+				Log.Verbose("Saved screenshot: {ScreenshotFile}", FileName);
+			}
+			catch (Exception Ex)
+			{
+				Log.Info("Encountered an {ExceptionType} when trying to capture a screenshot: {ExceptionMessage}", Ex.GetType().Name, Ex.Message);
+			}
+		}
+
 		#region Legacy Implementations
 		public IAppInstall InstallApplication(UnrealAppConfig AppConfig)
 		{
@@ -1293,7 +1387,7 @@ namespace Gauntlet
 			KillRunningProcess(Build.AndroidPackageName);
 
 			// "/mnt/sdcard";
-																	   // remote dir used to save things
+			// remote dir used to save things
 			string ExternalStoragePath = StoragePath + "/UnrealGame/" + AppConfig.ProjectName;
 			string ExternalFilesPath = StoragePath + "/Android/data/" + Build.AndroidPackageName + "/files/UnrealGame/" + AppConfig.ProjectName;
 			string DeviceBaseDir = Build.UsesExternalFilesDir ? ExternalFilesPath : ExternalStoragePath;
@@ -1550,7 +1644,7 @@ namespace Gauntlet
 		public FileReference ProjectFile { get; protected set; }
 
 		public UnrealTargetConfiguration Configuration { get; protected set; }
-		
+
 		public UnrealAppConfig AppConfig { get; protected set; }
 
 		public AndroidAppInstall(UnrealAppConfig InAppConfig, TargetDeviceAndroid InDevice, string InName, string InAndroidPackageName, string InCommandLine, FileReference InProjectFile, UnrealTargetConfiguration InConfiguration, string InAppTag = "UE")
@@ -1670,6 +1764,8 @@ namespace Gauntlet
 
 		public AdbScreenRecorder Recorder;
 
+
+
 		public AndroidAppInstance(TargetDeviceAndroid InDevice, AndroidAppInstall InInstall, IProcessResult InProcess)
 		{
 			AndroidDevice = InDevice;
@@ -1713,7 +1809,7 @@ namespace Gauntlet
 				Install.AndroidDevice.RunAFSDeviceCommand(string.Format("-p \"{0}\" -k \"{1}\" mkdir \"^saved/Logs\"", Install.AndroidPackageName, AFSToken));
 			}
 			else
-			{ 
+			{
 				Install.AndroidDevice.RunAdbDeviceCommand($"shell mkdir -p {Install.AndroidDevice.DeviceArtifactPath}/Logs/");
 			}
 			Recorder = new AdbScreenRecorder();
@@ -1755,6 +1851,10 @@ namespace Gauntlet
 			}
 			File.WriteAllText(Path.Combine(SavedDirectory, LogcatFilename), LogcatResult.Output);
 			bHaveSavedArtifacts = true;
+			if (!Globals.Params.ParseParam("NoAndroidScreenshots"))
+			{
+				AndroidDevice.StopScreenshotLoop(SavedDirectory);
+			}
 		}
 
 		/// <summary>
@@ -2058,13 +2158,13 @@ namespace Gauntlet
 
 				}
 
-				if (InstanceCount == 0 && !String.IsNullOrEmpty(KeyPath))
+				if (InstanceCount == 0 && !System.String.IsNullOrEmpty(KeyPath))
 				{
 
 					Log.Info("Using adb keys at {0}", KeyPath);
 
 					string LocalKeyPath = Path.Combine(Environment.GetFolderPath(Environment.SpecialFolder.UserProfile), ".android");
-					if(!Directory.Exists(LocalKeyPath))
+					if (!Directory.Exists(LocalKeyPath))
 					{
 						Directory.CreateDirectory(LocalKeyPath);
 					}
@@ -2144,4 +2244,4 @@ namespace Gauntlet
 		public override UnrealTargetPlatform? Platform => UnrealTargetPlatform.Android;
 		public override bool IsHostMountingSupported() => false;
 	}
-}
\ No newline at end of file
+}
