Review Profile PGO
#rnx
#tests Local compile and preflight

#Synced-CL 49995865

diff --git a/Engine/Source/Programs/AutomationTool/Gauntlet/Platform/Windows/Gauntlet.WindowsPGO.cs b/Engine/Source/Programs/AutomationTool/Gauntlet/Platform/Windows/Gauntlet.WindowsPGO.cs
--- a/Engine/Source/Programs/AutomationTool/Gauntlet/Platform/Windows/Gauntlet.WindowsPGO.cs
+++ b/Engine/Source/Programs/AutomationTool/Gauntlet/Platform/Windows/Gauntlet.WindowsPGO.cs
@@ -1,10 +1,12 @@
 // Copyright Epic Games, Inc. All Rights Reserved.
 
 using AutomationTool;
+using EpicGames.Core;
 using System;
 using System.IO;
 using System.Text.RegularExpressions;
 using UnrealBuildTool;
+using Log = Gauntlet.Log;
 
 namespace Gauntlet
 {
@@ -36,6 +38,7 @@ namespace Gauntlet
 			string[] PGCFiles = Directory.GetFiles(PGODir, "*.pgc");
 			if (PGCFiles.Length == 0)
 			{
+				Log.Error(KnownLogEvents.Gauntlet, "No .pgc files were generated. Make sure that you are using a build with PGO instrumentation.");
 				throw new AutomationException("no .pgc files were found in {0}", PGODir );
 			}
 			Log.Info("Found {0} .pgc files in \"{1}\"", PGCFiles.Length, PGODir);
