UGSSyncExtension for Project Launcher 2

#rnx
#Synced-CL 50055232

diff --git a/Engine/Source/Programs/AutomationTool/Scripts/BuildAcquire.Automation.cs b/Engine/Source/Programs/AutomationTool/Scripts/BuildAcquire.Automation.cs
--- a/Engine/Source/Programs/AutomationTool/Scripts/BuildAcquire.Automation.cs
+++ b/Engine/Source/Programs/AutomationTool/Scripts/BuildAcquire.Automation.cs
@@ -109,7 +109,11 @@ class BuildAcquire : BuildCommand, IProjectParamsHelpers
 				Command: this,
 				RawProjectPath: ProjectPath
 			);
-		DownloadBuilds(Params);
+
+		if (DownloadBuild)
+		{
+			DownloadBuilds(Params);
+		}
 
 		if (SyncUGS)
 		{
@@ -141,6 +145,7 @@ class BuildAcquire : BuildCommand, IProjectParamsHelpers
 		DownloadClient = ParseParam("DownloadClient");
 		DownloadServer = ParseParam("DownloadServer");
 		DownloadGame = ParseParam("DownloadGame");
+		DownloadBuild = DownloadClient || DownloadServer || DownloadGame || ParseParam("BuildPlatform");
 
 		Stream = ParseParamValue("Stream");
 		if (Stream == null)
@@ -154,9 +159,9 @@ class BuildAcquire : BuildCommand, IProjectParamsHelpers
 			throw new AutomationException("No CL number specified. Use -CL=<changelist>");
 		}
 
-		if (!(DownloadClient || DownloadServer || DownloadGame) && !ParseParam("BuildPlatform"))
+		if (!DownloadBuild && !SyncUGS)
 		{
-			throw new AutomationException("No build variant provided. Use -DownloadGame|-DownloadServer|-DownloadClient or -BuildPlatform");
+			throw new AutomationException("No action to perform. To download a build use -DownloadGame|-DownloadServer|-DownloadClient or -BuildPlatform. To sync UGS pass in -SyncUGS");
 		}
 
 		UGSExePath = FileReference.Combine(new DirectoryReference(Environment.GetFolderPath(Environment.SpecialFolder.LocalApplicationData)), "UnrealGameSync", "Latest", "ugs.exe");
@@ -330,7 +335,9 @@ class BuildAcquire : BuildCommand, IProjectParamsHelpers
 	bool DownloadServer;
 	bool DownloadClient;
 	bool DownloadGame;
+
 	bool SyncUGS;
+	bool DownloadBuild;
 
 	FileReference UGSExePath;
 	// TODO: Add support for build config once established how PL2 wants to pass in build configs. Currently it uses build platform which contains the full set: platform, config, executable type
