Make sure all error messages are passed to the report when test outcome is interrupted

#rnx

#Synced-CL 50160642

diff --git a/Engine/Source/Programs/AutomationTool/Gauntlet/Unreal/Base/Gauntlet.UnrealTestNode.cs b/Engine/Source/Programs/AutomationTool/Gauntlet/Unreal/Base/Gauntlet.UnrealTestNode.cs
--- a/Engine/Source/Programs/AutomationTool/Gauntlet/Unreal/Base/Gauntlet.UnrealTestNode.cs
+++ b/Engine/Source/Programs/AutomationTool/Gauntlet/Unreal/Base/Gauntlet.UnrealTestNode.cs
@@ -1753,15 +1753,9 @@ namespace Gauntlet
 				{
 					foreach (UnrealRoleResult Item in RoleResults)
 					{
-						string RoleType = Item.Artifacts.SessionRole.RoleType.ToString();
 						if (Item.ProcessResult == UnrealProcessResult.InitializationFailure)
 						{
-							var EventStream = InitPhase.GetStream();
-							foreach (UnrealTestEvent Event in Item.Events)
-							{
-								EventStream.AddEvent(Event);
-							}
-							EventStream.AddError($"{RoleType}: {Item.Summary}");
+							InitPhase.GetStream().AddError("No tests were executed!");
 							InitPhase.SetOutcome(TestPhaseOutcome.Interrupted);
 						}
 						// Since we are going through the Role, lets add the related device
@@ -1799,7 +1793,19 @@ namespace Gauntlet
 						}
 						else
 						{
-							InitPhase.GetStream().AddError("No tests were executed!");
+							var Stream = InitPhase.GetStream();
+							Stream.AddError("No tests were executed!");
+							if (RoleResults != null)
+							{
+								// Add any relevant events collected through log parsing
+								foreach (UnrealRoleResult RoleResult in RoleResults)
+								{
+									foreach (UnrealTestEvent Item in RoleResult.Events)
+									{
+										Stream.AddEvent(Item);
+									}
+								}
+							}
 						}
 					}
 					SessionReport.PopulateFromUnrealAutomatedTests(JsonTestPassResults, UnrealAutomatedTestReportPath);
@@ -1833,22 +1839,24 @@ namespace Gauntlet
 				{
 					var InitPhase = SessionReport.GetPhase("Initialization");
 					InitPhase.GetStream().AddWarning("Could not find Unreal Automated test report at {FilePath}. No test report details available.", JsonReportPath);
-					// Check if there is any non-OK exit to add more context
-					UnrealRoleResult RelevantRoleResult = RoleResults?.FirstOrDefault(I => I.ProcessResult != UnrealProcessResult.ExitOk && I.ProcessResult != UnrealProcessResult.EncounteredEnsure);
-					if (RelevantRoleResult != null)
+					if (RoleResults != null)
 					{
-						string RoleType = RelevantRoleResult.Artifacts.SessionRole.RoleType.ToString();
-						var Stream = InitPhase.GetStream();
-						foreach(UnrealTestEvent Item in RelevantRoleResult.Events)
+						// Check if there is any non-OK exit to add more context
+						foreach (UnrealRoleResult RoleResult in RoleResults.Where(R => R.ProcessResult != UnrealProcessResult.ExitOk && R.ProcessResult != UnrealProcessResult.EncounteredEnsure))
 						{
-							Stream.AddEvent(Item);
-						}
-						Stream.AddError("{Role}: {Summary}", RoleType, RelevantRoleResult.Summary);
+							string RoleType = RoleResult.Artifacts.SessionRole.RoleType.ToString();
+							var Stream = InitPhase.GetStream();
+							Stream.AddError("{Role}: {Summary}", RoleType, RoleResult.Summary);
+							foreach (UnrealTestEvent Item in RoleResult.Events)
+							{
+								Stream.AddEvent(Item);
+							}
 
-						// In case of a fatal error or initialization failure, it is an interrupted outcome
-						if (RelevantRoleResult.ProcessResult == UnrealProcessResult.EncounteredFatalError || RelevantRoleResult.ProcessResult == UnrealProcessResult.InitializationFailure)
-						{
-							InitPhase.SetOutcome(TestPhaseOutcome.Interrupted);
+							// In case of a fatal error or initialization failure, it is an interrupted outcome
+							if (RoleResult.ProcessResult == UnrealProcessResult.EncounteredFatalError || RoleResult.ProcessResult == UnrealProcessResult.InitializationFailure)
+							{
+								InitPhase.SetOutcome(TestPhaseOutcome.Interrupted);
+							}
 						}
 					}
 				}
