Improving logging for BootTest and gauntlet step message during errors

#tests tested locally and in horde look at preflight bellow

#Synced-CL 49993336

diff --git a/Engine/Source/Programs/AutomationTool/Gauntlet/Framework/Gauntlet.TestExecutor.cs b/Engine/Source/Programs/AutomationTool/Gauntlet/Framework/Gauntlet.TestExecutor.cs
--- a/Engine/Source/Programs/AutomationTool/Gauntlet/Framework/Gauntlet.TestExecutor.cs
+++ b/Engine/Source/Programs/AutomationTool/Gauntlet/Framework/Gauntlet.TestExecutor.cs
@@ -816,6 +816,7 @@ namespace Gauntlet
 					TestInfo.CancellationReason = string.Format("Terminating Test {0} due to maximum duration of {1} seconds. ", TestInfo.TestNode, TestInfo.TestNode.MaxDuration);
 					TestInfo.FinalResult = TestResult.TimedOut;
 					Log.Verbose("Terminating current test pass {TestName} due to maximum duration of {TestMaxDuration} seconds.", TestInfo.TestNode, TestInfo.TestNode.MaxDuration);
+					Horde.UpdateStepMessage(Horde.GauntletHordeStepMessage.TestTimeoutExceeded);
 				}
 				else if (TestInfo.TestNode.MaxDurationReachedResult == EMaxDurationReachedResult.Success)
 				{
diff --git a/Engine/Source/Programs/AutomationTool/Gauntlet/Unreal/BuildSource/Gauntlet.UnrealBuildSource.cs b/Engine/Source/Programs/AutomationTool/Gauntlet/Unreal/BuildSource/Gauntlet.UnrealBuildSource.cs
--- a/Engine/Source/Programs/AutomationTool/Gauntlet/Unreal/BuildSource/Gauntlet.UnrealBuildSource.cs
+++ b/Engine/Source/Programs/AutomationTool/Gauntlet/Unreal/BuildSource/Gauntlet.UnrealBuildSource.cs
@@ -10,7 +10,6 @@ using System.Linq;
 using EpicGames.Core;
 using System.Text.Json.Serialization;
 using System.Text.Json;
-using UnrealGame;
 
 namespace Gauntlet
 {
@@ -346,7 +345,8 @@ namespace Gauntlet
 
 			if (string.IsNullOrEmpty(OutBuildName) || OutBuildPaths == null)
 			{
-				Log.Error("Unable to resolve build argument '{0}'", InBuildReference);
+				Log.Error("Unable to resolve build argument '{BuildReference}'", InBuildReference);
+				Horde.UpdateStepMessage(Horde.GauntletHordeStepMessage.BuildNotFound);
 				return false;
 			}
 
