Add telemetry event for an island codes plugin dependencies.

Also add support for gauntlet CSV parser to fill missing entries with default values.

#Synced-CL 49940768

diff --git a/Engine/Source/Programs/AutomationTool/Gauntlet/Framework/Utils/Gauntlet.CSVParser.cs b/Engine/Source/Programs/AutomationTool/Gauntlet/Framework/Utils/Gauntlet.CSVParser.cs
--- a/Engine/Source/Programs/AutomationTool/Gauntlet/Framework/Utils/Gauntlet.CSVParser.cs
+++ b/Engine/Source/Programs/AutomationTool/Gauntlet/Framework/Utils/Gauntlet.CSVParser.cs
@@ -3,13 +3,14 @@
 using System.IO;
 using System.Collections.Generic;
 using AutomationTool;
+using System.Linq;
 
 
 namespace Gauntlet
 {
 	public static class CSVParser
 	{
-		public static List<Dictionary<string, string>> Load(string File)
+		public static List<Dictionary<string, string>> Load(string File, string DefaultForMissingColumns = null)
 		{
 			List<Dictionary<string, string>> Rows = new List<Dictionary<string, string>>();
 			// Parse csv file
@@ -26,6 +27,16 @@ namespace Gauntlet
 					string Line = Reader.ReadLine();
 					LineNumber++;
 					string[] Values = Line.Split(',');
+
+					if (DefaultForMissingColumns != null)
+					{
+						if (Values.Length < HeaderValues.Length)
+						{
+							
+							Values = Enumerable.Concat(Values,Enumerable.Repeat(DefaultForMissingColumns, HeaderValues.Length - Values.Length)).ToArray();
+						}
+					}
+
 					if (Values.Length >= HeaderValues.Length)
 					{
 						for (int i = 0; i < HeaderValues.Length; i++)
