Buildgraph support for virtual bundles. No virtual bundles are actually added with this change, so this should be a NOP.

#rnx
#Synced-CL 50087640

#Synced-CL 50093261

diff --git a/Engine/Source/Programs/AutomationTool/AutomationUtils/BundleUtils.cs b/Engine/Source/Programs/AutomationTool/AutomationUtils/BundleUtils.cs
--- a/Engine/Source/Programs/AutomationTool/AutomationUtils/BundleUtils.cs
+++ b/Engine/Source/Programs/AutomationTool/AutomationUtils/BundleUtils.cs
@@ -28,7 +28,9 @@ namespace AutomationUtils.Automation
 			public bool	UseChunkDBs { get; set; }
 			public bool UseDetailedInstallSizes { get; set; } // default is true
 			public string ExecFileName { get; set; } // TODO: We never used this.  Clean this up.
-			public string CloudAppNameOverride { get; set; }
+			public string CloudAppNameOverride { get; set; }			
+			public List<string> VirtualBundleChildren { get; set; } // Bundles with children should never get files assigned to them, but they still generate chunkdb inis and such.
+			public List<string> VirtualBundleParents { get; set; } // These are the virtual bundles we belong to. We have to be in the entry for the virtual bundle, and they have to be in our entry, to try and avoid mistakes.
 		}
 
 		public class BundleSubsetSettings
@@ -94,6 +96,41 @@ namespace AutomationUtils.Automation
 						Bundle.Tags = new List<string>(); 
 					}
 				}
+				{
+					List<string> Aliases;
+					if (BundleConfig.GetArray(SectionName, "VirtualChildren", out Aliases))
+					{
+						if (Bundle.Name.Contains("virtual", StringComparison.OrdinalIgnoreCase) == false)
+						{
+							throw new AutomationException("Bundle name {0} doesn't contain keyword Virtual in the name but has virtual children", Bundle.Name);
+						}
+						Bundle.VirtualBundleChildren = Aliases;
+						Console.WriteLine($"{Bundle.Name} has children: {string.Join(",", Aliases)}");
+					}
+					else 
+					{
+						Bundle.VirtualBundleChildren = new List<string>(); 
+					}
+				}
+				{
+					List<string> VirtualParents;
+					if (BundleConfig.GetArray(SectionName, "VirtualParents", out VirtualParents))
+					{
+						foreach (string VirtualParent in VirtualParents)
+						{
+							if (VirtualParent.Contains("virtual", StringComparison.OrdinalIgnoreCase) == false)
+							{
+								throw new AutomationException("Bundle name {0} has a non virtual virtual parent (must contain Virtual in the name): {1}", Bundle.Name, VirtualParent);
+							}
+						}
+						Bundle.VirtualBundleParents = VirtualParents;
+						Console.WriteLine($"{Bundle.Name} has parents: {string.Join(",", VirtualParents)}");
+					}
+					else 
+					{
+						Bundle.VirtualBundleParents = new List<string>(); 
+					}
+				}
 				{
 					List<string> Dependencies;
 					if (BundleConfig.GetArray(SectionName, "Dependencies", out Dependencies))
@@ -190,6 +227,35 @@ namespace AutomationUtils.Automation
 
 			// Use OrderBy and not Sort because OrderBy is stable
 			Bundles = Results.OrderBy(b => b.Order).ToDictionary(b => b.Name, b => b);
+			
+						
+			// Iterate over the bundles and ensure that all virtual bundles have a two way mapping to try and prevent
+			// changes making busted patch sizes.
+			foreach (TPlatformBundleSettings Bundle in Bundles.Values)
+			{
+				foreach (string VirtualBundleChild in Bundle.VirtualBundleChildren)
+				{
+					TPlatformBundleSettings ChildBundle;
+					if (!Bundles.TryGetValue(VirtualBundleChild, out ChildBundle))
+					{
+						throw new AutomationException("Virtual bundle {0} references child bundle that doesn't exist: {1}", Bundle.Name, VirtualBundleChild);
+					}
+					bool bFound = false;
+					Console.WriteLine($"{Bundle.Name} Checking parents for {ChildBundle.Name}: {string.Join(",", ChildBundle.VirtualBundleParents)}");
+					foreach (string VirtualBundleParent in ChildBundle.VirtualBundleParents)
+					{
+						if (Bundle.Name.Equals(VirtualBundleParent, StringComparison.OrdinalIgnoreCase))
+						{
+							bFound =  true;
+							break;
+						}
+					}
+					if (!bFound)
+					{
+						throw new AutomationException("Virtual bundle {0} references child bundle that doesn't reference it back: {1}", Bundle.Name, VirtualBundleChild);
+					}
+				}
+			}
 		}
 
 		public static void LoadBundleConfig<TPlatformBundleSettings>(DirectoryReference ProjectDir, UnrealTargetPlatform Platform, string CustomConfig,
diff --git a/Engine/Source/Programs/AutomationTool/AutomationUtils/MCPPublic.cs b/Engine/Source/Programs/AutomationTool/AutomationUtils/MCPPublic.cs
--- a/Engine/Source/Programs/AutomationTool/AutomationUtils/MCPPublic.cs
+++ b/Engine/Source/Programs/AutomationTool/AutomationUtils/MCPPublic.cs
@@ -1399,6 +1399,15 @@ namespace EpicGames.MCP.Automation
 			/// Tag sets to be compared between manifests 
 			/// </summary>
 			public List<HashSet<string>> CompareTagSets;
+			// If the set of tags associated with the install has changed, pass the new set here and the old set in CompareTagSets.
+			// This will report the patch between those sets.
+			public List<HashSet<string>> CompareNewTagSet;
+			
+			// If an optional optimized delta is required, pass it here.
+			public string DeltaFilenameTrailer;
+			
+			// If we expect an optimized delta in the patch, pass this to ensure a failure to avoid looking at unrepresentative data.
+			public bool bRequireOptimizedDelta;
 		}
 
 		public class ManifestDiffOutput
@@ -1632,6 +1641,12 @@ namespace EpicGames.MCP.Automation
 			/// Specifies additional arguments to be passed to BPT
 			/// </summary>
 			public string AdditionalArguments;
+
+			// If a custom optimized delta is required, specify it here.
+			public string DeltaFilenameTrailer;
+			
+			// If we are packaging the chunks for an optimized delta, pass this to make sure if the optimized delta is missing the operation fails.
+			public bool RequireOptimizedDelta;
 		}
 
 		public class PackageChunksOutput
@@ -1664,8 +1679,47 @@ namespace EpicGames.MCP.Automation
 			/// Specifies in bytes, an upper limit for original diffs to try to enhance.
 			/// </summary>
 			public ulong DiffAbortThreshold;
+
+			// If the patch matcher should find matches from more than one tag, pass the set of tags here.
+			// This should be a single entry with the list of tags to match across, e.g.
+			// SourceManifestAliases.Add("MyAliasedTag", {"Tag1", "Tag2"})
+			// The destination manifest can have a different set of tags associated with the alias.
+			public Dictionary<string, List<string>> SourceManifestAliases;
+			public Dictionary<string, List<string>> DestinationManifestAliases;
+			
+			// If generating with a custom optimized delta, specify it here.
+			public string DeltaFilenameTrailer;
 		}
 
+		public class ManifestMetadataOptions
+		{
+			public string SourceManifestFilename;
+			public string OutputFilename; // This can be empty, a temp path will be used instead.
+		}
+		public class ManifestMetadataOutput
+		{
+			public string FeatureLevel;
+			public ulong AppId;
+			public string AppName;
+			public string VersionString;
+			public string UniqueBuildId;
+			public string LaunchExe;
+			public string LaunchCommand;
+			public ulong DownloadSize;
+			public ulong BuildSize;
+			public List<string> FileTagList;
+
+			public class ManifestMetadataOutputFile
+			{
+				public string Filename;
+				public List<string> Tags;
+				public string SHA1;
+				public ulong FileSize;
+			}
+			public List<ManifestMetadataOutputFile> Files;
+		}
+
+
 		static BuildPatchToolBase Handler = null;
 
 		public static BuildPatchToolBase Get()
@@ -1693,6 +1747,9 @@ namespace EpicGames.MCP.Automation
 			return Handler;
 		}
 
+		// Runs BPT to request information about the given metadata.
+		public abstract void Execute(ManifestMetadataOptions Opts, out ManifestMetadataOutput Output, ToolVersion Version = ToolVersion.Live);
+
 		/// <summary>
 		/// Runs the Build Patch Tool executable to generate patch data using the supplied parameters.
 		/// </summary>
