Remove duplicate SetTagFilter command line option from UE.Atuomation.cs

#rnx

#Synced-CL 0

diff --git a/Engine/Source/Programs/AutomationTool/Gauntlet/Unreal/Automation/UE.Automation.cs b/Engine/Source/Programs/AutomationTool/Gauntlet/Unreal/Automation/UE.Automation.cs
--- a/Engine/Source/Programs/AutomationTool/Gauntlet/Unreal/Automation/UE.Automation.cs
+++ b/Engine/Source/Programs/AutomationTool/Gauntlet/Unreal/Automation/UE.Automation.cs
@@ -370,11 +370,6 @@ namespace UE
 				AppConfig.CommandLine += string.Format(" -LogCmds=\"{0}\"", LogCmdsString);
 			}
 
-			if (!string.IsNullOrEmpty(TagFilter))
-			{
-				AppConfig.CommandLine += string.Format(" -ExecCmds=\"Automation SetTagFilter {0}\"", TagFilter);
-			}
-
 		}
 	}
 
