Move intialization failure error message for Engine Test automation report

#rnx

#Synced-CL 50249658

diff --git a/Engine/Source/Programs/AutomationTool/Gauntlet/Unreal/Base/Gauntlet.UnrealTestNode.cs b/Engine/Source/Programs/AutomationTool/Gauntlet/Unreal/Base/Gauntlet.UnrealTestNode.cs
--- a/Engine/Source/Programs/AutomationTool/Gauntlet/Unreal/Base/Gauntlet.UnrealTestNode.cs
+++ b/Engine/Source/Programs/AutomationTool/Gauntlet/Unreal/Base/Gauntlet.UnrealTestNode.cs
@@ -1753,12 +1753,7 @@ namespace Gauntlet
 				{
 					foreach (UnrealRoleResult Item in RoleResults)
 					{
-						if (Item.ProcessResult == UnrealProcessResult.InitializationFailure)
-						{
-							InitPhase.GetStream().AddError("No tests were executed!");
-							InitPhase.SetOutcome(TestPhaseOutcome.Interrupted);
-						}
-						// Since we are going through the Role, lets add the related device
+						// Add the related devices to initialization phase
 						var ReportDevice = AddRoleDeviceToSessionReport(Item, SessionReport);
 						InitPhase.deviceKeys.Add(ReportDevice.key);
 					}
@@ -1855,6 +1850,10 @@ namespace Gauntlet
 							// In case of a fatal error or initialization failure, it is an interrupted outcome
 							if (RoleResult.ProcessResult == UnrealProcessResult.EncounteredFatalError || RoleResult.ProcessResult == UnrealProcessResult.InitializationFailure)
 							{
+								if (RoleResult.ProcessResult == UnrealProcessResult.InitializationFailure)
+								{
+									InitPhase.GetStream().AddError("No tests were executed!");
+								}
 								InitPhase.SetOutcome(TestPhaseOutcome.Interrupted);
 							}
 						}
