[Gauntlet] Prevent stripping commandline argument quotes from json blobs in LongProcessResult

#Synced-CL 50207041

diff --git a/Engine/Source/Programs/AutomationTool/Gauntlet/Unreal/Base/Gauntlet.UnrealTestConfiguration.cs b/Engine/Source/Programs/AutomationTool/Gauntlet/Unreal/Base/Gauntlet.UnrealTestConfiguration.cs
--- a/Engine/Source/Programs/AutomationTool/Gauntlet/Unreal/Base/Gauntlet.UnrealTestConfiguration.cs
+++ b/Engine/Source/Programs/AutomationTool/Gauntlet/Unreal/Base/Gauntlet.UnrealTestConfiguration.cs
@@ -328,7 +328,39 @@ namespace Gauntlet
 
 		public static List<string> GetArgumentList(string Args)
 		{
-			return CommandLineRegex.Matches(Args).Select(Match => Match.Value.Replace("\"", string.Empty)).ToList();
+			List<string> ArgList = CommandLineRegex.Matches(Args).Select(Match => Match.Value).ToList();
+			for (int Index = 0; Index < ArgList.Count; Index++)
+			{
+				bool bJsonScope = false;
+				int ScopeEndIndex = -1;
+				for (int CharIndex = 0; CharIndex < ArgList[Index].Length; CharIndex++)
+				{
+					// Strip quotes as long as they're not embedded in json
+					if (!bJsonScope)
+					{
+						if (ArgList[Index][CharIndex] == '{')
+						{
+							ScopeEndIndex = ArgList[Index].LastIndexOf('}');
+							if (ScopeEndIndex != -1 && ScopeEndIndex > CharIndex)
+							{
+								bJsonScope = true;
+							}
+						}
+
+						if (ArgList[Index][CharIndex] == '"')
+						{
+							ArgList[Index] = ArgList[Index].Remove(CharIndex, 1);
+							CharIndex--;
+						}
+					}
+					else if (CharIndex == ScopeEndIndex)
+					{
+						bJsonScope = false;
+					}
+				}
+			}
+
+			return ArgList;
 		}
 	}
 
