[Automation] Fixing a bug in Gauntlet.UnrealBuildSource.cs that prevented the use of multiple -dpcvar arguments in a commandline.

#Synced-CL 50283614

diff --git a/Engine/Source/Programs/AutomationTool/Gauntlet/Unreal/BuildSource/Gauntlet.UnrealBuildSource.cs b/Engine/Source/Programs/AutomationTool/Gauntlet/Unreal/BuildSource/Gauntlet.UnrealBuildSource.cs
--- a/Engine/Source/Programs/AutomationTool/Gauntlet/Unreal/BuildSource/Gauntlet.UnrealBuildSource.cs
+++ b/Engine/Source/Programs/AutomationTool/Gauntlet/Unreal/BuildSource/Gauntlet.UnrealBuildSource.cs
@@ -723,7 +723,7 @@ namespace Gauntlet
 		private string GenerateProcessedCommandLine(string InCommandLine)
 		{
 			// Break down Commandline into individual tokens
-			Dictionary<string, string> CommandlineTokens = new Dictionary<string, string>(StringComparer.OrdinalIgnoreCase);
+			Dictionary<string, List<string>> CommandlineTokens = new Dictionary<string, List<string>>(StringComparer.OrdinalIgnoreCase);
 			// turn Name(p1,etc) into a collection of Name|(p1,etc) groups
 			MatchCollection Matches = GauntletCommandLine.CommandlineOptionRegex.Matches(InCommandLine);
 
@@ -749,34 +749,52 @@ namespace Gauntlet
 					else if (Name.ToLower() == "execcmds")
 					{
 						// Special cased as execcmds is something that is totally able to be appended to. Requote everything when we're done.
-						CommandlineTokens[Name] = string.Format("\"{0}, {1}\"", CommandlineTokens[Name].Replace("\"", ""), Params.Replace("\"", ""));
+						CommandlineTokens[Name][0] = string.Format("\"{0}, {1}\"", CommandlineTokens[Name][0].Replace("\"", ""), Params.Replace("\"", ""));
 					}
 					else
 					{
-						if (CommandlineTokens[Name] == Params)
+						bool bFoundDuplicateValue = false;
+						foreach (string ExistingValue in CommandlineTokens[Name])
 						{
-							Log.Info(string.Format("Duplicate flag {0}={1} found and ignored. Please fix this. ", Name, Params));
+							if (ExistingValue == Params)
+							{
+								Log.Info(string.Format("Duplicate flag {0}={1} found and ignored. Please fix this. ", Name, Params));
+								bFoundDuplicateValue = true;
+								break;
+							}
 						}
-						else
+
+						if (!bFoundDuplicateValue)
 						{
-							Log.Warning(string.Format("Multiple values for flag {0} found: {1} and {2}. The former value will be discarded. ", Name, CommandlineTokens[Name], Params));
-							CommandlineTokens[Name] = Params.Trim();
+							// dpcvars are also a special parameter because the pattern is to repeat it like: -dpcvars=var1=val1 -dpcvars=var2=val2 etc
+							if (Name.ToLower() != "dpcvars")
+							{
+								Log.Warning(string.Format("Multiple values for flag {0} found: {1} and {2}. The former value will be discarded. ", Name, CommandlineTokens[Name], Params));
+								CommandlineTokens[Name][0] = Params.Trim();
+							}
+							else
+							{
+								CommandlineTokens[Name].Add(Params);
+							}
 						}
 					}
 				}
 				else
 				{
-					CommandlineTokens.Add(Name, Params);
+					CommandlineTokens.Add(Name, new List<string>{ Params });
 				}
 			}
 
 			string CommandlineToReturn = "";
 			foreach (string DictKey in CommandlineTokens.Keys)
 			{
-				CommandlineToReturn += string.Format("-{0}{1} ",
-					DictKey,
-					string.IsNullOrWhiteSpace(CommandlineTokens[DictKey]) ? "" : string.Format("={0}", CommandlineTokens[DictKey])
-					);
+				foreach (string DictValue in CommandlineTokens[DictKey])
+				{
+					CommandlineToReturn += string.Format("-{0}{1} ",
+						DictKey,
+						string.IsNullOrWhiteSpace(DictValue) ? "" : string.Format("={0}", DictValue)
+						);
+				}
 			}
 			Gauntlet.Log.Verbose(string.Format("Pre-formatting Commandline: {0}", InCommandLine));
 			Gauntlet.Log.Verbose(string.Format("Post-formatting Commandline: {0}", CommandlineToReturn));
