Adding logic to avoid cascade failures during test retries when a log file is locked in use by another process

#Synced-CL 50214590

diff --git a/Engine/Source/Programs/AutomationTool/Gauntlet/Framework/Gauntlet.Utils.cs b/Engine/Source/Programs/AutomationTool/Gauntlet/Framework/Gauntlet.Utils.cs
--- a/Engine/Source/Programs/AutomationTool/Gauntlet/Framework/Gauntlet.Utils.cs
+++ b/Engine/Source/Programs/AutomationTool/Gauntlet/Framework/Gauntlet.Utils.cs
@@ -1604,7 +1604,35 @@ namespace Gauntlet
 						foreach(FileSystemInfo Info in Directory.EnumerateFiles("*", SearchOption.AllDirectories))
 						{
 							Info.Attributes = FileAttributes.Normal;
-							Info.Delete(); // throw exceptions here because we have requested a force clean
+							// Try to delete the file, with retry logic for locked files
+							bool bDeleted = false;
+							for (int Attempt = 0; Attempt < 3 && !bDeleted; Attempt++)
+							{
+								try
+								{
+									Info.Delete();
+									bDeleted = true;
+								}
+								catch (System.IO.IOException IoEx) when (IoEx.Message.Contains("because it is being used by another process"))
+								{
+									// Looking specifically for 
+									if (Attempt < 2)
+									{
+										Log.Info("File {File} is locked. Forcing garbage collection and retrying... (Attempt {Attempt}/3)", Info.FullName, Attempt + 1);
+										// Force garbage collection to run finalizers that might be holding file handles
+										System.GC.Collect();
+										System.GC.WaitForPendingFinalizers();
+										System.GC.Collect();
+										// Wait a bit for the OS to release the file handle
+										System.Threading.Thread.Sleep(500);
+									}
+									else
+									{
+										// Final attempt failed, throw the exception
+										throw;
+									}
+								}
+							}
 						}
 					}
 					else
