First version of Windows remote deployment & launching support
- Using the wdremote command line tool for now (https://aka.ms/ToolboxInstaller)
- no Gauntlet support yet

See https://aka.ms/GameRemoteDevtools for details

Add paired devices to engine ini:

%LOCALAPPDATA%\Unreal Engine\Engine\Config\UserEngine.ini

[RemoteWin]
+DeviceNames=myfirstpc
+DeviceNames=mysecondpc

#Synced-CL 50394790

diff --git a/Engine/Source/Developer/Windows/WindowsTargetPlatfomControls/Public/GenericWindowsTargetPlatformControls.h b/Engine/Source/Developer/Windows/WindowsTargetPlatfomControls/Public/GenericWindowsTargetPlatformControls.h
--- a/Engine/Source/Developer/Windows/WindowsTargetPlatfomControls/Public/GenericWindowsTargetPlatformControls.h
+++ b/Engine/Source/Developer/Windows/WindowsTargetPlatfomControls/Public/GenericWindowsTargetPlatformControls.h
@@ -15,6 +15,7 @@
 #endif
 #include "SteamDeck/SteamDeckDevice.h"
 #include "GenericWindowsTargetPlatformSettings.h"
+#include "RemoteWindowsDevice.h"
 
 // IWYU pragma: end_exports
 
@@ -45,11 +46,12 @@ public:
 		// only add local device if actually running on Windows
 		LocalDevice = MakeShareable(new TTargetDevice(*this));
 
-		// quick solution to not having WinGDK steamdeck devices
+		// quick solution to not having WinGDK devices
 		if (this->PlatformName().StartsWith(TEXT("Windows")))
 		{
-			// Check if we have any SteamDeck devices around
+			// Check if we have any SteamDeck or other remote devices around
 			SteamDevices = TSteamDeckDevice<FLocalPcTargetDevice>::DiscoverDevices(*this, TEXT("Proton"));
+			RemoteDevices = FRemoteWindowsDevice::DiscoverDevices(*this);
 		}
 #endif
 		GenericWindowsTargetPlatformSettings = (TSuperSettings*)TargetPlatformSettings;
@@ -76,6 +78,14 @@ public:
 				OutDevices.Add(SteamDeck);
 			}
 		}
+
+		for (const ITargetDevicePtr& RemoteDevice : RemoteDevices)
+		{
+			if (RemoteDevice.IsValid())
+			{
+				OutDevices.Add(RemoteDevice);
+			}
+		}
 	}
 
 	virtual bool GenerateStreamingInstallManifest(const TMultiMap<FString, int32>& PakchunkMap, const TSet<int32>& PakchunkIndicesInUse) const override
@@ -108,6 +118,14 @@ public:
 			}
 		}
 
+		for (const ITargetDevicePtr& RemoteDevice : RemoteDevices)
+		{
+			if (RemoteDevice.IsValid() && DeviceId == RemoteDevice->GetId())
+			{
+				return RemoteDevice;
+			}
+		}
+
 		return nullptr;
 	}
 
@@ -178,6 +196,7 @@ private:
 	ITargetDevicePtr LocalDevice;
 
 	TArray<ITargetDevicePtr> SteamDevices;
+	TArray<ITargetDevicePtr> RemoteDevices;
 
 	TSuperSettings* GenericWindowsTargetPlatformSettings;
 };
diff --git a/Engine/Source/Developer/Windows/WindowsTargetPlatfomControls/Public/RemoteWindowsDevice.h b/Engine/Source/Developer/Windows/WindowsTargetPlatfomControls/Public/RemoteWindowsDevice.h
new file mode 100644
--- /dev/null
+++ b/Engine/Source/Developer/Windows/WindowsTargetPlatfomControls/Public/RemoteWindowsDevice.h
@@ -0,0 +1,191 @@
+// Copyright Epic Games, Inc. All Rights Reserved.
+
+#pragma once
+
+// HEADER_UNIT_UNSUPPORTED - Common/TargetPlatformBase not part of additional include paths
+
+#include "CoreMinimal.h"
+#include "Common/TargetPlatformBase.h"
+#include "Logging/LogMacros.h"
+#include "Misc/Paths.h"
+#include "HAL/PlatformProcess.h"
+#include "HAL/PlatformMisc.h"
+
+DEFINE_LOG_CATEGORY_STATIC(LogRemoteWindowsTargetDevice, Log, All);
+
+/* Remote Windows device support currently requires the Microsoft.Gaming.RemoteIterationClient_Microsoft WinGet package. 
+*  (The remote devices need the Microsoft.Gaming.RemoteIterationEndpoint_Microsoft WinGetpackage)
+* 
+* It is recommended to use the Xbox PC Toolbox installer to configure host and remote devices : https://aka.ms/ToolboxInstaller
+* See https://aka.ms/GameRemoteDevtools for more details
+* 
+* At present, the list of devices is read from the engine ini. These devices should be paired already.
+* 
+*  %LOCALAPPDATA%\Unreal Engine\Engine\Config\UserEngine.ini  can be used for convenience
+*  
+*  [RemoteWin]
+*  +DeviceNames=myfirstpc
+*  +DeviceNames=mysecondpc
+*/
+
+
+class FRemoteWindowsDevice : public ITargetDevice
+{
+public:
+	FRemoteWindowsDevice(const ITargetPlatformControls& InTargetPlatformControls, const FString& InDeviceName)
+		: TargetPlatformControls(InTargetPlatformControls)
+		, DeviceName(InDeviceName)
+	{
+	}
+
+	virtual bool Connect() override
+	{
+		return true;
+	}
+
+	virtual void Disconnect() override
+	{
+	}
+
+	virtual ETargetDeviceTypes GetDeviceType() const override
+	{
+		return ETargetDeviceTypes::Desktop;
+	}
+
+	virtual FTargetDeviceId GetId() const override
+	{
+		return FTargetDeviceId(TargetPlatformControls.PlatformName(), GetName());
+	}
+
+	virtual FString GetName() const override
+	{
+		return DeviceName;
+	}
+
+	virtual FString GetOperatingSystemName() override
+	{
+		// This may not always be 64-bit in the future but there's no way to query that yet.
+		return TEXT("Windows (64-bit)");
+	}
+
+	virtual int32 GetProcessSnapshot(TArray<FTargetDeviceProcessInfo>& OutProcessInfos) override
+	{
+		return 0;
+	}
+
+	virtual const class ITargetPlatformSettings& GetPlatformSettings() const override
+	{
+		return *(TargetPlatformControls.GetTargetPlatformSettings());
+	}
+
+	virtual const class ITargetPlatformControls& GetPlatformControls() const override
+	{
+		return TargetPlatformControls;
+	}
+
+	virtual bool IsConnected() override
+	{
+		return true;
+	}
+
+	virtual bool IsDefault() const override
+	{
+		return true;
+	}
+
+	virtual bool PowerOff(bool Force) override
+	{
+		return false;
+	}
+
+	virtual bool PowerOn() override
+	{
+		return false;
+	}
+
+	virtual bool Reboot(bool bReconnect = false) override
+	{
+		return false;
+	}
+
+	virtual bool SupportsFeature(ETargetDeviceFeatures Feature) const override
+	{
+		return false;
+	}
+
+	virtual bool TerminateProcess(const int64 ProcessId) override
+	{
+		return RunWdRemote( FString::Printf(TEXT("/action:terminate /processid:%lld"), ProcessId));
+	}
+
+	virtual bool TerminateLaunchedProcess(const FString& ProcessIdentifier)
+	{
+		int64 ProcessId = FCString::Atoi64(*ProcessIdentifier);
+		return (ProcessId > 0) && TerminateProcess(ProcessId);
+	}
+
+	virtual void SetUserCredentials(const FString& UserName, const FString& UserPassword) override {}
+
+	virtual bool GetUserCredentials(FString& OutUserName, FString& OutUserPassword) override
+	{
+		return false;
+	}
+
+	static TArray<FString> GetAllRemoteDeviceNames()
+	{
+		TArray<FString> Result;
+
+		// read device names from the engine ini. must already be paired
+		// best place is %LOCALAPPDATA%\Unreal Engine\Engine\Config\UserEngine.ini
+		TArray<FString> DeviceNames;
+		GConfig->GetArray(TEXT("RemoteWin"), TEXT("DeviceNames"), DeviceNames, GEngineIni);
+		Result.Append(DeviceNames);
+
+		return MoveTemp(Result);
+	}
+
+	static TArray<ITargetDevicePtr> DiscoverDevices(const ITargetPlatformControls& TargetPlatformControls)
+	{
+		TArray<ITargetDevicePtr> Devices;
+		for (const FString& DeviceName : GetAllRemoteDeviceNames())
+		{
+			Devices.Add(MakeShared<FRemoteWindowsDevice>(TargetPlatformControls, DeviceName));
+		}
+
+		return Devices;
+	}
+
+private:
+	bool RunWdRemote( const FString& Parameters ) const
+	{
+		return RunWdRemote( Parameters, DeviceName );
+	}
+
+	static bool RunWdRemote( const FString& Parameters, const FString& DeviceName )
+	{
+		FString CommandLineArgs = Parameters;
+		if (!DeviceName.IsEmpty())
+		{
+			CommandLineArgs += FString::Printf(TEXT(" /device:%s"), *DeviceName);
+		}
+
+		FString LocalAppData = FPlatformMisc::GetEnvironmentVariable(TEXT("LOCALAPPDATA"));
+		FString WdRemotePath = FPaths::Combine(LocalAppData, TEXT("Microsoft\\WinGet\\Links\\wdRemote.exe") );
+		if (!FPaths::FileExists(WdRemotePath))
+		{
+			UE_LOG( LogRemoteWindowsTargetDevice, Error, TEXT("%s not found!"), *WdRemotePath);
+			return false;
+		}
+
+		int32 ReturnCode = 0;
+		FString StdOut;
+		FString StdErr;
+		FPlatformProcess::ExecProcess(*WdRemotePath, *CommandLineArgs, &ReturnCode, &StdOut, &StdErr);
+
+		UE_CLOG( (ReturnCode != 0), LogRemoteWindowsTargetDevice, Error, TEXT("wdremote %s  - failed (0x%X)\n%s\n%s"), *CommandLineArgs, ReturnCode, *StdOut, *StdErr);
+		return (ReturnCode == 0);
+	}
+
+	const ITargetPlatformControls& TargetPlatformControls;
+	FString DeviceName;
+};
diff --git a/Engine/Source/Programs/AutomationTool/Win/RemoteWinSupport.cs b/Engine/Source/Programs/AutomationTool/Win/RemoteWinSupport.cs
new file mode 100644
--- /dev/null
+++ b/Engine/Source/Programs/AutomationTool/Win/RemoteWinSupport.cs
@@ -0,0 +1,306 @@
+// Copyright Epic Games, Inc. All Rights Reserved.
+using AutomationTool;
+using EpicGames.Core;
+using Google.Protobuf.WellKnownTypes;
+using Microsoft.Extensions.Logging;
+using System;
+using System.Collections.Generic;
+using System.Diagnostics;
+using System.IO;
+using UnrealBuildTool;
+
+using static AutomationTool.CommandUtils;
+
+
+/* Remote Windows device support currently requires the Microsoft.Gaming.RemoteIterationClient_Microsoft WinGet package. 
+* (The remote devices need the Microsoft.Gaming.RemoteIterationEndpoint_Microsoft WinGetpackage)
+* 
+* It is recommended to use the Xbox PC Toolbox installer to configure host and remote devices : https://aka.ms/ToolboxInstaller
+* See https://aka.ms/GameRemoteDevtools for more details
+* 
+* At present, the list of devices is read from the engine ini. These devices should be paired already.
+* 
+*  %LOCALAPPDATA%\Unreal Engine\Engine\Config\UserEngine.ini  can be used for convenience
+*  
+*  [RemoteWin]
+*  +DeviceNames=myfirstpc
+*  +DeviceNames=mysecondpc
+*/
+
+public static class RemoteWinSupport
+{
+	public static readonly string DeviceType = "RemoteWin";
+	private static string WdRemotePath => Environment.ExpandEnvironmentVariables("%LocalAppData%\\Microsoft\\WinGet\\Links\\wdRemote.exe");
+
+	public static List<DeviceInfo> GetDevices(UnrealTargetPlatform Platform)
+	{
+		List<DeviceInfo> Devices = [];
+
+		// read device names from the engine ini. must already be paired
+		// best place is %LOCALAPPDATA%\Unreal Engine\Engine\Config\UserEngine.ini
+		ConfigHierarchy EngineConfig = ConfigCache.ReadHierarchy(ConfigHierarchyType.Engine, null, Platform);
+		if (EngineConfig.GetArray("RemoteWin", "DeviceNames", out List<string> DeviceNames))
+		{
+			foreach (string DeviceName in DeviceNames)
+			{
+				string SoftwareVersion = Environment.OSVersion.Version.ToString(); // NB. this is the local version, not the remote
+				Devices.Add(new DeviceInfo(Platform, DeviceName, DeviceName, SoftwareVersion, DeviceType, true, true));
+			}
+		}
+
+		return Devices;
+	}
+
+	public static bool Deploy(UnrealTargetPlatform Platform, ProjectParams Params, DeploymentContext SC)
+	{
+		if (Params.Devices.Count == 0)
+		{
+			return false;
+		}
+		
+		foreach (string Device in Params.Devices)
+		{
+			if (Platform == GetTargetDevicePlatform(Device))
+			{
+				string DeviceName = GetTargetDeviceName(Device);
+				string SourcePath = SC.StageDirectory.FullName;
+				string DestinationPath = GetDestinationRoot(Platform, Params);
+
+				if (!RunWdRemote($"/action:deploy /source:{StringUtils.QuoteArgument(SourcePath)} /destination:{StringUtils.QuoteArgument(DestinationPath)}", DeviceName))
+				{
+					throw new AutomationException("Remote deployment to {0} failed", Params.Devices[0]);
+				}
+			}
+		}
+
+		return true;
+	}
+
+	public static IProcessResult RunClient(UnrealTargetPlatform Platform, CommandUtils.ERunOptions ClientRunFlags, string ClientApp, string ClientCmdLine, ProjectParams Params)
+	{
+		if (Params.Devices.Count == 0)
+		{
+			return null;
+		}
+
+		ClientCmdLine = SanitizeCommandLine(ClientCmdLine);
+
+		RemoteWinProcessResult Result = null;
+		foreach (string Device in Params.Devices)
+		{
+			if (Platform == GetTargetDevicePlatform(Device))
+			{
+				string DeviceName = GetTargetDeviceName(Device);
+				string ExecutablePath = GetRemoteExecutablePath(Platform, ClientApp, Params);
+
+				if (!RunWdRemote($"/action:launch /path:{StringUtils.QuoteArgument(ExecutablePath)} /args:\"{ClientCmdLine}\"", DeviceName))
+				{
+					throw new AutomationException("Remote launch on {0} failed", Device);
+				}
+
+				Logger.LogInformation("Running Package@Device:{ProcessName}@{DeviceName}", Path.GetFileName(ClientApp), DeviceName);
+				Result = new RemoteWinProcessResult(DeviceName, Path.GetFileName(ClientApp));
+			}
+		}
+
+		if (Result != null && !ClientRunFlags.HasFlag(CommandUtils.ERunOptions.NoWaitForExit))
+		{
+			Result.WaitForExit();
+		}
+
+		return Result;
+	}
+
+	private static string GetRemoteExecutablePath(UnrealTargetPlatform Platform, string ClientApp, ProjectParams Params)
+	{
+		// find the remote executable path
+		string RelativeExecutablePath = Path.GetFileName(ClientApp);
+		string[] Paths = ClientApp.Split([Path.DirectorySeparatorChar, Path.AltDirectorySeparatorChar], StringSplitOptions.RemoveEmptyEntries);
+		if (Paths.Length >= 4 &&
+			Paths[^3].Equals("Binaries", StringComparison.OrdinalIgnoreCase) &&
+			Paths[^2].Equals(Platform.ToString(), StringComparison.OrdinalIgnoreCase))
+		{
+			RelativeExecutablePath = Path.Combine(Paths[^4..]); // MyGame/Binaries/Win64/MyGame-Test.exe etc
+		}
+		string ExecutablePath = Path.Combine(GetDestinationRoot(Platform, Params), RelativeExecutablePath);
+		return ExecutablePath;
+	}
+
+
+	private static string SanitizeCommandLine(string ClientCmdLine)
+	{
+		// simplistic helper function - may need fleshing out for more complicated edge cases as they're discovered
+		void RemoveParam(string Key, bool bOnlyIfUsingAbsolutePath = false)
+		{
+			int Start = ClientCmdLine.IndexOf(Key);
+			if (Start >= 0)
+			{
+				int End;
+
+				bool bQuote = false;
+				for ( End = Start+1; End < ClientCmdLine.Length; End++ )
+				{
+					if ( ClientCmdLine[End] == '"' && ClientCmdLine[End-1] != '\\')
+					{
+						bQuote = !bQuote;
+					}
+					else if (!bQuote && ClientCmdLine[End] == ' ')
+					{
+						break;
+					}
+				}
+
+				if (End > Start)
+				{
+					string Arg = ClientCmdLine[Start..End];
+					if (bOnlyIfUsingAbsolutePath && Key.Contains('='))
+					{
+						string Value = StringUtils.StripQuoteArgument(ClientCmdLine[(Start + Key.Length)..End]);
+						if (Path.IsPathRooted(Value))
+						{
+							Logger.LogWarning("remote launch ignoring parameter {key} due to absolute path: {val} ", Key[..^1], Value);
+						}
+						else
+						{
+							return;
+						}
+					}
+
+					ClientCmdLine = ClientCmdLine.Replace(Arg,"");
+				}
+			}
+		}
+
+		// trim off any known absolute paths
+		RemoveParam("-project=", bOnlyIfUsingAbsolutePath:true);
+		RemoveParam("-abslog=");
+
+		// escape all quotes
+		// @fixme: this is only necessary because we're using wdremote.exe who's /args needs a quoted string. 
+		//         this will break if ClientCmdLine already contains escaped quotes!
+		ClientCmdLine = ClientCmdLine.Replace("\"", "\\\"");
+
+		return ClientCmdLine;
+	}
+
+	public static string GetTargetDeviceName(string DeviceName)
+	{
+		string[] Tokens = DeviceName.Split(['@', ' ', '(']); // platform@devicename (extra)
+		if (Tokens.Length > 1)
+		{
+			return Tokens[1];
+		}
+
+		return DeviceName;
+	}
+
+	private static UnrealTargetPlatform GetTargetDevicePlatform(string DeviceName)
+	{
+		string[] Tokens = DeviceName.Split(['@', ' ', '(']); // platform@devicename (extra)
+		if (Tokens.Length > 1)
+		{
+			if (Tokens[0] == "Windows")
+			{
+				return UnrealTargetPlatform.Win64;
+			}
+			else if (UnrealTargetPlatform.TryParse( Tokens[0], out UnrealTargetPlatform Platform))
+			{
+				return Platform;
+			}
+		}
+
+		throw new AutomationException("cannot get target platform for remote deployment from {0}", DeviceName);
+	}
+
+	private static string GetDestinationRoot(UnrealTargetPlatform Platform, ProjectParams Params)
+	{
+		return $"{Params.ShortProjectName}_{Platform}";
+	}
+
+
+	internal static bool RunWdRemote( string Parameters, string Device )
+	{
+		if (!File.Exists(WdRemotePath))
+		{
+			throw new AutomationException("{0} not found", WdRemotePath);
+		}
+
+		Parameters += $" /device:{Device}";
+		RunAndLog(WdRemotePath, Parameters, out int SuccessCode);
+
+		if (SuccessCode != 0)
+		{
+			Logger.LogError("wdremote {params} failed - exit code 0x{code}", Parameters, SuccessCode.ToString("X"));
+		}
+
+		return (SuccessCode == 0);
+	}
+
+
+
+	private class RemoteWinProcessResult : IProcessResult
+	{
+		private readonly string DeviceName;
+		private readonly string ProcessName;
+		private int ProcessExitCode = 0;
+
+		public RemoteWinProcessResult(string InDeviceName, string InProcessName)
+		{
+			DeviceName = InDeviceName;
+			ProcessName = InProcessName;
+		}
+
+		public void StopProcess(bool KillDescendants = true)
+		{
+			RemoteWinSupport.RunWdRemote("/action:terminate", DeviceName);
+		}
+
+		public bool HasExited => true; // @todo: cannot monitor the remote process yet. assume it has terminated so launching doesn't get stuck waiting indefinitely
+
+		public string GetProcessName()
+		{
+			return ProcessName;
+		}
+
+		public void OnProcessExited() { }
+
+		public void DisposeProcess() { }
+
+		public void StdOut(object sender, DataReceivedEventArgs e) { }
+
+		public void StdErr(object sender, DataReceivedEventArgs e) { }
+
+		public int ExitCode
+		{
+			get { return ProcessExitCode; }
+			set { ProcessExitCode = value; }
+		}
+
+		public bool bExitCodeSuccess => ExitCode == 0;
+
+		public string Output => "";
+
+		public Process ProcessObject => null;
+
+		public new string ToString()
+		{
+			return $"RemoteWin RemoteWinProcessResult: {GetProcessName()} (ExitCode: {ExitCode})";
+		}
+
+		public void WaitForExit()
+		{
+		}
+
+		public FileReference WriteOutputToFile(string FileName)
+		{
+			using (StreamWriter writer = new(FileName))
+			{
+				writer.Write(Output);
+			}
+
+			return new FileReference(FileName);
+		}
+	}
+
+
+}
diff --git a/Engine/Source/Programs/AutomationTool/Win/WinPlatform.Automation.cs b/Engine/Source/Programs/AutomationTool/Win/WinPlatform.Automation.cs
--- a/Engine/Source/Programs/AutomationTool/Win/WinPlatform.Automation.cs
+++ b/Engine/Source/Programs/AutomationTool/Win/WinPlatform.Automation.cs
@@ -44,6 +44,7 @@ public class Win64Platform : Platform
 			Devices.Add(LocalMachine);
 
 			Devices.AddRange(SteamDeckSupport.GetDevices(UnrealTargetPlatform.Win64));
+			Devices.AddRange(RemoteWinSupport.GetDevices(UnrealTargetPlatform.Win64));
 		}
 
 		return Devices.ToArray();
@@ -51,11 +52,15 @@ public class Win64Platform : Platform
 
 	public override void Deploy(ProjectParams Params, DeploymentContext SC)
 	{
-		// We only care about deploying for SteamDeck
+		// We only care about deploying for SteamDeck or other remote devices
 		if (Params.Devices.Count == 1 && GetDevices().FirstOrDefault(x => x.Id == Params.DeviceNames[0])?.Type == "SteamDeck")
 		{
 			SteamDeckSupport.Deploy(UnrealTargetPlatform.Win64, Params, SC);
 		}
+		else if (Params.Devices.Count > 0 && GetDevices().FirstOrDefault(x => Params.DeviceNames.Contains(x.Id))?.Type == "RemoteWin")
+		{
+			RemoteWinSupport.Deploy(PlatformType, Params, SC);
+		}
 	}
 
 	public override IProcessResult RunClient(ERunOptions ClientRunFlags, string ClientApp, string ClientCmdLine, ProjectParams Params)
@@ -64,6 +69,10 @@ public class Win64Platform : Platform
 		{
 			return SteamDeckSupport.RunClient(UnrealTargetPlatform.Win64, ClientRunFlags, ClientApp, ClientCmdLine, Params);
 		}
+		else if (Params.Devices.Count > 0 && GetDevices().FirstOrDefault(x => Params.DeviceNames.Contains(x.Id))?.Type == "RemoteWin")
+		{
+			return RemoteWinSupport.RunClient(PlatformType, ClientRunFlags, ClientApp, ClientCmdLine, Params);
+		}
 
 		return base.RunClient(ClientRunFlags, ClientApp, ClientCmdLine, Params);
 	}
