Add support for remote device references without a platform:

-device=windows[REDACTED_HANDLE]
-device=[REDACTED_HANDLE]
-device=mypc

#rnx

#Synced-CL 50394790

diff --git a/Engine/Source/Programs/AutomationTool/Win/RemoteWinSupport.cs b/Engine/Source/Programs/AutomationTool/Win/RemoteWinSupport.cs
--- a/Engine/Source/Programs/AutomationTool/Win/RemoteWinSupport.cs
+++ b/Engine/Source/Programs/AutomationTool/Win/RemoteWinSupport.cs
@@ -60,7 +60,7 @@ public static class RemoteWinSupport
 		
 		foreach (string Device in Params.Devices)
 		{
-			if (Platform == GetTargetDevicePlatform(Device))
+			if (Platform == GetTargetDevicePlatform(Device, Platform))
 			{
 				string DeviceName = GetTargetDeviceName(Device);
 				string SourcePath = SC.StageDirectory.FullName;
@@ -88,7 +88,7 @@ public static class RemoteWinSupport
 		RemoteWinProcessResult Result = null;
 		foreach (string Device in Params.Devices)
 		{
-			if (Platform == GetTargetDevicePlatform(Device))
+			if (Platform == GetTargetDevicePlatform(Device, Platform))
 			{
 				string DeviceName = GetTargetDeviceName(Device);
 				string ExecutablePath = GetRemoteExecutablePath(Platform, ClientApp, Params);
@@ -185,7 +185,7 @@ public static class RemoteWinSupport
 
 	public static string GetTargetDeviceName(string DeviceName)
 	{
-		string[] Tokens = DeviceName.Split(['@', ' ', '(']); // platform@devicename (extra)
+		string[] Tokens = DeviceName.Split(['@', ' ', '(']); // 'platform@devicename (extra)'   or '@device (extra)' or just 'device'
 		if (Tokens.Length > 1)
 		{
 			return Tokens[1];
@@ -194,20 +194,28 @@ public static class RemoteWinSupport
 		return DeviceName;
 	}
 
-	private static UnrealTargetPlatform GetTargetDevicePlatform(string DeviceName)
+	private static UnrealTargetPlatform GetTargetDevicePlatform(string DeviceName, UnrealTargetPlatform DefaultPlatform)
 	{
-		string[] Tokens = DeviceName.Split(['@', ' ', '(']); // platform@devicename (extra)
+		string[] Tokens = DeviceName.Split(['@', ' ', '(']); // 'platform@devicename (extra)'   or '@device (extra)' or just 'device'
 		if (Tokens.Length > 1)
 		{
 			if (Tokens[0] == "Windows")
 			{
 				return UnrealTargetPlatform.Win64;
 			}
+			else if (Tokens[0].Length == 0)
+			{
+				return DefaultPlatform;
+			}
 			else if (UnrealTargetPlatform.TryParse( Tokens[0], out UnrealTargetPlatform Platform))
 			{
 				return Platform;
 			}
 		}
+		else if (Tokens.Length == 1)
+		{
+			return DefaultPlatform;
+		}
 
 		throw new AutomationException("cannot get target platform for remote deployment from {0}", DeviceName);
 	}
