Reporter class has been implemented, the logic of the Report* methods from UnrealTestNode has been moved into Reporter class. Report* methods are marked as Deprecated.
#rnx

#Synced-CL 50417652

diff --git a/Engine/Source/Programs/AutomationTool/Gauntlet/Framework/Gauntlet.Utils.cs b/Engine/Source/Programs/AutomationTool/Gauntlet/Framework/Gauntlet.Utils.cs
--- a/Engine/Source/Programs/AutomationTool/Gauntlet/Framework/Gauntlet.Utils.cs
+++ b/Engine/Source/Programs/AutomationTool/Gauntlet/Framework/Gauntlet.Utils.cs
@@ -1,25 +1,26 @@
 // Copyright Epic Games, Inc. All Rights Reserved.
 
+using AutomationTool;
+using EpicGames.Core;
+using Gauntlet.TestTracking;
+using Google.Protobuf;
+using ImageMagick;
+using Microsoft.Extensions.Logging;
 using System;
 using System.Collections.Generic;
+using System.Diagnostics;
 using System.IO;
 using System.Linq;
-using System.Diagnostics;
-using AutomationTool;
-using System.Threading;
-using System.Text.RegularExpressions;
-using System.Threading.Tasks;
 using System.Reflection;
-using ImageMagick;
-using UnrealBuildTool;
-using EpicGames.Core;
 using System.Security.Cryptography;
 using System.Text;
-using Microsoft.Extensions.Logging;
-using Logging = Microsoft.Extensions.Logging;
-
-using static AutomationTool.CommandUtils;
+using System.Text.RegularExpressions;
+using System.Threading;
+using System.Threading.Tasks;
 using UnrealBuildBase;
+using UnrealBuildTool;
+using static AutomationTool.CommandUtils;
+using Logging = Microsoft.Extensions.Logging;
 
 namespace Gauntlet
 {
@@ -478,6 +479,140 @@ namespace Gauntlet
 		}
 	}
 
+#nullable enable
+	/// <summary>
+	/// Helper class providing logging methods for test events of various levels.
+	/// Used to clearly display test events on the Horde dashboard, eliminates inconsistent reporting.
+	/// </summary>
+	public class Reporter
+	{
+		private static volatile Reporter? _Instance;
+		private static readonly object _Lock = new object();
+
+		private Reporter()
+		{
+			;
+		}
+
+		public static Reporter Instance
+		{
+			get
+			{
+				if (_Instance == null)
+				{
+					lock (_Lock)
+					{
+						if (_Instance == null)
+						{
+							_Instance = new Reporter();
+						}
+					}
+				}
+				return _Instance;
+			}
+		}
+
+		private PhaseQueue? _Phases;
+		private Action<TestResult>? _SetUnrealTestResult;
+		private BaseTest? _TestNode;
+
+		/// <summary>
+		/// This method is used to switch context when running tests in parallel.
+		/// The method should be called each time TestExecutor switches between test nodes, such as in node loops, etc.
+		/// </summary>
+		/// <param name="Phases">Reference to Phases variable in test node</param>
+		/// <param name="SetUnrealTestResult">Reference to SetUnrealTestResult method from test node</param>
+		/// <param name="TestNode">Reference to test node that is calling this method</param>
+		public void UpdateReferences(PhaseQueue Phases, Action<TestResult> SetUnrealTestResult, BaseTest TestNode)
+		{
+			this._Phases = Phases;
+			this._SetUnrealTestResult = SetUnrealTestResult;
+			this._TestNode = TestNode;
+		}
+
+		/// <summary>
+		/// Report a test-halting error that is NOT directly related to the behavior that is being tested.
+		/// </summary>
+		/// <param name="Message"></param>
+		/// <param name="Args"></param>
+		public void Interrupt(string Message, params object[] Args)
+		{
+			Instance?._Phases?.Current?.AddEvent(Logging.LogLevel.Critical, Message, Args);
+			Instance?._TestNode?.AddTestEvent(new UnrealTestEvent(EventSeverity.Fatal, Message, Args));
+			
+			if (_TestNode == null || !_TestNode.LogWarningsAndErrorsAfterSummary)
+			{
+				Log.Critical(KnownLogEvents.Gauntlet_FatalEvent, Message, Args);
+			}
+
+			if (Instance?._TestNode?.GetTestStatus() == TestStatus.Complete && Instance._TestNode.GetTestResult() == TestResult.Passed)
+			{
+				Instance?._SetUnrealTestResult?.Invoke(TestResult.OperationalException);
+			}
+		}
+
+		/// <summary>
+		/// Report a test failure.  Specifically, the behavior that is targeted by the test case did not function as expected.
+		/// </summary>
+		/// <param name="Message"></param>
+		/// <param name="Args"></param>
+		public void Failure(string Message, params object[] Args)
+		{
+			Instance?._Phases?.Current?.AddEvent(Logging.LogLevel.Error, Message, Args);
+			Instance?._TestNode?.AddTestEvent(new UnrealTestEvent(EventSeverity.Error, Message, Args));
+
+			if (_TestNode == null || !_TestNode.LogWarningsAndErrorsAfterSummary)
+			{
+				Log.Error(KnownLogEvents.Gauntlet_TestEvent, Message, Args);
+			}
+
+			if (Instance?._TestNode?.GetTestStatus() == TestStatus.Complete && Instance._TestNode.GetTestResult() == TestResult.Passed)
+			{
+				Instance?._SetUnrealTestResult?.Invoke(TestResult.Failed);
+			}
+		}
+
+		/// <summary>
+		/// Report a warning with Gauntlet_TestEvent
+		/// </summary>
+		/// <param name="Message"></param>
+		/// <param name="Args"></param>
+		public void Warning(string Message, params object[] Args)
+		{
+			Warning(KnownLogEvents.Gauntlet_TestEvent, Message, Args);
+		}
+
+		/// <summary>
+		/// Report a warning with required Log Event type
+		/// </summary>
+		/// <param name="EventId"></param>
+		/// <param name="Message"></param>
+		/// <param name="Args"></param>
+		public void Warning(EventId EventId, string Message, params object[] Args)
+		{
+			Instance?._Phases?.Current?.AddEvent(Logging.LogLevel.Warning, Message, Args);
+			Instance?._TestNode?.AddTestEvent(new UnrealTestEvent(EventSeverity.Warning, Message, Args));
+
+			if (_TestNode == null || !_TestNode.LogWarningsAndErrorsAfterSummary)
+			{
+				Log.Warning(EventId, Message, Args);
+			}
+		}
+
+		/// <summary>
+		/// Report information
+		/// </summary>
+		/// <param name="Message"></param>
+		/// <param name="Args"></param>
+		public void Info(string Message, params object[] Args)
+		{
+			Instance?._Phases?.Current?.AddEvent(Logging.LogLevel.Information, Message, Args);
+			Instance?._TestNode?.AddTestEvent(new UnrealTestEvent(EventSeverity.Info, Message, Args));
+			Log.Info(KnownLogEvents.Gauntlet_TestEvent, Message, Args);
+		}
+	}
+#nullable disable
+
 	public class Hasher
 	{
 		public static string ComputeHash(string Input)
diff --git a/Engine/Source/Programs/AutomationTool/Gauntlet/Unreal/Base/Gauntlet.UnrealTestNode.cs b/Engine/Source/Programs/AutomationTool/Gauntlet/Unreal/Base/Gauntlet.UnrealTestNode.cs
--- a/Engine/Source/Programs/AutomationTool/Gauntlet/Unreal/Base/Gauntlet.UnrealTestNode.cs
+++ b/Engine/Source/Programs/AutomationTool/Gauntlet/Unreal/Base/Gauntlet.UnrealTestNode.cs
@@ -10,7 +10,6 @@ using System.Linq;
 using EpicGames.Core;
 using EpicGames.LogMatchers;
 using Microsoft.Extensions.Logging;
-using Logging = Microsoft.Extensions.Logging;
 using static Gauntlet.HordeReport;
 using static Gauntlet.UnrealSessionInstance;
 using Gauntlet.Utils;
@@ -210,65 +209,51 @@ namespace Gauntlet
 		/// <param name="Args"></param>
 		public virtual void ReportError(string Message, params object[] Args)
 		{
-			// While "ReportError" might logically/syntactically be closer to the "ReportInterrupt" functionality, it historically matches
-			// up with "ReportFailure".  That's why we have all three (for now, at least).
-			ReportFailure(Message, Args);
+			Reporter.Instance.Failure(Message, Args);
 		}
 
-        /// <summary>
-        /// Report a test-halting error that is NOT directly related to the behavior that is being tested.
-        /// </summary>
-        /// <param name="Message"></param>
-        /// <param name="Args"></param>
-        public virtual void ReportInterrupt(string Message, params object[] Args)
+		/// <summary>
+		/// Report a test-halting error that is NOT directly related to the behavior that is being tested.
+		/// This method is deprecated. Please use Gauntlet.Reporter.Instance.Interrupt instead of this method.
+		/// </summary>
+		/// <param name="Message"></param>
+		/// <param name="Args"></param>
+		public virtual void ReportInterrupt(string Message, params object[] Args)
         {
-			Phases.Current?.AddEvent(Logging.LogLevel.Critical, Message, Args);
-            AddTestEvent(new UnrealTestEvent(EventSeverity.Fatal, Message, Args));
-            if (!LogWarningsAndErrorsAfterSummary) { Log.Critical(KnownLogEvents.Gauntlet_FatalEvent, Message, Args); }
-            if (GetTestStatus() == TestStatus.Complete && GetTestResult() == TestResult.Passed)
-            {
-                SetUnrealTestResult(TestResult.OperationalException);
-            }
-        }
+			Reporter.Instance.Interrupt(Message, Args);
+		}
 
-        /// <summary>
-        /// Report a test failure.  Specifically, the behavior that is targeted by the test case did not function as expected.
-        /// </summary>
-        /// <param name="Message"></param>
-        /// <param name="Args"></param>
-        public virtual void ReportFailure(string Message, params object[] Args)
+		/// <summary>
+		/// Report a test failure.  Specifically, the behavior that is targeted by the test case did not function as expected.
+		/// This method is deprecated. Please use Gauntlet.Reporter.Instance.Failure instead of this method.
+		/// </summary>
+		/// <param name="Message"></param>
+		/// <param name="Args"></param>
+		public virtual void ReportFailure(string Message, params object[] Args)
         {
-			Phases.Current?.AddEvent(Logging.LogLevel.Error, Message, Args);
-            AddTestEvent(new UnrealTestEvent(EventSeverity.Error, Message, Args));
-            if (!LogWarningsAndErrorsAfterSummary) { Log.Error(KnownLogEvents.Gauntlet_TestEvent, Message, Args); }
-            if (GetTestStatus() == TestStatus.Complete && GetTestResult() == TestResult.Passed)
-            {
-                SetUnrealTestResult(TestResult.Failed);
-            }
-        }
+			Reporter.Instance.Failure(Message, Args);
+		}
 
-        /// <summary>
-        /// Report a warning
-        /// </summary>
-        /// <param name="Message"></param>
-        /// <param name="Args"></param>
-        public virtual void ReportWarning(string Message, params object[] Args)
+		/// <summary>
+		/// Report a warning.
+		/// This method is deprecated. Please use Gauntlet.Reporter.Instance.Warning instead of this method.
+		/// </summary>
+		/// <param name="Message"></param>
+		/// <param name="Args"></param>
+		public virtual void ReportWarning(string Message, params object[] Args)
 		{
-			Phases.Current?.AddEvent(Logging.LogLevel.Warning, Message, Args);
-			AddTestEvent(new UnrealTestEvent(EventSeverity.Warning, Message, Args));
-			if (!LogWarningsAndErrorsAfterSummary) { Log.Warning(KnownLogEvents.Gauntlet_TestEvent, Message, Args); }
+			Reporter.Instance.Warning(Message, Args);
 		}
 
 		/// <summary>
-		/// Report information
+		/// Report information.
+		/// This method is deprecated. Please use Gauntlet.Reporter.Instance.Info instead of this method.
 		/// </summary>
 		/// <param name="Message"></param>
 		/// <param name="Args"></param>
 		public virtual void ReportInfo(string Message, params object[] Args)
 		{
-			Phases.Current?.AddEvent(Logging.LogLevel.Information, Message, Args);
-			AddTestEvent(new UnrealTestEvent(EventSeverity.Info, Message, Args));
-			Log.Info(KnownLogEvents.Gauntlet_TestEvent, Message, Args);
+			Reporter.Instance.Info(Message, Args);
 		}
 
 		// Begin UnrealTestNode properties and members
@@ -487,6 +472,7 @@ namespace Gauntlet
 			TestVersion = new Version("1.0.0");
 			ArtifactPath = string.Empty;
 			Phases = new();
+			Reporter.Instance.UpdateReferences(Phases, SetUnrealTestResult, this);
 			PopulateCommandlineInfo();
 			// We format warnings ourselves so don't show these
 			LogWarningsAndErrorsAfterSummary = false;
@@ -848,6 +834,8 @@ namespace Gauntlet
 
 		public override bool IsReadyToStart()
 		{
+			Reporter.Instance.UpdateReferences(Phases, SetUnrealTestResult, this);
+
 			if (UnrealApp == null)
 			{
 				PrepareUnrealApp();
@@ -965,6 +953,8 @@ namespace Gauntlet
 			LastActiveHeartbeatTime = DateTime.MinValue;
 			Phases.Clear();
 
+			Reporter.Instance.UpdateReferences(Phases, SetUnrealTestResult, this);
+
 			TConfigClass Config = GetCachedConfiguration();
 
 			ReserveArtifactPath();
@@ -1154,6 +1144,8 @@ namespace Gauntlet
 		/// <returns></returns>
 		public override void CleanupTest()
 		{
+			Reporter.Instance.UpdateReferences(Phases, SetUnrealTestResult, this);
+
 			if (TestInstance != null)
 			{
 				TestInstance.Dispose();
@@ -1231,6 +1223,8 @@ namespace Gauntlet
 		/// <returns></returns>
 		public override void TickTest()
 		{
+			Reporter.Instance.UpdateReferences(Phases, SetUnrealTestResult, this);
+
 			TickTest(TestInstance);
 
 			// Check status and health after updating logs
