Add support to Horde Job summary message

#rnx
#Synced-CL -1

diff --git a/Engine/Source/Programs/AutomationTool/Gauntlet/Framework/Horde/Gauntlet.Horde.cs b/Engine/Source/Programs/AutomationTool/Gauntlet/Framework/Horde/Gauntlet.Horde.cs
--- a/Engine/Source/Programs/AutomationTool/Gauntlet/Framework/Horde/Gauntlet.Horde.cs
+++ b/Engine/Source/Programs/AutomationTool/Gauntlet/Framework/Horde/Gauntlet.Horde.cs
@@ -10,6 +10,8 @@ using System.Text;
 using System.Text.Json;
 using System.Text.Json.Serialization;
 
+#nullable enable
+
 namespace Gauntlet
 {
 	public static class Horde
@@ -23,7 +25,7 @@ namespace Gauntlet
 			}
 		}
 
-		static public string JobId
+		static public string? JobId
 		{
 			get
 			{
@@ -31,7 +33,7 @@ namespace Gauntlet
 			}
 		}
 
-		static public string StepId
+		static public string? StepId
 		{
 			get
 			{
@@ -39,6 +41,14 @@ namespace Gauntlet
 			}
 		}
 
+		static public string? StepName
+		{
+			get
+			{
+				return Environment.GetEnvironmentVariable("UE_HORDE_STEPNAME");
+			}
+		}
+
 		public enum JobStepMessageLevels
 		{
 			None,
@@ -106,11 +116,9 @@ namespace Gauntlet
 
 		}
 
-#nullable enable
 		private static string? curStepMessage { get; set; } = null;
 		private static string? curStepSeverity { get; set; } = null;
 		private static string? curStepColor { get; set; } = null;
-#nullable disable
 
 		/// <summary>
 		/// Update the step message (short message written in-line in step view)
@@ -148,7 +156,7 @@ namespace Gauntlet
 			UpdateStepMessage(InMessage);
 			if (bGenerateSummary)
 			{
-				GenerateSummary();
+				GenerateStepSummary();
 			}
 		}
 
@@ -166,18 +174,24 @@ namespace Gauntlet
 		public class JobStepReport
 		{
 			public string scope { get; } = "Step";
-			public string name { get; set; }
+			public string? name { get; set; }
 			public string placement { get; } = "Summary";
-			public string fileName { get; set; }
-#nullable enable
+			public string? fileName { get; set; }
 			public string? message { get; set; }
 			public string? severity { get; set; }
 			public string? color { get; set; }
+		}
 
-#nullable disable
-
+		public class JobSummaryReport
+		{
+			public string scope { get; } = "Job";
+			public string? name { get; set; }
+			public string placement { get; } = "Summary";
+			public string? fileName { get; set; }
 		}
+
 		static bool bHasGeneratedStepReport = false;
+		static bool bHasGeneratedJobReport = false;
 
 		public delegate string SummaryLineDelegate();
 		/// <summary>
@@ -185,14 +199,20 @@ namespace Gauntlet
 		/// Returned string is appended as a line to the summary report.
 		/// </summary>
 		/// <returns></returns>
-		public static SummaryLineDelegate OnStepSummary;
+		public static SummaryLineDelegate? OnStepSummary;
 
 		/// <summary>
 		/// Delegate called when the Horde step summary is finalized.
 		/// Returned string is appended as a line to the summary report.
 		/// </summary>
 		/// <returns></returns>
-		public static SummaryLineDelegate OnStepSummaryFinalization;
+		public static SummaryLineDelegate? OnStepSummaryFinalization;
+
+		/// <summary>
+		/// Delegate called when the Horde job summary is generated.
+		/// Return stringis append as a line to the job summary report.
+		/// </summary>
+		public static SummaryLineDelegate? OnJobSummary;
 
 		public static void AddJobBuildMetaData(IBuildSource buildSource)
 		{
@@ -205,7 +225,7 @@ namespace Gauntlet
 
 				if (buildSource != null)
 				{
-					CommandUtils.AddHordeJobMetadata(EpicGames.Horde.Jobs.JobId.Parse(JobId), null, new Dictionary<string, List<string>> { { StepId, new List<string> { $"BuildName=${buildSource.BuildName}?display=true" } } }).Wait();
+					CommandUtils.AddHordeJobMetadata(EpicGames.Horde.Jobs.JobId.Parse(JobId!), null, new Dictionary<string, List<string>> { { StepId!, new List<string> { $"BuildName=${buildSource.BuildName}?display=true" } } }).Wait();
 				}
 
 			}
@@ -222,7 +242,7 @@ namespace Gauntlet
 			return $"* **Netshare Artifacts**: {System.Text.RegularExpressions.Regex.Replace(ArtifactPath, @"([|\\*\+])", @"\$1")}";
 		}
 
-		public static void GenerateSummary()
+		public static void GenerateStepSummary()
 		{
 			try
 			{
@@ -285,7 +305,53 @@ namespace Gauntlet
 			}
 			catch (Exception Ex)
 			{
-				Log.Info("Exception while generating Horde summary\n{0}\n", Ex);
+				Log.Info("Exception while generating Horde step summary\n{0}\n", Ex);
+			}
+		}
+
+		public static void GenerateJobSummary()
+		{
+			try
+			{
+				if (!IsHordeJob)
+				{
+					return;
+				}
+
+				if (bHasGeneratedJobReport)
+				{
+					Log.Info("Trying to generate a job summary report when we have already generated one. ");
+					return;
+				}
+
+				StringBuilder Markdown = new StringBuilder();
+
+				if (OnJobSummary != null)
+				{
+					foreach (SummaryLineDelegate callback in OnJobSummary.GetInvocationList())
+					{
+						Markdown.AppendLine(callback());
+					}
+				}
+
+				if (Markdown.Length > 0)
+				{
+					string LogFolder = CommandUtils.CmdEnv.LogFolder;
+					string MarkdownFilename = "GauntletJobSummary.md";
+
+					File.WriteAllText(Path.Combine(LogFolder, MarkdownFilename), Markdown.ToString());
+					JsonSerializerOptions options = new()
+					{
+						DefaultIgnoreCondition = JsonIgnoreCondition.WhenWritingNull
+					};
+					JobSummaryReport report = new JobSummaryReport() { name = $"{StepName} Summary:", fileName = MarkdownFilename };
+					File.WriteAllText(Path.Combine(LogFolder, "GauntletJobSummary.report.json"), JsonSerializer.Serialize(report, options));
+				}
+				bHasGeneratedJobReport = true;
+			}
+			catch (Exception Ex)
+			{
+				Log.Info("Exception while generating Horde job summary\n{0}\n", Ex);
 			}
 		}
 
diff --git a/Engine/Source/Programs/AutomationTool/Gauntlet/Unreal/Base/Gauntlet.UnrealTestNode.cs b/Engine/Source/Programs/AutomationTool/Gauntlet/Unreal/Base/Gauntlet.UnrealTestNode.cs
--- a/Engine/Source/Programs/AutomationTool/Gauntlet/Unreal/Base/Gauntlet.UnrealTestNode.cs
+++ b/Engine/Source/Programs/AutomationTool/Gauntlet/Unreal/Base/Gauntlet.UnrealTestNode.cs
@@ -1006,7 +1006,8 @@ namespace Gauntlet
 		/// <param name="Report">The curret ITestReport.</param>
 		/// <param name="ReportName">The name you want to show up on the Horde Child Report.</param>
 		/// <param name="URLLink">The URL Link you want to populate on the Horde Child Report.</param>
-		public void AddLinkToReport(ITestReport Report, string ReportName, string URLLink)
+		/// <param name="TextLink">The Text you want to show as URL link. Default is "Open link".</param>
+		public void AddLinkToReport(ITestReport Report, string ReportName, string URLLink, string TextLink = null)
 		{
 			if (Report == null)
 			{
@@ -1021,7 +1022,13 @@ namespace Gauntlet
 				}
 			}
 
-			Horde.OnStepSummaryFinalization += () => $"* **{ReportName}**: [Open link]({URLLink})";
+			if (string.IsNullOrEmpty(TextLink))
+			{
+				TextLink = "Open link";
+			}
+
+			Horde.OnStepSummaryFinalization += () => $"* **{ReportName}**: [{TextLink}]({URLLink})";
+			Horde.OnJobSummary += () => $"* **{ReportName}**: [{TextLink}]({URLLink})";
 		}
 
 		public virtual void PopulateCommandlineInfo()
diff --git a/Engine/Source/Programs/AutomationTool/Gauntlet/Unreal/RunUnreal.cs b/Engine/Source/Programs/AutomationTool/Gauntlet/Unreal/RunUnreal.cs
--- a/Engine/Source/Programs/AutomationTool/Gauntlet/Unreal/RunUnreal.cs
+++ b/Engine/Source/Programs/AutomationTool/Gauntlet/Unreal/RunUnreal.cs
@@ -345,7 +345,8 @@ namespace Gauntlet
 			TeardownAccounts();
 
 			// Generate Horde summary for CIS test
-			Horde.GenerateSummary();
+			Horde.GenerateStepSummary();
+			Horde.GenerateJobSummary();
 
 			return AllTestsPassed ? ExitCode.Success : ExitCode.Error_TestFailure;
 		}
