Add basic remote Windows device support to Gauntlet:
-devices=win64:mysecondpc

(Note that most tests that depend on parsing the output will fail as there is current no way to read this with the remote deployment api)

#rnx

#Synced-CL 50434826

diff --git a/Engine/Source/Programs/AutomationTool/Gauntlet/Framework/Devices/Gauntlet.DevicePool.cs b/Engine/Source/Programs/AutomationTool/Gauntlet/Framework/Devices/Gauntlet.DevicePool.cs
--- a/Engine/Source/Programs/AutomationTool/Gauntlet/Framework/Devices/Gauntlet.DevicePool.cs
+++ b/Engine/Source/Programs/AutomationTool/Gauntlet/Framework/Devices/Gauntlet.DevicePool.cs
@@ -644,6 +644,7 @@ namespace Gauntlet
 					{
 						string[] DevicesList = InputReference.Split(',');
 
+						int UserDeviceIndex = 0;
 						foreach (string DeviceRef in DevicesList)
 						{
 
@@ -691,7 +692,8 @@ namespace Gauntlet
 							Def.Name = DeviceAddress;
 							Def.Platform = DevicePlatform;
 							Def.Model = Model;
-							UnprovisionedDevices.Add(Def);
+							UnprovisionedDevices.Insert(UserDeviceIndex, Def); // add user-specified devices at the start of the list so they are given priority
+							UserDeviceIndex++;
 						}
 					}
 				}
diff --git a/Engine/Source/Programs/AutomationTool/Gauntlet/Gauntlet.Automation.csproj b/Engine/Source/Programs/AutomationTool/Gauntlet/Gauntlet.Automation.csproj
--- a/Engine/Source/Programs/AutomationTool/Gauntlet/Gauntlet.Automation.csproj
+++ b/Engine/Source/Programs/AutomationTool/Gauntlet/Gauntlet.Automation.csproj
@@ -91,6 +91,7 @@
     <ProjectReference Include="..\AutomationUtils\AutomationUtils.Automation.csproj" PrivateAssets="All"><Private>false</Private></ProjectReference>
     <ProjectReference Include="..\Scripts\AutomationScripts.Automation.csproj" PrivateAssets="All"><Private>false</Private></ProjectReference>
     <ProjectReference Include="..\OneSkyLocalization\OneSkyLocalization.Automation.csproj" PrivateAssets="All"><Private>false</Private></ProjectReference>
+    <ProjectReference Include="..\Win\Win.Automation.csproj" PrivateAssets="All"><Private>false</Private></ProjectReference>
     <ProjectReference Include="..\XLocLocalization\XLocLocalization.Automation.csproj" PrivateAssets="All"><Private>false</Private></ProjectReference>
     <ProjectReference Include="..\Localization\Localization.Automation.csproj" PrivateAssets="All"><Private>false</Private></ProjectReference>
     <ProjectReference Include="..\Android\Android.Automation.csproj" PrivateAssets="All"><Private>false</Private></ProjectReference>
diff --git a/Engine/Source/Programs/AutomationTool/Gauntlet/Platform/Windows/Gauntlet.TargetDeviceRemoteWin.cs b/Engine/Source/Programs/AutomationTool/Gauntlet/Platform/Windows/Gauntlet.TargetDeviceRemoteWin.cs
new file mode 100644
--- /dev/null
+++ b/Engine/Source/Programs/AutomationTool/Gauntlet/Platform/Windows/Gauntlet.TargetDeviceRemoteWin.cs
@@ -0,0 +1,240 @@
+// Copyright Epic Games, Inc. All Rights Reserved.
+
+using AutomationTool;
+using System;
+using System.Collections.Generic;
+using System.IO;
+using UnrealBuildTool;
+
+namespace Gauntlet
+{
+	/*
+	 * 
+	 * Represents a remote Windows device
+	 * 
+	 */
+	public class TargetDeviceRemoteWin : ITargetDevice
+	{
+		public TargetDeviceRemoteWin(string InName, string InCacheDir, UnrealTargetPlatform InPlatform)
+		{
+			Name = InName;
+			Platform = InPlatform;
+			RunOptions = CommandUtils.ERunOptions.NoWaitForExit | CommandUtils.ERunOptions.NoLoggingOfRunCommand;
+			LocalCachePath = InCacheDir ?? Path.Combine(Globals.TempDir, $"Remote{Platform}_{Name}");
+		}
+
+		#region IDisposable Support
+		private bool disposedValue = false; // To detect redundant calls
+
+		~TargetDeviceRemoteWin()
+		{
+			Dispose(false);
+		}
+
+		// This code added to correctly implement the disposable pattern.
+		public void Dispose()
+		{
+			// Do not change this code. Put cleanup code in Dispose(bool disposing) above.
+			Dispose(true);
+		}
+
+		protected virtual void Dispose(bool disposing)
+		{
+			if (!disposedValue)
+			{
+				if (disposing)
+				{
+					// TODO: dispose managed state (managed objects).
+				}
+
+				if (IsConnected)
+				{
+					Disconnect();
+				}
+
+				disposedValue = true;
+			}
+		}
+		#endregion
+
+
+		#region ITargetDevice
+		public string Name { get; protected set; }
+		public UnrealTargetPlatform? Platform { get; protected set; }
+		public CommandUtils.ERunOptions RunOptions { get; set; }
+		public bool IsAvailable => true;
+		public bool IsConnected => true;
+		public bool IsOn => true;
+		public string LocalCachePath { get; protected set; }
+		public void CleanArtifacts(UnrealAppConfig AppConfig) { }
+		public bool Connect() => true;
+		public void CopyAdditionalFiles(IEnumerable<UnrealFileToCopy> FilesToCopy) { }
+		public bool Disconnect(bool bForce = false) => true;
+		public void FullClean() { }
+		public bool PowerOff() => true;
+		public bool PowerOn() => true;
+		public bool Reboot() => true;
+
+
+		public Dictionary<EIntendedBaseCopyDirectory, string> GetPlatformDirectoryMappings() => LocalDirectoryMappings;
+
+		public IAppInstall CreateAppInstall(UnrealAppConfig AppConfig)
+		{
+			switch (AppConfig.Build)
+			{
+				case StagedBuild:
+					return CreateStagedInstall(AppConfig, AppConfig.Build as StagedBuild);
+
+				default:
+					throw new AutomationException("{0} is not supported for remote Windows devices", AppConfig.Build.GetType().Name);
+			}
+		}
+
+		public IAppInstall InstallApplication(UnrealAppConfig AppConfig)
+		{
+			InstallBuild(AppConfig);
+			return CreateAppInstall(AppConfig);
+		}
+
+		public void InstallBuild(UnrealAppConfig AppConfig)
+		{
+			switch(AppConfig.Build)
+			{
+				case StagedBuild: 
+					InstallStagedBuild(AppConfig, AppConfig.Build as StagedBuild);
+					break;
+
+				default:
+					throw new AutomationException("{0} is not supported for remote Windows devices", AppConfig.Build.GetType().Name);
+			}
+		}
+
+		public IAppInstance Run(IAppInstall App)
+		{
+			RemoteWinAppInstall WinApp = App as RemoteWinAppInstall;
+			if (WinApp == null)
+			{
+				throw new AutomationException("AppInstance is of incorrect type!");
+			}
+
+			Log.Info("Launching {0} on {1}", WinApp.Name, ToString());
+
+			if (!RemoteWinSupport.Run(Name, WinApp.ExecutablePath, WinApp.CommandLine))
+			{
+				throw new AutomationException("Failed to launch {0}", WinApp.ExecutablePath);
+			}
+
+			return new RemoteWinAppInstance(this, WinApp.CommandLine);
+		}
+		#endregion
+
+		// @fixme: no way to access remote paths yet, so this has to remain empty
+		protected Dictionary<EIntendedBaseCopyDirectory, string> LocalDirectoryMappings { get; set; } = [];
+
+		public override string ToString() => Name;
+
+		protected IAppInstall CreateStagedInstall(UnrealAppConfig AppConfig, StagedBuild Build)
+		{
+			string RemoteExecutablePath = Build.ExecutablePath;
+			if (Path.IsPathRooted(RemoteExecutablePath))
+			{
+				RemoteExecutablePath = Path.GetRelativePath(Build.BuildPath, RemoteExecutablePath);
+			}
+
+			string RemoteRoot = GetRemoteRoot(AppConfig);
+			RemoteExecutablePath = Path.Combine(RemoteRoot, RemoteExecutablePath);
+
+			RemoteWinAppInstall AppInstall = new RemoteWinAppInstall(AppConfig.Name, this)
+			{
+				ExecutablePath = RemoteExecutablePath,
+				CommandLine = AppConfig.CommandLine,
+			};
+
+			return AppInstall;
+		}
+
+		protected IAppInstall InstallStagedBuild(UnrealAppConfig AppConfig, StagedBuild Build)
+		{
+			string RemoteRoot = GetRemoteRoot(AppConfig);
+			if (!RemoteWinSupport.Deploy(Name, Build.BuildPath, RemoteRoot))
+			{
+				throw new AutomationException("Could not deploy {0} to {1}", Build.BuildPath, Name);
+			}
+
+			return CreateStagedInstall(AppConfig, Build);
+		}
+
+		protected string GetRemoteRoot(UnrealAppConfig AppConfig)
+		{
+			return $"Gauntlet_{AppConfig.ProjectName}_{Platform}";
+		}
+	}
+
+	/*
+	 * 
+	 * A running application on a remote Windows device
+	 * 
+	 */
+	public class RemoteWinAppInstance : IAppInstance
+	{
+		private TargetDeviceRemoteWin TargetDevice { get; set; }
+
+		public RemoteWinAppInstance(TargetDeviceRemoteWin InTargetDevice, string InCommandLine)
+		{
+			CommandLine = InCommandLine;
+			TargetDevice = InTargetDevice;
+
+			HasExited = true; // @todo: assume it has exited already because we can't track it yet
+		}
+
+		public void Kill(bool GenerateDumpOnKill = false)
+		{
+			if (!HasExited)
+			{
+				RemoteWinSupport.Kill(TargetDevice.Name);
+				WasKilled = true;
+				HasExited = true;
+			}
+		}
+		public ITargetDevice Device => TargetDevice;
+		public bool HasExited { get; private set; }
+		public bool WasKilled { get; private set; }
+		public string CommandLine { get; protected set; }
+
+
+		// @todo: these are stubs. we can't track process lifetime or read output yet
+		public string StdOut => "";
+		public int ExitCode => 0;
+		public string ArtifactPath => TargetDevice.LocalCachePath;
+		public ILogStreamReader GetLogBufferReader() => new DynamicStringReader(() => "");
+		public ILogStreamReader GetLogReader() => new DynamicStringReader(() => "");
+		public bool WriteOutputToFile(string FilePath) => false;
+		public int WaitForExit() => ExitCode;
+	};
+
+	/*
+	 * 
+	 * An installed application on a remote Windows device
+	 * 
+	 */
+	public class RemoteWinAppInstall : IAppInstall
+	{
+		public RemoteWinAppInstall(string InName, TargetDeviceRemoteWin InTargetDevice)
+		{
+			Name = InName;
+			TargetDevice = InTargetDevice;
+		}
+
+		public string Name { get; protected set; }
+		public ITargetDevice Device => TargetDevice;
+		public string ExecutablePath { get; set; }
+		public string CommandLine { get; set; }
+
+		public TargetDeviceRemoteWin TargetDevice { get; protected set; }
+
+		public IAppInstance Run()
+		{
+			return TargetDevice.Run(this);
+		}
+	};
+}
\ No newline at end of file
diff --git a/Engine/Source/Programs/AutomationTool/Gauntlet/Platform/Windows/Gauntlet.TargetDeviceWindows.cs b/Engine/Source/Programs/AutomationTool/Gauntlet/Platform/Windows/Gauntlet.TargetDeviceWindows.cs
--- a/Engine/Source/Programs/AutomationTool/Gauntlet/Platform/Windows/Gauntlet.TargetDeviceWindows.cs
+++ b/Engine/Source/Programs/AutomationTool/Gauntlet/Platform/Windows/Gauntlet.TargetDeviceWindows.cs
@@ -481,6 +481,11 @@ namespace Gauntlet
 
 		public ITargetDevice CreateDevice(string InRef, string InCachePath, string InParam = null)
 		{
+			if (RemoteWinSupport.IsRemoteDevice(InRef, UnrealTargetPlatform.Win64))
+			{
+				return new TargetDeviceRemoteWin(InRef, InCachePath, UnrealTargetPlatform.Win64);
+			}
+
 			return new TargetDeviceWindows(InRef, InCachePath);
 		}
 	}
diff --git a/Engine/Source/Programs/AutomationTool/Win/RemoteWinSupport.cs b/Engine/Source/Programs/AutomationTool/Win/RemoteWinSupport.cs
--- a/Engine/Source/Programs/AutomationTool/Win/RemoteWinSupport.cs
+++ b/Engine/Source/Programs/AutomationTool/Win/RemoteWinSupport.cs
@@ -32,25 +32,67 @@ public static class RemoteWinSupport
 	public static readonly string DeviceType = "RemoteWin";
 	private static string WdRemotePath => Environment.ExpandEnvironmentVariables("%LocalAppData%\\Microsoft\\WinGet\\Links\\wdRemote.exe");
 
-	public static List<DeviceInfo> GetDevices(UnrealTargetPlatform Platform)
+	public static bool IsAvailable()
 	{
-		List<DeviceInfo> Devices = [];
+		return File.Exists(WdRemotePath);
+	}
+
+	public static bool IsRemoteDevice( string DeviceName, UnrealTargetPlatform Platform )
+	{
+		return GetAllRemoteDeviceNames(Platform).Contains(DeviceName);
+	}
+	
+	public static List<string> GetAllRemoteDeviceNames(UnrealTargetPlatform Platform)
+	{
+		List<string> Result = [];
 
 		// read device names from the engine ini. must already be paired
 		// best place is %LOCALAPPDATA%\Unreal Engine\Engine\Config\UserEngine.ini
 		ConfigHierarchy EngineConfig = ConfigCache.ReadHierarchy(ConfigHierarchyType.Engine, null, Platform);
 		if (EngineConfig.GetArray("RemoteWin", "DeviceNames", out List<string> DeviceNames))
 		{
-			foreach (string DeviceName in DeviceNames)
-			{
-				string SoftwareVersion = Environment.OSVersion.Version.ToString(); // NB. this is the local version, not the remote
-				Devices.Add(new DeviceInfo(Platform, DeviceName, DeviceName, SoftwareVersion, DeviceType, true, true));
-			}
+			Result.AddRange(DeviceNames);
+		}
+
+		return Result;
+	}
+
+	public static List<DeviceInfo> GetDevices(UnrealTargetPlatform Platform)
+	{
+		List<DeviceInfo> Devices = [];
+
+		foreach (string DeviceName in GetAllRemoteDeviceNames(Platform))
+		{
+			string SoftwareVersion = Environment.OSVersion.Version.ToString(); // NB. this is the local version, not the remote
+			Devices.Add(new DeviceInfo(Platform, DeviceName, DeviceName, SoftwareVersion, DeviceType, true, true));
 		}
 
 		return Devices;
 	}
 
+	public static bool Deploy( string DeviceName, string SourcePath, string DestinationPath )
+	{
+		return RunWdRemote($"/action:deploy /source:{StringUtils.QuoteArgument(SourcePath)} /destination:{StringUtils.QuoteArgument(DestinationPath)}", DeviceName);
+	}
+
+	public static bool Run(string DeviceName, string RemoteExecutablePath, string ClientCmdLine)
+	{
+		// escape all quotes
+		// @fixme: this is only necessary because we're using wdremote.exe who's /args needs a quoted string. 
+		//         this will break if ClientCmdLine already contains escaped quotes!
+		ClientCmdLine = ClientCmdLine.Replace("\"", "\\\"");
+
+		RemoteExecutablePath = StringUtils.QuoteArgument(RemoteExecutablePath);
+		RemoteExecutablePath = RemoteExecutablePath.Replace('/', '\\');
+
+		return RunWdRemote($"/action:launch /path:{RemoteExecutablePath} /args:\"{ClientCmdLine}\"", DeviceName);
+	}
+
+	public static bool Kill(string DeviceName)
+	{
+		return RunWdRemote("/action:terminate", DeviceName);
+	}
+
 	public static bool Deploy(UnrealTargetPlatform Platform, ProjectParams Params, DeploymentContext SC)
 	{
 		if (Params.Devices.Count == 0)
@@ -66,7 +108,7 @@ public static class RemoteWinSupport
 				string SourcePath = SC.StageDirectory.FullName;
 				string DestinationPath = GetDestinationRoot(Platform, Params);
 
-				if (!RunWdRemote($"/action:deploy /source:{StringUtils.QuoteArgument(SourcePath)} /destination:{StringUtils.QuoteArgument(DestinationPath)}", DeviceName))
+				if (!Deploy(DeviceName, SourcePath, DestinationPath))
 				{
 					throw new AutomationException("Remote deployment to {0} failed", Params.Devices[0]);
 				}
@@ -93,7 +135,7 @@ public static class RemoteWinSupport
 				string DeviceName = GetTargetDeviceName(Device);
 				string ExecutablePath = GetRemoteExecutablePath(Platform, ClientApp, Params);
 
-				if (!RunWdRemote($"/action:launch /path:{StringUtils.QuoteArgument(ExecutablePath)} /args:\"{ClientCmdLine}\"", DeviceName))
+				if (!Run(DeviceName, ExecutablePath, ClientCmdLine))
 				{
 					throw new AutomationException("Remote launch on {0} failed", Device);
 				}
@@ -175,11 +217,6 @@ public static class RemoteWinSupport
 		RemoveParam("-project=", bOnlyIfUsingAbsolutePath:true);
 		RemoveParam("-abslog=");
 
-		// escape all quotes
-		// @fixme: this is only necessary because we're using wdremote.exe who's /args needs a quoted string. 
-		//         this will break if ClientCmdLine already contains escaped quotes!
-		ClientCmdLine = ClientCmdLine.Replace("\"", "\\\"");
-
 		return ClientCmdLine;
 	}
 
@@ -260,7 +297,7 @@ public static class RemoteWinSupport
 
 		public void StopProcess(bool KillDescendants = true)
 		{
-			RemoteWinSupport.RunWdRemote("/action:terminate", DeviceName);
+			Kill(DeviceName);
 		}
 
 		public bool HasExited => true; // @todo: cannot monitor the remote process yet. assume it has terminated so launching doesn't get stuck waiting indefinitely
