Remove 2GB limitation when generating psym files by never storing the output as a string

#Synced-CL 50199288

diff --git a/Engine/Source/Programs/AutomationTool/AutomationUtils/ProcessUtils.cs b/Engine/Source/Programs/AutomationTool/AutomationUtils/ProcessUtils.cs
--- a/Engine/Source/Programs/AutomationTool/AutomationUtils/ProcessUtils.cs
+++ b/Engine/Source/Programs/AutomationTool/AutomationUtils/ProcessUtils.cs
@@ -1,5 +1,6 @@
 // Copyright Epic Games, Inc. All Rights Reserved.
 using System;
+using System.Buffers;
 using System.Collections;
 using System.Collections.Generic;
 using System.Linq;
@@ -19,6 +20,132 @@ using static AutomationTool.CommandUtils;
 
 namespace AutomationTool
 {
+	/// <summary>
+	/// Reads from a StringBuilder line-by-line using pooled buffers to minimize allocations.
+	/// This avoids the 2GB string limitation when processing large outputs.
+	/// Uses GetChunks() for O(n) iteration instead of O(n²) from indexer access.
+	/// </summary>
+	public class StringBuilderReader : IDisposable
+	{
+		// Note: ChunkEnumerator is a struct - do NOT mark as readonly or MoveNext() won't work!
+		private StringBuilder.ChunkEnumerator ChunkEnumerator;
+		private ReadOnlyMemory<char> CurrentChunk;
+		private int ChunkPosition;
+		private bool HasMoreChunks;
+		private char[] Buffer;
+
+		public StringBuilderReader(StringBuilder builder)
+		{
+			if (builder == null)
+			{
+				throw new ArgumentNullException(nameof(builder));
+			}
+			ChunkEnumerator = builder.GetChunks();
+			HasMoreChunks = ChunkEnumerator.MoveNext();
+			if (HasMoreChunks)
+			{
+				CurrentChunk = ChunkEnumerator.Current;
+			}
+			ChunkPosition = 0;
+			Buffer = ArrayPool<char>.Shared.Rent(4096);
+		}
+
+		/// <summary>
+		/// Reads the next line as a ReadOnlySpan. The span is only valid until the next call to ReadLine.
+		/// Returns false when there are no more lines to read.
+		/// </summary>
+		public bool ReadLine(out ReadOnlySpan<char> Line)
+		{
+			if (!HasMoreChunks || (ChunkPosition >= CurrentChunk.Length && !MoveToNextChunk()))
+			{
+				Line = default;
+				return false;
+			}
+
+			int BufferPos = 0;
+
+			// Read characters until we hit a newline or end of content
+			while (HasMoreChunks)
+			{
+				ReadOnlySpan<char> ChunkSpan = CurrentChunk.Span;
+
+				while (ChunkPosition < ChunkSpan.Length)
+				{
+					char C = ChunkSpan[ChunkPosition];
+
+					if (C == '\r' || C == '\n')
+					{
+						// Skip the newline character(s)
+						ChunkPosition++;
+						// Handle \r\n sequence
+						if (C == '\r')
+						{
+							// Check if next char is \n (may be in next chunk)
+							if (ChunkPosition < ChunkSpan.Length)
+							{
+								if (ChunkSpan[ChunkPosition] == '\n')
+								{
+									ChunkPosition++;
+								}
+							}
+							else if (MoveToNextChunk())
+							{
+								if (CurrentChunk.Span[0] == '\n')
+								{
+									ChunkPosition = 1;
+								}
+							}
+						}
+						goto Done;
+					}
+
+					// Ensure buffer is large enough
+					if (BufferPos >= Buffer.Length)
+					{
+						char[] NewBuffer = ArrayPool<char>.Shared.Rent(Buffer.Length * 2);
+						Buffer.AsSpan(0, BufferPos).CopyTo(NewBuffer);
+						ArrayPool<char>.Shared.Return(Buffer, clearArray: true);
+						Buffer = NewBuffer;
+					}
+
+					Buffer[BufferPos++] = C;
+					ChunkPosition++;
+				}
+
+				// Move to next chunk
+				if (!MoveToNextChunk())
+				{
+					break;
+				}
+			}
+
+		Done:
+			Line = Buffer.AsSpan(0, BufferPos);
+			return true;
+		}
+
+		private bool MoveToNextChunk()
+		{
+			if (ChunkEnumerator.MoveNext())
+			{
+				CurrentChunk = ChunkEnumerator.Current;
+				ChunkPosition = 0;
+				return true;
+			}
+			HasMoreChunks = false;
+			return false;
+		}
+
+		public void Dispose()
+		{
+			if (Buffer != null)
+			{
+				ArrayPool<char>.Shared.Return(Buffer, clearArray: true);
+				Buffer = null;
+			}
+		}
+	}
+
 	public enum CtrlTypes
 	{
 		CTRL_C_EVENT = 0,
@@ -249,6 +376,7 @@ namespace AutomationTool
 		string ToString();
 		void WaitForExit();
 		FileReference WriteOutputToFile(string FileName);
+		StringBuilderReader GetOutputReader() => throw new NotImplementedException("GetOutputReader is not supported on this platform");
 	}
 	
 	/// <summary>
@@ -616,6 +744,15 @@ namespace AutomationTool
 			return new FileReference(FileName);
 		}
 
+		public StringBuilderReader GetOutputReader()
+		{
+			if (ProcessOutput == null)
+			{
+				return null;
+			}
+			return new StringBuilderReader(ProcessOutput);
+		}
+
 		/// <summary>
 		/// Finds child processes of the current process.
 		/// </summary>
diff --git a/Engine/Source/Programs/AutomationTool/BuildGraph/Tasks/SymGenTask.cs b/Engine/Source/Programs/AutomationTool/BuildGraph/Tasks/SymGenTask.cs
--- a/Engine/Source/Programs/AutomationTool/BuildGraph/Tasks/SymGenTask.cs
+++ b/Engine/Source/Programs/AutomationTool/BuildGraph/Tasks/SymGenTask.cs
@@ -209,7 +209,6 @@ namespace BuildGraph.Tasks
 				IProcessResult result = CommandUtils.Run(symbolDumperExeForFile, dumpSymsArgs, null, CommandUtils.ERunOptions.AppMustExist, null, FilterSpew, null, workingDirectory);
 				if (result.ExitCode == 0)
 				{
-					StringBuilder processedSymbols = null;
 					if (bUseRadSym)
 					{
 						// rad dumper outputs to a file, we thunk to a custom exe to do symbol munging for speed.
@@ -221,32 +220,21 @@ namespace BuildGraph.Tasks
 					}
 					else
 					{
+						// Use streaming approach to avoid 2GB string limitation
+						// Read directly from the process output without materializing it as a string
 						try
 						{
-							// Process the symbols
-							using (StringReader reader = new StringReader(result.Output))
+							using (StringBuilderReader reader = result.GetOutputReader())
 							{
-								ProcessSymbols(reader, out processedSymbols);
-							}
-						}
-						catch (OutOfMemoryException)
-						{
-							// If we catch an OOM, it is too large to turn into a string.
-							// Write to a file and then load it into a string.
-							string tempFileName = Path.Combine(Path.GetTempPath(), Path.GetTempFileName());
-							FileReference symbolFile = (result as ProcessResult).WriteOutputToFile(tempFileName);
-
-							try
-							{
-								using (StreamReader reader = new StreamReader(symbolFile.FullName))
+								if (reader != null)
 								{
-									ProcessSymbols(reader, out processedSymbols);
+									using (StreamWriter writer = new StreamWriter(symbolFileName))
+									{
+										ProcessSymbols(reader, writer);
+										symbolFiles.Add(new FileReference(symbolFileName));
+									}
 								}
 							}
-							finally
-							{
-								FileReference.Delete(symbolFile);
-							}
 						}
 						catch (Exception ex)
 						{
@@ -254,15 +242,6 @@ namespace BuildGraph.Tasks
 							throw new AutomationException($"Symbol generation failed: Error Generating Symbols for {symbolFileName}, Error: {ExceptionUtils.FormatException(ex)}");
 						}
 					}
-
-					if (processedSymbols != null && processedSymbols.Length > 0)
-					{
-						using (StreamWriter writer = new StreamWriter(symbolFileName))
-						{
-							writer.Write(processedSymbols);
-						}
-						symbolFiles.Add(new FileReference(symbolFileName));
-					}
 				}
 				else
 				{
@@ -292,26 +271,23 @@ namespace BuildGraph.Tasks
 		/// Processes the raw symbol dump
 		/// </summary>
 
-		static bool ProcessSymbols(TextReader reader, out StringBuilder processedSymbols)
+		static bool ProcessSymbols(StringBuilderReader reader, TextWriter writer)
 		{
 			char[] fieldSeparator = { ' ' };
 			string rootDirectory = CommandUtils.ConvertSeparators(PathSeparator.Slash, Unreal.RootDirectory.FullName).TrimEnd('/');
 
-			processedSymbols = new StringBuilder();
-
-			string line;
 			bool bSawModule = false;
-			while ((line = reader.ReadLine()) != null)
+			while (reader.ReadLine(out ReadOnlySpan<char> lineSpan))
 			{
-				if (line.Contains(" = ", StringComparison.InvariantCulture))
+				if (lineSpan.Contains(" = ", StringComparison.InvariantCulture))
 				{
-					Logger.LogInformation("{Text}", line);
+					Logger.LogInformation("{Text}", lineSpan.ToString());
 					continue;
 				}
 				// Ignore any output from symbol dump before MODULE, these may included erroneous warnings, etc
 				if (!bSawModule)
 				{
-					if (!line.StartsWith("MODULE", StringComparison.InvariantCulture))
+					if (!lineSpan.StartsWith("MODULE", StringComparison.InvariantCulture))
 					{
 						continue;
 					}
@@ -319,30 +295,35 @@ namespace BuildGraph.Tasks
 					bSawModule = true;
 				}
 
-				string newLine = line;
-
 				// Process source reference FILE blocks
-				if (line.StartsWith("FILE", StringComparison.InvariantCulture))
+				if (lineSpan.StartsWith("FILE", StringComparison.InvariantCulture))
 				{
+					// Convert to string only for FILE lines that need processing
+					string line = lineSpan.ToString();
 					string[] fields = line.Split(fieldSeparator, 3);
 
-					string fileName = CommandUtils.ConvertSeparators(PathSeparator.Slash, fields[2]);
-
-					// If the file exists locally, and is within the root, convert path
-					if (File.Exists(fileName) && fileName.StartsWith(rootDirectory, StringComparison.OrdinalIgnoreCase))
+					// Check we got all the expected fields
+					if (fields.Length == 3)
 					{
-						// Restore proper filename case on Windows (the symbol dump filenames are all lowercase)
-						if (BuildHostPlatform.Current.Platform == UnrealTargetPlatform.Win64)
+						string fileName = CommandUtils.ConvertSeparators(PathSeparator.Slash, fields[2]);
+
+						// If the file exists locally, and is within the root, convert path
+						if (File.Exists(fileName) && fileName.StartsWith(rootDirectory, StringComparison.OrdinalIgnoreCase))
 						{
-							fileName = FileUtils.FindCorrectCase(new FileInfo(fileName)).FullName;
-						}
+							// Restore proper filename case on Windows (the symbol dump filenames are all lowercase)
+							if (BuildHostPlatform.Current.Platform == UnrealTargetPlatform.Win64)
+							{
+								fileName = FileUtils.FindCorrectCase(new FileInfo(fileName)).FullName;
+							}
 
-						// Shave off the root directory
-						newLine = String.Format("FILE {0} {1}", fields[1], fileName.Substring(rootDirectory.Length + 1).Replace('\\', '/'));
+							// Shave off the root directory
+							writer.WriteLine(String.Format("FILE {0} {1}", fields[1], fileName.Substring(rootDirectory.Length + 1).Replace('\\', '/')));
+							continue;
+						}
 					}
 				}
 
-				processedSymbols.AppendLine(newLine);
+				writer.WriteLine(lineSpan);
 			}
 
 			return true;
