Increasing timeout in between process launch and starting appium

#Synced-CL 50214590

diff --git a/Engine/Source/Programs/AutomationTool/Gauntlet/Platform/IOS/Gauntlet.TargetDeviceIOS.cs b/Engine/Source/Programs/AutomationTool/Gauntlet/Platform/IOS/Gauntlet.TargetDeviceIOS.cs
--- a/Engine/Source/Programs/AutomationTool/Gauntlet/Platform/IOS/Gauntlet.TargetDeviceIOS.cs
+++ b/Engine/Source/Programs/AutomationTool/Gauntlet/Platform/IOS/Gauntlet.TargetDeviceIOS.cs
@@ -844,6 +844,7 @@ namespace Gauntlet
 			{
 				lock (Mutex)
 				{
+					Thread.Sleep(5000);
 					AppiumContainer.CancelAppleSignOnAlert(IOSInstall.PackageName);
 				}
 			}
