Adjust markdown details for Gauntlet Job summary

#rnx
#Synced-CL -1

diff --git a/Engine/Source/Programs/AutomationTool/Gauntlet/Framework/Horde/Gauntlet.Horde.cs b/Engine/Source/Programs/AutomationTool/Gauntlet/Framework/Horde/Gauntlet.Horde.cs
--- a/Engine/Source/Programs/AutomationTool/Gauntlet/Framework/Horde/Gauntlet.Horde.cs
+++ b/Engine/Source/Programs/AutomationTool/Gauntlet/Framework/Horde/Gauntlet.Horde.cs
@@ -336,6 +336,9 @@ namespace Gauntlet
 
 				if (Markdown.Length > 0)
 				{
+					// a blank line to exit markdown paragraph
+					Markdown.AppendLine();
+
 					string LogFolder = CommandUtils.CmdEnv.LogFolder;
 					string MarkdownFilename = "GauntletJobSummary.md";
 
@@ -344,7 +347,7 @@ namespace Gauntlet
 					{
 						DefaultIgnoreCondition = JsonIgnoreCondition.WhenWritingNull
 					};
-					JobSummaryReport report = new JobSummaryReport() { name = $"{StepName} Summary:", fileName = MarkdownFilename };
+					JobSummaryReport report = new JobSummaryReport() { name = $"{StepName} Summary", fileName = MarkdownFilename };
 					File.WriteAllText(Path.Combine(LogFolder, "GauntletJobSummary.report.json"), JsonSerializer.Serialize(report, options));
 				}
 				bHasGeneratedJobReport = true;
