Move GDK UnrealBuildTool & AutomationTool source into public engine

#rnx

#Synced-CL 50642529

diff --git a/Engine/Source/Programs/AutomationTool/Gauntlet/Platform/Microsoft/GDK/Gauntlet.GDKBuildSource.cs b/Engine/Source/Programs/AutomationTool/Gauntlet/Platform/Microsoft/GDK/Gauntlet.GDKBuildSource.cs
new file mode 100644
--- /dev/null
+++ b/Engine/Source/Programs/AutomationTool/Gauntlet/Platform/Microsoft/GDK/Gauntlet.GDKBuildSource.cs
@@ -0,0 +1,250 @@
+// Copyright Epic Games, Inc. All Rights Reserved.
+
+using System;
+using System.Collections.Generic;
+using System.IO;
+using System.Linq;
+using System.Text.RegularExpressions;
+using System.Xml;
+using System.Xml.Linq;
+using System.Xml.XPath;
+using EpicGames.Core;
+using Gauntlet;
+using UnrealBuildTool;
+using UnrealBuildTool.GDK;
+
+namespace Gauntlet
+{
+	/// <summary>
+	/// Represents a packaged GDK build that can be installed
+	/// </summary>
+	public class GDKPackagedBuild : IBuild
+	{
+		/// <summary>
+		/// Preference Order for this build over other types (loose file)
+		/// </summary>
+		public int PreferenceOrder { get { return 0; } }
+
+		/// <summary>
+		/// Metadata for this package
+		/// </summary>
+		public UnrealBuildTool.GDK.GDKGameConfigInfo MetaData { get; protected set; }
+
+		/// <summary>
+		/// Flags that represent the properties of this build
+		/// </summary>
+		public BuildFlags Flags { get; protected set; }
+
+		/// <summary>
+		/// Flavor of this build.
+		/// </summary>
+		public string Flavor => "";
+
+		public bool SupportsAdditionalFileCopy => false;
+
+		/// <summary>
+		/// Configuration of the executable in this package
+		/// </summary>
+		public UnrealTargetConfiguration Configuration { get; protected set; }
+
+		/// <summary>
+		/// Platform that this build runs on
+		/// </summary>
+		public UnrealTargetPlatform Platform { get; private set; }
+
+		/// <summary>
+		/// Path of this package
+		/// </summary>
+		public string PackageFilename;
+
+		/// <summary>
+		/// /Constructor that takes all relevant info
+		/// </summary>
+		/// <param name="InPlatform"></param>
+		/// <param name="InConfig"></param>
+		/// <param name="InFlags"></param>
+		/// <param name="InPackageFilename"></param>
+		/// <param name="InInfo"></param>
+		public GDKPackagedBuild(UnrealTargetPlatform InPlatform, UnrealTargetConfiguration InConfig, BuildFlags InFlags, string InPackageFilename, GDKGameConfigInfo InInfo)
+		{
+			Platform = InPlatform;
+			Configuration = InConfig;
+			Flags = InFlags | BuildFlags.Packaged;
+			MetaData = InInfo;
+			PackageFilename = InPackageFilename;
+		}
+
+		/// <summary>
+		/// Packaged builds can only support Clients
+		/// </summary>
+		/// <param name="Role"></param>
+		/// <returns></returns>
+		public bool CanSupportRole(UnrealTargetRole Role)
+		{
+			return Role.IsClient();
+		}
+
+		/// <summary>
+		/// Attempts to find any packaged builds at the provided path.
+		/// </summary>
+		/// <param name="InPlatform"></param>
+		/// <param name="InProjectName"></param>
+		/// <param name="InPath"></param>
+		/// <param name="InPackageFileRexEx"></param>
+		/// <returns></returns>
+		public static IEnumerable<GDKPackagedBuild> CreateFromPath(UnrealTargetPlatform InPlatform, string InProjectName, string InPath, string InPackageFileRexEx)
+		{
+			return CreateFromPath<GDKPackagedBuild>(InPlatform, InProjectName, InPath, InPackageFileRexEx);
+		}
+
+		internal static IEnumerable<T> CreateFromPath<T>(UnrealTargetPlatform InPlatform, string InProjectName, string InPath, string InPackageFileRexEx) where T : GDKPackagedBuild
+		{
+			List<T> DiscoveredBuilds = new List<T>();
+
+			IEnumerable<string> Packages = DirectoryUtils.FindFiles(InPath, new Regex(InPackageFileRexEx));
+
+			foreach (string PackageName in Packages)
+			{
+				// First get common meta data
+				string PackageFilename = Path.GetFileNameWithoutExtension(PackageName);
+				GDKGameConfigInfo PackageMetaData = new GDKGameConfigInfo(PackageFilename);
+
+				// now we look at the appxmanifest file to see what binaries were in it
+				string AppxManifestFilename = Path.Combine(Path.GetDirectoryName(PackageName), $"{PackageFilename}_appxmanifest.xml");
+
+				if (!File.Exists(AppxManifestFilename))
+				{
+					Log.Warning(KnownLogEvents.Gauntlet_BuildDropEvent, "No appxmanifest {File} file for GDK package {Package}. Ignoring.", Path.GetFileName(AppxManifestFilename), PackageName);
+					continue;
+				}
+
+				// Get the package build configuration
+				XDocument AppxManifestXML = null;
+				try
+				{
+					AppxManifestXML = XDocument.Load(AppxManifestFilename);
+				}
+				catch (Exception Ex)
+				{
+					Log.Warning(KnownLogEvents.Gauntlet_BuildDropEvent, "Error loading appxmanifest xml in {File} : {Exception}, ignoring", AppxManifestFilename, Ex.Message);
+					continue;
+				}
+
+				XNamespace np = AppxManifestXML.Root.GetDefaultNamespace();
+				string[] PackageBinaries = (from Application in AppxManifestXML.Descendants(np+"Application")
+											where ((string)Application.Attribute("Executable")).Contains(".exe") && ((string)Application.Attribute("Executable")).Contains(InProjectName)
+											select Application.Attribute("Executable").Value).ToArray();
+				if (PackageBinaries.Length == 0)
+				{
+					Log.Warning(KnownLogEvents.Gauntlet_BuildDropEvent, "Unable to parse application binary from appxmanifest {File}, ignoring", AppxManifestFilename);
+					return DiscoveredBuilds;
+				}
+				foreach (string PackageBinary in PackageBinaries)
+				{
+					string PackageBinaryFileName = Path.GetFileName(PackageBinary);
+
+					UnrealTargetConfiguration UnrealConfig = UnrealHelpers.GetConfigurationFromExecutableName(InProjectName, PackageBinaryFileName);
+
+					// all builds are packaged and can be launched with any command line
+					BuildFlags Flags = BuildFlags.Packaged | BuildFlags.CanReplaceCommandLine;
+
+					// flag as content only if we're using UnrealGame.exe
+					if (PackageBinaryFileName.Equals("UnrealGame.exe", StringComparison.OrdinalIgnoreCase))
+					{
+						Flags |= BuildFlags.ContentOnlyProject;
+					}
+
+					T DiscoveredBuild = Activator.CreateInstance(typeof(T), InPlatform, UnrealConfig, Flags, PackageName, PackageMetaData) as T;
+					DiscoveredBuilds.Add(DiscoveredBuild);
+				}
+			}
+
+			return DiscoveredBuilds;
+		}
+	}
+
+	/// <summary>
+	/// A build source for GDK that finds and returns packaged builds
+	/// </summary>
+	public abstract class GDKPackagedBuildSource : IFolderBuildSource
+	{
+		protected UnrealTargetPlatform Platform { get; set; }
+
+		protected string DirectorySearchPattern { get; set; }
+
+		public string BuildName { get; private set; }
+
+		public string PackageFileExtension { get; protected set; }
+
+		public string PackageSuffix { get; protected set; }
+
+		/// <summary>
+		/// Regular expression for standard GDK package file name - [PackageName][Suffix].[Extension]
+		/// </summary>
+		public string PackageFileRexEx { get { return GDKGameConfigInfo.PackageIdentityRegEx + PackageSuffix + "." + PackageFileExtension + "$"; } }
+
+
+		/// <summary>
+		/// Platforms we can support
+		/// </summary>
+		/// <param name="InPlatform"></param>
+		public bool CanSupportPlatform(UnrealTargetPlatform InPlatform)
+		{
+			return InPlatform == Platform;
+		}
+
+		protected GDKPackagedBuildSource( UnrealTargetPlatform InPlatform, string InPackageFileExtension, string InDirectorySearchPattern, string InBuildName, string InPackageSuffix = "")
+		{
+			Platform = InPlatform;
+			PackageFileExtension = InPackageFileExtension;
+			DirectorySearchPattern = InDirectorySearchPattern;
+			BuildName = InBuildName;
+			PackageSuffix = InPackageSuffix;
+		}
+
+		/// <summary>
+		/// Find all builds at the specified path, recursing as specified
+		/// </summary>
+		/// <param name="InProjectName"></param>
+		/// <param name="InPath"></param>
+		/// <param name="MaxRecursion"></param>
+		/// <returns></returns>
+		public virtual List<IBuild> GetBuildsAtPath(string InProjectName, string InPath, int MaxRecursion = 3)
+		{
+			return GetBuildsAtPath<GDKPackagedBuild>(InProjectName, InPath, MaxRecursion);
+
+		}
+
+		internal List<IBuild> GetBuildsAtPath<T>(string InProjectName, string InPath, int MaxRecursion = 3) where T : GDKPackagedBuild
+		{
+			List<IBuild> Builds = new List<IBuild>();
+
+			// Check if the platform name is in the path. If so we'll search from here
+			bool IsInPlatformFolder = DirectoryUtils.PathContainsComponent(InPath, Platform.ToString(), true);
+
+			IEnumerable<DirectoryInfo> RootFolders = null;
+
+			if (IsInPlatformFolder)
+			{
+				RootFolders = new DirectoryInfo[] { new DirectoryInfo(InPath) };
+			}
+			else
+			{
+				// find all the top level subdirectories that match this platform
+				RootFolders = DirectoryUtils.FindMatchingDirectories(InPath, DirectorySearchPattern);
+			}
+
+			// Find all matching files
+			IEnumerable<FileInfo> PackageFiles = RootFolders.SelectMany(D => DirectoryUtils.FindMatchingFiles(D.FullName, PackageFileRexEx, MaxRecursion));
+
+			// Now try to create something out of their locations
+			foreach (FileInfo Package in PackageFiles)
+			{
+				var FoundBuilds = GDKPackagedBuild.CreateFromPath<T>(Platform, InProjectName, Package.DirectoryName, PackageFileRexEx);
+				Builds.AddRange(FoundBuilds);
+			}			
+
+			return Builds;
+		}
+	}
+}
diff --git a/Engine/Source/Programs/AutomationTool/Microsoft/GDK/GDKAutomationUtils.cs b/Engine/Source/Programs/AutomationTool/Microsoft/GDK/GDKAutomationUtils.cs
new file mode 100644
--- /dev/null
+++ b/Engine/Source/Programs/AutomationTool/Microsoft/GDK/GDKAutomationUtils.cs
@@ -0,0 +1,161 @@
+// Copyright Epic Games, Inc. All Rights Reserved.
+using System;
+using System.Collections.Generic;
+using System.Linq;
+using System.Text;
+using System.IO;
+using AutomationTool;
+using UnrealBuildTool;
+using UnrealBuildTool.GDK;
+using EpicGames.Core;
+using System.Xml.Linq;
+using Microsoft.Extensions.Logging;
+
+namespace AutomationUtils.GDK
+{
+	public class GDKAutomationUtils
+	{
+		public static void GenerateGameConfigForStaging(ProjectParams Params, DeploymentContext SC, string IniSection_TargetSettings, ILogger Logger)
+		{
+			DirectoryReference ManifestTargetDirectory = DirectoryReference.Combine(SC.ProjectRoot, "Saved", SC.StageTargetPlatform.PlatformType.ToString(), "Manifest");
+			InternalUtils.SafeDeleteDirectory(ManifestTargetDirectory.FullName);
+
+			Dictionary<UnrealTargetConfiguration, string> ExecutablePairs = GetGameConfigExecutables(Params, SC, GameConfigPurpose.ForStaging, IniSection_TargetSettings);
+			GDKGameConfigGenerator.Create(SC.StageTargetPlatform.PlatformType, Logger, ManifestTargetDirectory, SC.ShortProjectName, Params.RawProjectPath, ExecutablePairs, DLCFile: Params.DLCFile, CustomConfig: SC.CustomConfig, NoBootstrapExe: Params.NoBootstrapExe);
+
+			if (Params.ApplyIoStoreOnDemand)
+			{
+				ApplyIoStoreOnDemandSettings(Params, ManifestTargetDirectory, Logger);
+			}
+
+			SC.StageFiles(StagedFileType.SystemNonUFS, ManifestTargetDirectory, "*", StageFilesSearch.AllDirectories, StagedDirectoryReference.Root);
+		}
+
+		public enum GameConfigPurpose
+		{
+			ForStaging,
+			ForPackaging,
+			ForBootstrapping,
+		}
+		public static Dictionary<UnrealTargetConfiguration, string> GetGameConfigExecutables(ProjectParams Params, DeploymentContext SC, GameConfigPurpose Purpose, string IniSection_TargetSettings, UnrealTargetConfiguration? OnlyConfig = null)
+		{
+			bool bForStaging = (Purpose == GameConfigPurpose.ForStaging);
+			bool bOnlyBootstrap = (Purpose == GameConfigPurpose.ForBootstrapping);
+
+			// build filtered stage target list
+			List<StageTarget> FilteredStageTargets = SC.StageTargets.Where(X => OnlyConfig == null || (X.Receipt.Configuration == OnlyConfig!.Value)).ToList();
+			if (OnlyConfig != null && !FilteredStageTargets.Any())
+			{
+				throw new AutomationException($"Could not find receipt for config {OnlyConfig!.Value} in currently stage targets - has the game been built & staged for this configuration?");
+			}
+
+			Dictionary<UnrealTargetConfiguration, string> ExecutablePairs = FilteredStageTargets  // note: only takes the first staged file per configuration at the moment
+				.GroupBy( X => X.Receipt.Configuration )
+				.ToDictionary(X => X.Key, X => SC.GetStagedFileLocation(X.First().Receipt.Launch!).ToString());
+
+			if (!Params.NoBootstrapExe && SC.StageTargetPlatform.PlatformType.IsInGroup(UnrealPlatformGroup.Windows))
+			{
+				// bootstrap naming logic from WinPlatform.Automation.cs -> GetFilesToDeployOrStage
+				Dictionary<UnrealTargetConfiguration, string> BootstrapExecutablePairs;
+				if (SC.StageTargetConfigurations.Count > 1) // note: must check unfiltered stage target here to make sure the correct bootstrap executable is referenced
+				{
+					BootstrapExecutablePairs = ExecutablePairs.ToDictionary(X => X.Key, X => Path.GetFileName(X.Value));
+				}
+				else if (Params.IsCodeBasedProject)
+				{
+					BootstrapExecutablePairs = ExecutablePairs.ToDictionary(X => X.Key, X => SC.StageTargets.First().Receipt.TargetName + ".exe");
+				}
+				else
+				{
+					BootstrapExecutablePairs = ExecutablePairs.ToDictionary(X => X.Key, X => SC.ShortProjectName + ".exe");
+				}
+
+				// only use the bootstrap executables if they exist (to allow for development iteration scenarios). Bootstrap must exist for distribution
+				bool bAllExist = BootstrapExecutablePairs.All(X =>
+					bForStaging ? SC.FilesToStage.NonUFSSystemFiles.ContainsKey(new StagedFileReference(X.Value)) : FileReference.Exists(FileReference.Combine(SC.StageDirectory, X.Value)));
+
+				if (Params.Distribution || bAllExist)
+				{
+					ExecutablePairs = BootstrapExecutablePairs;
+				}
+				else if (bOnlyBootstrap)
+				{
+					// bootstrap executables not found
+					ExecutablePairs.Clear();
+				}
+			}
+			else if (Params.NoBootstrapExe)
+			{
+				ConfigHierarchy EngineIni = ConfigCache.ReadHierarchy(ConfigHierarchyType.Engine, DirectoryReference.FromFile(SC.RawProjectPath), SC.StageTargetPlatform.PlatformType, SC.CustomConfig);
+
+				foreach (UnrealTargetConfiguration Config in ExecutablePairs.Keys)
+				{
+					string ConfigEntry = $"OverrideLaunchExe_{Config}_x64";
+					if (EngineIni.GetString(IniSection_TargetSettings, ConfigEntry, out string OverrideExectuablePath) && !string.IsNullOrEmpty(OverrideExectuablePath))
+					{
+						ExecutablePairs[Config] = OverrideExectuablePath;
+					}
+				}
+			}
+			else if (bOnlyBootstrap)
+			{
+				// this configuration will never return bootstrap executables
+				ExecutablePairs.Clear();
+			}
+
+			return ExecutablePairs;
+		}
+
+		/// <summary>
+		/// Makes sure that the project config file is setup to allow enough persistent local storage
+		/// for IoStoreOnDemand caching. If the project is not set up for this or has not requested enough
+		/// space then the config file will be modified to satisfy our requirements.
+		/// </summary>
+		/// <param name="Params">Parameters for the current project</param>
+		/// <param name="ManifestTargetDirectory">Path to the directory containing the manifests</param>
+		/// <param name="Logger">Logger interface so we can track progress</param>
+		private static void ApplyIoStoreOnDemandSettings(ProjectParams Params, DirectoryReference ManifestTargetDirectory, ILogger Logger)
+		{
+			try
+			{
+				bool UpdateDoc = false;
+
+				string ConfigFilePath = ManifestTargetDirectory + "\\MicrosoftGame.config";
+				XDocument Document = XDocument.Load(ConfigFilePath);
+
+				XElement Node;
+				if ((Node = Document.Root.Element("PersistentLocalStorage")?.Element("SizeMB")) != null)
+				{
+					int SizeMB = 0;
+					int.TryParse(Node.Value, out SizeMB);
+					if (SizeMB < 256)
+					{
+						Node.Value = "256";
+						UpdateDoc = true;
+					}
+				}
+				else if ((Node = Document.Root.Element("PersistentLocalStorage")) != null)
+				{
+					Node.Add(new XElement("SizeMB", "256"));
+					UpdateDoc = true;
+				}
+				else
+				{
+					Document.Root.Add(new XElement("PersistentLocalStorage", new XElement("SizeMB", "256")));
+					UpdateDoc = true;
+				}
+
+				if (UpdateDoc)
+				{
+					Document.Save(ConfigFilePath);
+
+					Logger.LogInformation("IoStoreOnDemand applied changes to persistent local storage settings to the project");
+				}
+			}
+			catch (Exception Ex)
+			{
+				Logger.LogWarning("Error encountered while trying to apply persistent download data settings, IoStoreOnDemand may not work! [{Msg}]", Ex.Message);
+			}
+		}
+	}
+}
diff --git a/Engine/Source/Programs/AutomationTool/Microsoft/GDK/GDKPackageGenerator.cs b/Engine/Source/Programs/AutomationTool/Microsoft/GDK/GDKPackageGenerator.cs
new file mode 100644
--- /dev/null
+++ b/Engine/Source/Programs/AutomationTool/Microsoft/GDK/GDKPackageGenerator.cs
@@ -0,0 +1,1539 @@
+// Copyright Epic Games, Inc. All Rights Reserved.
+using System;
+using System.Collections.Generic;
+using System.Linq;
+using System.Text;
+using System.IO;
+using System.Text.Json;
+using System.Text.RegularExpressions;
+using System.Xml.Linq;
+using System.ComponentModel;
+using Microsoft.Extensions.Logging;
+using AutomationTool;
+using AutomationUtils.Automation;
+using UnrealBuildTool;
+using UnrealBuildTool.GDK;
+using EpicGames.Core;
+
+using static AutomationTool.CommandUtils;
+using System.Threading.Tasks;
+using ILogger = Microsoft.Extensions.Logging.ILogger;
+
+namespace AutomationUtils.GDK
+{
+	public abstract class GDKPackageGenerator
+	{
+		protected readonly string IniSection_TargetSettings;
+		protected readonly UnrealTargetPlatform PlatformType;
+		protected readonly ConfigHierarchy EngineIni;
+
+		/// Whether to generate separate packages for each configuration
+		public bool GenerateSinglePackages { get; protected set; } = true;
+
+		/// Whether to generate a single package containing all configurations combined
+		public bool GenerateCombinedPackage { get; protected set; } = false;
+
+		/// Whether to create metadata during staging that can be used to make the package. This ensures that only the explicitly staged files will be included in the package. The alternative is that everything in the staging directory is scraped and packaged
+		public bool PrecachePackageLayout { get; protected set; } = true;
+
+		/// Whether to parallelize calls to makepkg when generating multiple configuration packages. This requires PrecachePackageLayout
+		public bool ParallelizePackageGeneration { get; protected set; } = true;
+
+		/// When true, specifying -NoDebugInfo will cause /skipsymbolbundling to be passed to the package generator for non-submission packages. This will mean the symbols zip file for uploading to Partner Center will not be created. Potentially useful for iteration.
+		public bool NoDebugInfoSkipsSymbolBundling { get; protected set; } = false;
+
+
+	protected static int GDKEdition => GDKExports.GetGDKVersionNumber() ?? 0;
+		protected static string GDKRootDir => Utils.CleanDirectorySeparators(Path.Combine(GDKExports.GetGSDKRoot()));
+		protected static string GDKBinariesDir => Path.Combine( GDKRootDir, "bin");
+
+
+		public GDKPackageGenerator(FileReference InProjectPath, UnrealTargetPlatform InPlatformType, string InIniSection_TargetSettings, string CustomConfig)
+		{
+			PlatformType = InPlatformType;
+			IniSection_TargetSettings = InIniSection_TargetSettings;
+			EngineIni = ConfigCache.ReadHierarchy(ConfigHierarchyType.Engine, DirectoryReference.FromFile(InProjectPath), PlatformType, CustomConfig);
+
+			if (EngineIni.GetBool(IniSection_TargetSettings, "bGenerateSinglePackages", out bool bGenerateSinglePackages))
+			{
+				GenerateSinglePackages = bGenerateSinglePackages;
+			}
+			if (EngineIni.GetBool(IniSection_TargetSettings, "bGenerateCombinedPackage", out bool bGenerateCombinedPackage))
+			{
+				GenerateCombinedPackage = bGenerateCombinedPackage;
+			}
+			if (EngineIni.GetBool(IniSection_TargetSettings, "bPrecachePackageLayout", out bool bPrecachePackageLayout))
+			{
+				PrecachePackageLayout = bPrecachePackageLayout;
+			}
+			if (EngineIni.GetBool(IniSection_TargetSettings, "bParallelizePackageGeneration", out bool bParallelizePackageGeneration))
+			{
+				ParallelizePackageGeneration = bParallelizePackageGeneration;
+			}
+			if (EngineIni.GetBool(IniSection_TargetSettings, "bNoDebugInfoSkipsSymbolBundling", out bool bNoDebugInfoSkipsSymbolBundling))
+			{
+				NoDebugInfoSkipsSymbolBundling = bNoDebugInfoSkipsSymbolBundling;
+			}
+
+			// makepkg validation only honors layout.xml from March 2024 GDK. (prior to this it would ensure the executables were relative to the MicrosoftGame.config) 
+			if (GDKEdition < 240300)
+			{
+				PrecachePackageLayout = false;
+			}
+
+			if (!PrecachePackageLayout)
+			{
+				ParallelizePackageGeneration = false;
+			}
+
+			VerifyWindowsToolsEnvironment();
+		}
+
+
+		public abstract string GetPlatformPackageFileExtension();
+
+		protected abstract string GetPlatformSpecificPackageCommandlineOptions(ProjectParams Params, DeploymentContext SC);
+
+		protected virtual void GetPlatformSpecificPackageSymbolPaths(HashSet<string> SymbolPaths, ProjectParams Params, DeploymentContext SC)
+		{
+		}
+
+		protected string GetPackagingEncryptionOptions(ProjectParams Params)
+		{
+			string EncryptionOptions = "/lt";
+			if (Params.Distribution)
+			{
+				if (String.IsNullOrWhiteSpace(Params.PackageEncryptionKeyFile))
+				{
+					EncryptionOptions = "/l";
+				}
+				else
+				{
+					EncryptionOptions = string.Format("/lk {0}", Params.PackageEncryptionKeyFile);
+				}
+			}
+
+			return EncryptionOptions;
+		}
+
+		protected string GetPackageEncryptionOptions(string PackagePath)
+		{
+			string PackageUtilPath = Path.Combine(GDKBinariesDir, "PackageUtil.exe");
+			string PackageUtilCommandLine = String.Format("license {0}", PackagePath);
+
+			string PackageEncryptionOptions = "";
+			IProcessResult RunResult = CommandUtils.Run(PackageUtilPath, PackageUtilCommandLine);
+			if (RunResult.ExitCode == 0)
+			{
+				PackageEncryptionOptions = RunResult.Output.Trim();
+			}
+
+			return PackageEncryptionOptions;
+		}
+
+		protected string FindPackageFile(DirectoryReference PackagePath)
+		{
+			string ResultFilePath = null;
+
+			if (Directory.Exists(PackagePath.ToString()))
+			{
+				string PackageFilenameFilter = String.Format("*{0}", GetPlatformPackageFileExtension());
+				DirectoryInfo LatestPatchDI = new DirectoryInfo(PackagePath.ToString());
+				var FileInLatestPatchDirectory = LatestPatchDI.GetFiles(PackageFilenameFilter, SearchOption.AllDirectories).Select(file => file.FullName);
+
+				if (FileInLatestPatchDirectory.Any())
+				{
+					ResultFilePath = FileInLatestPatchDirectory.First();
+					Logger.LogDebug("{Text}", String.Format("Found package file {0}", ResultFilePath));
+				}
+			}
+
+			return ResultFilePath;
+		}
+
+		string GetOptionalConfigGuidOption(string KeyName, string Parameter)
+		{
+			string Value;
+			if (EngineIni.GetString(IniSection_TargetSettings, KeyName, out Value) && !String.IsNullOrEmpty(Value))
+			{
+				Guid DummyGuid;
+				if (Guid.TryParse(Value, out DummyGuid))
+				{
+					return " " + Parameter + " " + Value;
+				}
+				else
+				{
+					throw new AutomationException("{0} must be a valid GUID: {0}. Check the value set in {1}", KeyName, Value, IniSection_TargetSettings);
+				}
+			}
+
+			return "";
+		}
+
+		public void GetPackageFilesToArchive(ProjectParams Params, DeploymentContext SC, string PackageOutputPathFragment = "")
+		{
+			if (GenerateSinglePackages || Params.HasCreateReleaseVersion)
+			{
+				foreach (var Config in SC.StageTargetConfigurations)
+				{
+					DirectoryReference PackagePath = DirectoryReference.Combine(SC.ProjectRoot, "Saved", "Packages", SC.CookPlatform, PackageOutputPathFragment, $"{Config}");
+					if (Params.HasCreateReleaseVersion)
+					{
+						PackagePath = DirectoryReference.FromString(Path.Combine(Params.GetCreateReleaseVersionPath(SC, Params.Client), PackageOutputPathFragment, Config.ToString()));
+					}
+					SC.ArchiveFiles(PackagePath.FullName, NewPath: Path.Combine(PackageOutputPathFragment, Config.ToString()));
+				}
+			}
+
+			if (GenerateCombinedPackage && !Params.HasCreateReleaseVersion)
+			{
+				DirectoryReference PackagePath = DirectoryReference.Combine(SC.ProjectRoot, "Saved", "Packages", SC.CookPlatform, PackageOutputPathFragment, "Combined");
+				SC.ArchiveFiles(PackagePath.FullName, NewPath: Path.Combine(PackageOutputPathFragment, "Combined")); ;
+			}
+		}
+
+
+
+		public void GeneratePackages(ProjectParams Params, DeploymentContext SC, DirectoryReference PackageMetadataBasePath = null, string PackageOutputPathFragment = "")
+		{
+			// prepare the list of configurations to package
+			List<UnrealTargetConfiguration?> ConfigurationsToPackage = new();
+			if (GenerateSinglePackages || Params.HasCreateReleaseVersion)
+			{
+				ConfigurationsToPackage.AddRange(SC.StageTargetConfigurations.Select(X => (UnrealTargetConfiguration?)X));
+			}
+			if (GenerateCombinedPackage && !Params.HasCreateReleaseVersion)
+			{
+				ConfigurationsToPackage.Add(null); // null means "no specific single config - package all configurations"
+			}
+			if (ConfigurationsToPackage.Count == 0)
+			{
+				return;
+			}
+
+			// generate the packages
+			if (ParallelizePackageGeneration && ConfigurationsToPackage.Count > 1)
+			{
+				Logger.LogInformation("Generating {Num} packages in parallel...", ConfigurationsToPackage.Count);
+
+				var LockObject = new object();
+				Parallel.ForEach(ConfigurationsToPackage, (UnrealTargetConfiguration? Config) =>
+				{
+					CaptureLogger CaptureLogger = new();
+					try
+					{
+						GenerateSinglePackage(Params, SC, PackageMetadataBasePath, Config, CaptureLogger, PackageOutputPathFragment);
+					}
+					catch(AutomationException e)
+					{
+						CaptureLogger.LogError("{Ex}", e.Message);
+						throw new AutomationException("Packaging failed");
+					}
+					finally
+					{
+						// lock ensures each package's output is not interleaved
+						lock (LockObject) 
+						{
+							CaptureLogger.RenderTo(Logger);
+						}
+					}
+				});
+			}
+			else
+			{
+				foreach (UnrealTargetConfiguration? Config in ConfigurationsToPackage)
+				{
+					GenerateSinglePackage(Params, SC, PackageMetadataBasePath, Config, null, PackageOutputPathFragment);
+				}
+			}
+		}
+
+
+
+
+
+
+		public void GenerateSinglePackage(ProjectParams Params, DeploymentContext SC, DirectoryReference PackageMetadataBasePath = null, UnrealTargetConfiguration? OnlyConfig = null, ILogger CustomLogger = null, string PackageOutputPathFragment = "")
+		{
+			string PackageConfig = OnlyConfig?.ToString() ?? "Combined";
+
+			// default to global logger unless one has been specified & output banner
+			CustomLogger ??= Logger;
+			CustomLogger.LogInformation("");
+			CustomLogger.LogInformation("*** {PackageConfig} {Platform} Package ***", PackageConfig, PlatformType);
+
+			// set package output directory
+			DirectoryReference PackageOutputDir = DirectoryReference.Combine(SC.ProjectRoot, "Saved", "Packages", SC.CookPlatform, PackageOutputPathFragment, PackageConfig);
+			if (Params.HasCreateReleaseVersion)
+			{
+				PackageOutputDir = DirectoryReference.FromString(Path.Combine(Params.GetCreateReleaseVersionPath(SC, Params.Client), PackageOutputPathFragment, PackageConfig));
+			}
+
+			// fall back to legacy packaging if there is no metadata
+			if (!PrecachePackageLayout)
+			{
+				PackageImmediate(Params, SC, PackageOutputDir, CustomLogger, OnlyConfig);
+				return;
+			}
+
+			// make sure there is valid metadata
+			PackageMetadataBasePath ??= DirectoryReference.Combine(SC.StageDirectory, ".pkgsrc");
+			DirectoryReference PackageMetadataDir = DirectoryReference.Combine(PackageMetadataBasePath, PackageConfig);
+			FileReference LayoutXmlFile = FileReference.Combine(PackageMetadataDir, "layout.xml");
+			if (!FileReference.Exists(LayoutXmlFile))
+			{
+				throw new AutomationException($"{LayoutXmlFile} not found - was the game staged?");
+			}
+
+			// attempt to package
+			PackageInternal(Params, SC, PackageOutputDir, LayoutXmlFile, PackageMetadataDir, CustomLogger, OnlyConfig);
+		}
+
+
+
+
+		protected virtual void PackageInternal(ProjectParams Params, DeploymentContext SC, DirectoryReference PackagePath, FileReference LayoutPath, DirectoryReference SourceFolder, ILogger Logger, UnrealTargetConfiguration? OnlyConfig = null)
+		{
+			// sanity check
+			if (SC.StageDirectory == null)
+			{
+				throw new AutomationException("Cannot package: No staging directory. Please specify -stage, -skipstage or -stagingdirectory");
+			}
+			else if (!DirectoryReference.Exists(SC.StageDirectory))
+			{
+				throw new AutomationException($"Cannot package: Staging directory {SC.StageDirectory} does not exist");
+			}
+
+			// locate submission validator DLL
+			string AutoSdkSubValRoot = Path.GetFullPath(@"..\..\SubmissionValidator_Latest\", GDKRootDir); // this will be %UE_SDKS_ROOT%\HostWin64\GDK\SubmissionValidator_Latest\ folder when using AutoSDK
+
+			string MakePkgValidationPath;
+			string SubmissionValidator = Path.Combine(AutoSdkSubValRoot, "SubmissionValidator.dll");
+			if (File.Exists(SubmissionValidator))
+			{
+				// use shared AutoSDK version of SubmissionValidator.dll			
+				MakePkgValidationPath = Path.TrimEndingDirectorySeparator(AutoSdkSubValRoot);
+			}
+			else
+			{
+				// try default location, alongside makepkg itself
+				SubmissionValidator = Path.Combine(GDKBinariesDir, "SubmissionValidator.dll");
+				MakePkgValidationPath = null;
+			}
+
+			// add explicit warning if submission validator DLL cannot be found. The actual error from makepkg just says it is out of date which isn't helpful.			
+			bool bIncludesShipping = (OnlyConfig != null) ? (OnlyConfig == UnrealTargetConfiguration.Shipping) : SC.StageTargetConfigurations.Contains(UnrealTargetConfiguration.Shipping);
+			if (bIncludesShipping && Params.Distribution && !File.Exists(SubmissionValidator))
+			{
+				Log.TraceWarningOnce($"SubmissionValidator DLL is not present in {GDKBinariesDir} or %UE_SDKS_ROOT%/HostWin64/GDK/SubmissionValidator_Latest/ - It is a separate download from the main GDK and is used when making a submission package for uploading to Partner Center. Please see the GDK documentation for details.");
+			}
+
+			// prepare paths
+			string MakePkgPath = Path.Combine(GDKBinariesDir, "makepkg.exe");
+			if (DirectoryReference.Exists(PackagePath))
+			{
+				CommandUtils.DeleteDirectory(PackagePath);
+			}
+			DirectoryReference.CreateDirectory(PackagePath);
+
+			// prepare the makepkg command line
+			string AdditionalOptions = GetPlatformSpecificPackageCommandlineOptions(Params, SC) + " " + GetPackagingEncryptionOptions(Params);
+			string PackageCommandline = string.Format("pack /f \"{0}\" {1} /d \"{2}\" /pd \"{3}\" /loggable", LayoutPath, AdditionalOptions, SourceFolder, PackagePath);
+
+			// add path to shared submission validator, if necessary
+			if (MakePkgValidationPath != null)
+			{
+				PackageCommandline += $" /validationpath \"{MakePkgValidationPath}\"";
+			}
+
+			// read store Id: we only include the legacy ProductId if the store Id is not specified
+			if (!Params.HasDLCName)
+			{
+				bool bHasStoreId = (EngineIni.GetString(IniSection_TargetSettings, "StoreId", out string StoreId) && !String.IsNullOrEmpty(StoreId));
+				if (!bHasStoreId)
+				{
+					PackageCommandline += GetOptionalConfigGuidOption("ProductId", "/productid");
+				}
+				PackageCommandline += GetOptionalConfigGuidOption("ContentId", "/contentid");
+			}
+
+			// do not generate a symbols zip if they don't want debug info
+			if (NoDebugInfoSkipsSymbolBundling && Params.NoDebugInfo && !Params.Distribution)
+			{
+				PackageCommandline += " /skipsymbolbundling";
+			}
+			else
+			{
+				// append all paths containing PDBs, for SubmissionValidator.dll
+				HashSet<string> SymbolPaths = new HashSet<string>();
+				GetPlatformSpecificPackageSymbolPaths(SymbolPaths, Params, SC);
+				SymbolPaths.Add(Params.GetProjectBinariesPathForPlatform(PlatformType).ToString());
+				foreach (StageTarget Target in SC.StageTargets.Where(X => OnlyConfig == null || OnlyConfig == X.Receipt.Configuration))
+				{
+					foreach (RuntimeDependency RuntimeDependency in Target.Receipt.RuntimeDependencies.Where(x => x.Path.HasExtension("pdb") || x.Path.HasExtension("dll")))
+					{
+						SymbolPaths.Add(RuntimeDependency.Path.Directory.ToString());
+					}
+					foreach (BuildProduct BuildProduct in Target.Receipt.BuildProducts.Where(x => x.Path.HasExtension("pdb") || x.Path.HasExtension("dll")))
+					{
+						SymbolPaths.Add(BuildProduct.Path.Directory.ToString());
+					}
+				}
+				PackageCommandline += " /symbolpaths ";
+				foreach (string SymbolPath in SymbolPaths)
+				{
+					PackageCommandline += "\"" + SymbolPath + "\";";
+				}
+			}
+
+
+			// Look for a release package to use as base package for patch generation
+			string ReleasePackagePath = null;
+			if (Params.HasBasedOnReleaseVersion && !Params.HasDLCName)
+			{
+				DirectoryReference ReleaseVersionBasePath = new DirectoryReference(Params.GetBasedOnReleaseVersionPath(SC, Params.Client));
+				DirectoryReference[] ReleaseVersionPaths = new DirectoryReference[]
+				{
+					DirectoryReference.Combine(ReleaseVersionBasePath, "LatestPatch"), // legacy support. check this first as the package search is recursive.
+					ReleaseVersionBasePath, // for consistency with other platforms
+				};
+
+				foreach (DirectoryReference ReleaseVersionPath in ReleaseVersionPaths)
+				{
+					ReleasePackagePath = FindPackageFile(ReleaseVersionPath);
+					if (!String.IsNullOrEmpty(ReleasePackagePath))
+					{
+						Logger.LogDebug("{Text}", String.Format("Found base package for delta patching {0}", ReleasePackagePath));
+						PackageCommandline += String.Format(" /priorpackage \"{0}\"", ReleasePackagePath);
+						break;
+					}
+				}
+
+				if (String.IsNullOrEmpty(ReleasePackagePath))
+				{
+					Logger.LogWarning("Attempting to create a new release package but no previous release package could be found in {ReleaseVersionBasePath}. This may negatively affect patch download sizes.", ReleaseVersionBasePath);
+				}
+			}
+
+			string AdditionalMakePkgArguments;
+			if (EngineIni.GetString(IniSection_TargetSettings, "AdditionalMakePkgArguments", out AdditionalMakePkgArguments) && !String.IsNullOrEmpty(AdditionalMakePkgArguments))
+			{
+				PackageCommandline += String.Format(" {0}", AdditionalMakePkgArguments);
+			}
+
+			// run makepkg & apply message filtering
+			SpewFilterHelper SpewFilter = new SpewFilterHelper();
+			PrepareMessageFilterForMakePkg(SpewFilter, EngineIni, Params, SC);
+
+			// Get makepkg to write its temp files to a controlled folder by overriding TMP environment variable to 
+			DirectoryReference MakePKGTempPath = DirectoryReference.Combine(SC.ProjectRoot, "Intermediate", PlatformType.ToString(), "makepkgtmp");
+			DirectoryReference.CreateDirectory(MakePKGTempPath);
+			Dictionary<string, string> ExtraEnvVars = new();
+			ExtraEnvVars["TMP"] = MakePKGTempPath.FullName;
+            ExtraEnvVars["TEMP"] = MakePKGTempPath.FullName;
+
+			Logger.LogInformation("Running \"{MakePkgPath}\" {PackageCommandline}", MakePkgPath, PackageCommandline);
+			IProcessResult Result = CommandUtils.Run(MakePkgPath, PackageCommandline, Options: ERunOptions.AppMustExist|ERunOptions.NoLoggingOfRunCommand, Env: ExtraEnvVars);
+			PrettyLogProcessOutput(Logger, Result.Output.EnumerateLines().Select(X => SpewFilter.FilterMessage(X)));
+
+			// check the result from makepkg
+			if (Result.ExitCode != 0)
+			{
+				throw new AutomationException("Packaging failed - \"{0}\" {1} returned {2}", MakePkgPath, PackageCommandline, Result.ExitCode);
+			}
+
+			// estimate size of delta patch if we are generating a patch.  Packages must have compatible encryption keys.
+			if (!String.IsNullOrEmpty(ReleasePackagePath) && !Params.HasDLCName)
+			{
+				string PatchPackagePath = FindPackageFile(PackagePath);
+				if (!String.IsNullOrEmpty(PatchPackagePath))
+				{
+					GDKGameConfigInfo ReleasePackageInfo = new(Path.GetFileNameWithoutExtension(ReleasePackagePath));
+					GDKGameConfigInfo PatchPackageInfo = new(Path.GetFileNameWithoutExtension(PatchPackagePath));
+					if (Version.Parse(ReleasePackageInfo.Version) <= Version.Parse(PatchPackageInfo.Version))
+					{
+						string ReleasePackageEncryptionOptions = GetPackageEncryptionOptions(ReleasePackagePath);
+						string PatchPackageEncryptionOptions = GetPackageEncryptionOptions(PatchPackagePath);
+						if (GetPackagingEncryptionOptions(Params) != "/l" && ReleasePackageEncryptionOptions == PatchPackageEncryptionOptions)
+						{
+							Logger.LogDebug("Found patch package for patch size estimate {releasepath}", ReleasePackagePath);
+
+							string PackageUtilPath = Path.Combine(GDKBinariesDir, "PackageUtil.exe");
+							string PackageUtilCommandLine = String.Format("compare {0} {1}", ReleasePackagePath, PatchPackagePath);
+
+							Logger.LogInformation("Running \"{PackageUtilPath}\" {PackageUtilCommandLine}", PackageUtilPath, PackageUtilCommandLine);
+							Result = CommandUtils.Run(PackageUtilPath, PackageUtilCommandLine, Options: ERunOptions.AppMustExist | ERunOptions.NoLoggingOfRunCommand);
+							PrettyLogProcessOutput(Logger, Result.Output.EnumerateLines());
+							if (Result.ExitCode != 0)
+							{
+								Logger.LogWarning("Patch size estimation failed - {packageutil} {params} returned {errcode}", PackageUtilPath, PackageUtilCommandLine, Result.ExitCode);
+							}
+						}
+						else
+						{
+							Logger.LogDebug("Base and patch package encryption options are not compatible for patch size estimation.  Base = {baseopt}  Patch = {patchopt}", ReleasePackageEncryptionOptions, PatchPackageEncryptionOptions);
+						}
+					}
+					else
+					{
+						Logger.LogWarning("Release package is newer than the patch: skipping patch comparison. Base = {basever}  Patch = {patchver}", ReleasePackageInfo.Version, PatchPackageInfo.Version);
+					}
+				}
+			}
+
+			// write out the output, suppressing needless empty lines. this matches the output when the gdk tool is run directly from the command line, outside of UAT.
+			static void PrettyLogProcessOutput(ILogger Logger, IEnumerable<string> OutputLines)
+			{
+				bool bWriteBlank = false;
+				foreach (string OutputLine in OutputLines)
+				{
+					bool bIsBlank = (OutputLine.Length == 0);
+					if (!bIsBlank || (bIsBlank == bWriteBlank))
+					{
+						// Log HRESULT as warning
+						if (Regex.Match(OutputLine, @"\(0x8.{7}\)").Success)
+						{
+							Logger.LogWarning("{OutputLine}", OutputLine);
+						}
+						else
+						{
+							Logger.LogInformation("{OutputLine}", OutputLine);
+						}
+					}
+					bWriteBlank = bIsBlank && !bWriteBlank;
+				}
+			}
+		}
+
+
+		private List<FileReference> MoveFiles(List<FileReference> Files, DirectoryReference SrcDir, DirectoryReference DstDir)
+		{
+			DirectoryReference.CreateDirectory(DstDir);
+
+			List<FileReference> DstFiles = new List<FileReference>();
+			foreach (FileReference SrcFile in Files)
+			{
+				FileReference DstFile = FileReference.Combine(DstDir, SrcFile.MakeRelativeTo(SrcDir));
+				DirectoryReference.CreateDirectory(DstFile.Directory);
+				FileReference.Move(SrcFile, DstFile);
+
+				DstFiles.Add(DstFile);
+			}
+
+			return DstFiles;
+		}
+
+		private void DeleteFiles(List<FileReference> Files)
+		{
+			foreach (FileReference File in Files)
+			{
+				FileReference.Delete(File);
+			}
+		}
+
+
+
+		protected void PackageImmediate(ProjectParams Params, DeploymentContext SC, DirectoryReference PackagePath, ILogger Logger, UnrealTargetConfiguration? OnlyConfig = null)
+		{
+			// prepare intermediate packaging directory
+			DirectoryReference IntermediatePath = DirectoryReference.Combine(SC.ProjectRoot, "Intermediate", PlatformType.ToString(), "Packaging");
+			InternalUtils.SafeDeleteDirectory(IntermediatePath.FullName);
+			DirectoryReference.CreateDirectory(IntermediatePath);
+
+			// generate a filtered list of the files that should be included in the package
+			IEnumerable<StagedFileReference> FilesToInclude = FindStagedFilesToPackage(Params, SC, OnlyConfig).Where(X=>
+				ShouldPackageFile(X.Value, Params, SC, OnlyConfig))
+				.Select(f => f.Key);
+
+			// exclude any staged manifest resources as we're generating new ones
+			IEnumerable<StagedFileReference> ManifestFiles = GetGDKManifestFilesInDirectory(SC.StageDirectory)
+					.Select(File => StagedFileReference.Combine(File.MakeRelativeTo(SC.StageDirectory)));
+			FilesToInclude = FilesToInclude.Except(ManifestFiles);
+
+			// generate a GDK manifest that just contains the selected configurations
+			Dictionary<UnrealTargetConfiguration, string> ExecutablePairs = GDKAutomationUtils.GetGameConfigExecutables(Params, SC, GDKAutomationUtils.GameConfigPurpose.ForPackaging, IniSection_TargetSettings, OnlyConfig);
+			GDKGameConfigGenerator.Create(PlatformType, Logger, IntermediatePath, SC.ShortProjectName, Params.RawProjectPath, ExecutablePairs, bDeleteIntermediates: true, DLCFile: Params.DLCFile, CustomConfig: SC.CustomConfig, NoBootstrapExe: Params.NoBootstrapExe);
+			IEnumerable<StagedFileReference> GeneratedFilesToInclude =  GetGDKManifestFilesInDirectory(IntermediatePath)
+				.Select( File => StagedFileReference.Combine(File.MakeRelativeTo(IntermediatePath)));
+
+			// layout.xml & package.manifest generation
+			GeneratePackageMetadataInternal(Params, SC, FilesToInclude, GeneratedFilesToInclude, IntermediatePath, Logger);
+
+			// attempt to package
+			FileReference LayoutXmlFile = FileReference.Combine(IntermediatePath, "layout.xml");
+			PackageInternal(Params, SC, PackagePath, LayoutXmlFile, IntermediatePath, Logger, OnlyConfig);
+		}
+
+
+
+
+		private List<FileReference> GetGDKManifestFilesInDirectory(DirectoryReference Directory)
+		{
+			List<FileReference> Result = new List<FileReference>();
+
+			FileReference MSGameConfig = FileReference.Combine(Directory, "MicrosoftGame.config");
+			if (FileReference.Exists(MSGameConfig))
+			{
+				Result.Add(MSGameConfig);
+			}
+
+			Result.AddRange(DirectoryReference.EnumerateFiles(Directory, "*.pri", SearchOption.TopDirectoryOnly));
+
+			DirectoryReference ResourcesDir = DirectoryReference.Combine(Directory, "Resources");
+			if (DirectoryReference.Exists(ResourcesDir))
+			{
+				Result.AddRange(DirectoryReference.EnumerateFiles(ResourcesDir, "*.*", SearchOption.AllDirectories));
+			}
+
+			return Result;
+		}
+
+
+
+		protected void VerifyWindowsToolsEnvironment()
+		{
+			if (System.Environment.OSVersion.Version.Major < 10) //https://docs.microsoft.com/windows/win32/sysinfo/operating-system-version
+			{
+				Logger.LogWarning("Your Windows version {Arg0}.{Arg1}.{Arg2} may be too low to run GDK tools - it should be at least Windows 10 or Windows Server 2016", System.Environment.OSVersion.Version.Major, System.Environment.OSVersion.Version.Minor, System.Environment.OSVersion.Version.Build);
+			}
+		}
+
+
+		#region Packaging error/warning filter
+
+		#region Helper class
+		protected class SpewFilterHelper
+		{
+			private List<string> IgnoreList = new List<string>();
+			private Dictionary<string, string> ReplaceDict = new Dictionary<string, string>();
+			private Dictionary<string, string> RegExReplaceDict = new Dictionary<string, string>();
+
+			public void ReplaceMessages(Dictionary<string, string> Replacements)
+			{
+				foreach (var Replacement in Replacements)
+				{
+					ReplaceDict.Add(Replacement.Key, Replacement.Value);
+				}
+			}
+
+			public void RegExReplaceMessages(Dictionary<string, string> Replacements)
+			{
+				foreach (var Replacement in Replacements)
+				{
+					RegExReplaceDict.Add(Replacement.Key, Replacement.Value);
+				}
+			}
+
+			public void IgnoreMessages(List<string> Ignore)
+			{
+				IgnoreList.AddRange(Ignore);
+			}
+
+			public string FilterMessage(string Message)
+			{
+				foreach (string IgnoreMessage in IgnoreList)
+				{
+					Message = Message.Replace(IgnoreMessage, "", StringComparison.InvariantCultureIgnoreCase);
+				}
+
+				foreach (var ReplaceMessage in ReplaceDict)
+				{
+					if (Message.Contains(ReplaceMessage.Key))
+					{
+						Message = Message.Replace(ReplaceMessage.Key, ReplaceMessage.Value, StringComparison.InvariantCultureIgnoreCase);
+					}
+				}
+
+				foreach (var RegExReplacee in RegExReplaceDict)
+				{
+					Message = Regex.Replace(Message, RegExReplacee.Key, RegExReplacee.Value, RegexOptions.IgnoreCase);
+				}
+
+				return Message;
+			}
+		}
+		#endregion
+
+		private static readonly Dictionary<string, string> SubmissionValidatorErrorsToDowngrade = new Dictionary<string, string>()
+		{
+			{ "Error: A required update to Submission Validator is available", "Warning: A required update to Submission Validator is available (this will be an error in Shipping for Distribution)" },
+		};
+
+		private static readonly Dictionary<string, string> StoreAssociationMessageToAlter = new Dictionary<string, string>()
+		{
+			{ "Please use the MicrosoftGame.config Editor \"Associate with Store\" feature", "Please use the \"Update Partner Center Association\" button in Unreal Editor's Project Settings" },
+		};
+
+		protected virtual void PrepareMessageFilterForMakePkg(SpewFilterHelper Filter, ConfigHierarchy EngineIni, ProjectParams Params, DeploymentContext SC)
+		{
+			Filter.ReplaceMessages(StoreAssociationMessageToAlter);
+
+			bool bSubmissionValidatorExpirationError = false;
+			if (!EngineIni.GetBool(IniSection_TargetSettings, "bSubmissionValidatorExpirationError", out bSubmissionValidatorExpirationError) || !bSubmissionValidatorExpirationError)
+			{
+				if (!Params.Distribution || !Params.ClientConfigsToBuild.Contains(UnrealTargetConfiguration.Shipping)) //only downgrade the error if we're not building a shipping package for distribution
+				{
+					Filter.ReplaceMessages(SubmissionValidatorErrorsToDowngrade);
+				}
+			}
+
+			List<string> AllowedMissingSymbolFiles;
+			if (EngineIni.GetArray(IniSection_TargetSettings, "AllowedMissingSymbolFiles", out AllowedMissingSymbolFiles) && AllowedMissingSymbolFiles != null)
+			{
+				foreach (string BinaryName in AllowedMissingSymbolFiles)
+				{
+					string Key = string.Format("Warning:.*(?<msg>Symbol file not found for.*{0}\\.).*", Regex.Escape(BinaryName));
+					string Value = "${msg} (Ignoring warning due to the file being listed in AllowedMissingSymbolFiles config entry.)";
+					Filter.RegExReplaceMessages(new Dictionary<string, string>() { { Key, Value } });
+				}
+			}
+		}
+
+		#endregion
+
+		#region Package chunk layout generation
+
+		protected virtual bool ShouldPackageFile(FileReference SrcFile, ProjectParams Params, DeploymentContext SC, UnrealTargetConfiguration? OnlyConfig = null)
+		{
+			// Exclude all PDBs
+			if (SrcFile.HasExtension(".pdb"))
+			{
+				return false;
+			}
+
+			// Exclude specific files
+			string[] ExcludedFiles = { "package.manifest", $"Manifest_DebugFiles_{PlatformType}.txt", $"Manifest_NonUFSFiles_{PlatformType}.txt", $"Manifest_UFSFiles_{PlatformType}.txt", $"Manifest_ObsoleteNonUFSFiles.txt", $"Manifest_ObsoleteUFSFiles.txt", $"Manifest_DeltaUFSFiles.txt", $"Manifest_DeltaNonUFSFiles.txt" };
+			if (ExcludedFiles.Any(X => X.Equals(SrcFile.GetFileName(), StringComparison.InvariantCultureIgnoreCase)))
+			{
+				return false;
+			}
+
+			// Exclude launch executables for other configurations.  Otherwise, generating packages from a multi-config build will result in a mismatch
+			// between launch binaries in the package and those exposed in the game config.  This will result in potential failed attempts to launch the additional
+			// binaries in the package by Gauntlet.
+			if (OnlyConfig != null && SrcFile.HasExtension(".exe"))
+			{
+				bool bExcluded = SC.StageTargets.Any(
+					X => X.Receipt.Configuration != OnlyConfig &&
+					X.Receipt.Launch.GetFileName() == Path.GetFileName(SrcFile.FullName)
+				);
+				if (bExcluded)
+				{
+					return false;
+				}
+			}
+
+			return true;
+		}
+
+		private Dictionary<StagedFileReference, FileReference> FindStagedFilesToPackage(ProjectParams Params, DeploymentContext SC, UnrealTargetConfiguration? OnlyConfig = null)
+		{
+			Dictionary<StagedFileReference, FileReference> FilesToInclude = new();
+			foreach (FileReference SrcFile in DirectoryReference.EnumerateFiles(SC.StageDirectory, "*", SearchOption.AllDirectories))
+			{
+				if (!ShouldPackageFile(SrcFile, Params, SC, OnlyConfig))
+				{
+					continue;
+				}
+
+				StagedFileReference DstFile = StagedFileReference.Combine(SrcFile.MakeRelativeTo(SC.StageDirectory));
+				FilesToInclude.Add(DstFile, SrcFile);
+			}
+			return FilesToInclude;
+		}
+
+
+
+
+
+		internal struct IntelligentInstallRecipe
+		{
+			public string Id;
+			public string[] FeatureIds;
+			public string[] Devices;
+			public string[] StoreIds;
+
+			internal static IEnumerable<IntelligentInstallRecipe> GetFromConfig(ProjectParams Params, DeploymentContext SC, string IniSection_TargetSettings)
+			{
+				ConfigHierarchy EngineIni = ConfigCache.ReadHierarchy(ConfigHierarchyType.Engine, DirectoryReference.FromFile(SC.RawProjectPath), SC.StageTargetPlatform.PlatformType, SC.CustomConfig);
+
+				if( EngineIni.GetBool(IniSection_TargetSettings, "bUseFeaturesAndRecipes", out bool bUseFeaturesAndRecipes) &&
+					bUseFeaturesAndRecipes)
+				{
+					// read recipes
+					EngineIni.GetArray(IniSection_TargetSettings, "IntelligentDeliveryRecipes", out List<string> Recipes);
+					if (Recipes != null && Recipes.Count > 0)
+					{
+						foreach (string Recipe in Recipes)
+						{
+							Dictionary<string, string> Properties;
+							ConfigHierarchy.TryParse(Recipe, out Properties);
+
+							IntelligentInstallRecipe Mapping;
+
+							if (!TryGetValue(Properties, "Id", out Mapping.Id) || Mapping.Id == "")
+							{
+								throw new AutomationException("Missing Id in recipe");
+							}
+
+							TryGetValue(Properties, "IncludedFeatures", out Mapping.FeatureIds);
+							TryGetValue(Properties, "Devices", out Mapping.Devices);
+							TryGetValue(Properties, "StoreIds", out Mapping.StoreIds);
+
+							yield return Mapping;
+						}
+					}
+
+				}
+			}
+		}
+
+
+		internal class IntelligentInstallFeature
+		{
+			public string Id;
+			public string[] Tags;
+			public bool bHidden;
+			public string DisplayNameResourceId;
+			public string ImageFileName;
+
+			internal static IEnumerable<IntelligentInstallFeature> GetFromConfig(ProjectParams Params, DeploymentContext SC, string IniSection_TargetSettings)
+			{
+				ConfigHierarchy EngineIni = ConfigCache.ReadHierarchy(ConfigHierarchyType.Engine, DirectoryReference.FromFile(SC.RawProjectPath), SC.StageTargetPlatform.PlatformType, SC.CustomConfig);
+
+				if (EngineIni.GetBool(IniSection_TargetSettings, "bUseFeaturesAndRecipes", out bool bUseFeaturesAndRecipes) &&
+					bUseFeaturesAndRecipes)
+				{
+					// read features
+					HashSet<string> ReferencedTags = new HashSet<string>();
+					EngineIni.GetArray(IniSection_TargetSettings, "IntelligentDeliveryFeatures", out List<string> Features);
+					if (Features != null && Features.Count > 0)
+					{
+						foreach (string Feature in Features)
+						{
+							Dictionary<string, string> Properties;
+							ConfigHierarchy.TryParse(Feature, out Properties);
+
+							IntelligentInstallFeature Mapping = new();
+							if (!TryGetValue(Properties, "Id", out Mapping.Id) || Mapping.Id == "")
+							{
+								throw new AutomationException("Missing Id in feature");
+							}
+
+							TryGetValue(Properties, "Tags", out string Tags);
+							Mapping.Tags = Tags.Split(';', StringSplitOptions.RemoveEmptyEntries);
+
+							TryGetValue(Properties, "bHidden", out Mapping.bHidden);
+							if (Mapping.bHidden)
+							{
+								Mapping.DisplayNameResourceId = null;
+								Mapping.ImageFileName = null;
+							}
+							else
+							{
+								TryGetValue(Properties, "DefaultDisplayName", out string DefaultDisplayName);
+
+								// read localizable display name (localized strings will be added in GDKGameConfigGenerator)
+								Mapping.DisplayNameResourceId = !string.IsNullOrEmpty(DefaultDisplayName) ? $"ms-resource:Feature_{Mapping.Id}" : null;
+
+								// see if there is an image resource for this feature (it will be added in GDKGameConfigGenerator)
+								string ImageFilename = $"Feature_{Mapping.Id}.png";
+								bool bHasImageFilename = GDKGameConfigGenerator.DoesDefaultResourceBinaryFileExistForPlatform(SC.StageTargetPlatform.PlatformType, Logger, ImageFilename, bAllowEngineFallback: false);
+								Mapping.ImageFileName = bHasImageFilename ? $@"Resources\{ImageFilename}" : null;
+							}
+
+							yield return Mapping;
+						}
+					}
+				}
+			}
+
+			internal void AddTag(string Tag)
+			{
+				if (!Tags.Contains(Tag, StringComparer.OrdinalIgnoreCase))
+				{
+					Tags = [.. Tags, Tag];
+				}
+			}
+		}
+
+
+		internal class IntelligentInstallChunk
+		{
+			public int ChunkId;
+			public string ChunkTag = "";
+
+			public string Detail_Devices = "";
+			public string Detail_Languages = ""; //NB. these are GDK Culture Ids, not UE Stage Ids
+			public bool Detail_bRequiredForLaunch = false;
+
+			// some code readability helpers
+			public bool IsOnDemand => (!Detail_bRequiredForLaunch && ChunkTag.Length > 0 && Detail_Languages.Length == 0 && Detail_Devices.Length == 0);
+			public string[] Tags => ChunkTag.Split(';', StringSplitOptions.RemoveEmptyEntries | StringSplitOptions.TrimEntries);
+			public string[] CultureIDs => Detail_Languages.Split(';', StringSplitOptions.RemoveEmptyEntries | StringSplitOptions.TrimEntries);
+
+			internal static IEnumerable<IntelligentInstallChunk> GetFromConfig(ProjectParams Params, DeploymentContext SC, string IniSection_TargetSettings)
+			{
+				ConfigHierarchy EngineIni = ConfigCache.ReadHierarchy(ConfigHierarchyType.Engine, DirectoryReference.FromFile(SC.RawProjectPath), SC.StageTargetPlatform.PlatformType, SC.CustomConfig);
+
+				var IntelligentInstallChunks = new List<IntelligentInstallChunk>();
+
+				List<string> IntelligentDeliveryChunks;
+				EngineIni.GetArray(IniSection_TargetSettings, "IntelligentDeliveryChunks", out IntelligentDeliveryChunks);
+
+				if (IntelligentDeliveryChunks != null)
+				{
+					Dictionary<string, string> CultureIds = new Dictionary<string, string>();
+
+					if (EngineIni.GetString(IniSection_TargetSettings, "StageIdOverrides", out string StageIdOverridesString) && ConfigHierarchy.TryParseAsMap(StageIdOverridesString, out Dictionary<string, string> StageIdOverrides))
+					{
+						CultureIds = StageIdOverrides;
+					}
+
+					// parse all chunks
+					foreach (var IntelligentDeliveryChunk in IntelligentDeliveryChunks)
+					{
+						Dictionary<string, string> Properties;
+						ConfigHierarchy.TryParse(IntelligentDeliveryChunk, out Properties);
+
+						IntelligentInstallChunk ChunkEntry = new IntelligentInstallChunk();
+
+						// read chunk id
+						string ChunkId;
+						if (!Properties.TryGetValue("ChunkId", out ChunkId) || !int.TryParse(ChunkId, out ChunkEntry.ChunkId))
+						{
+							Logger.LogWarning("Couldn't parse ChunkId from IntelligentDeliveryChunks in section {IniSection_TargetSettings}", IniSection_TargetSettings);
+							ChunkEntry.ChunkId = -1;
+						}
+						else
+						{
+							ChunkEntry.ChunkId += OSChunkIdOffset;
+						}
+
+						// read tags
+						if (!Properties.TryGetValue("Tags", out ChunkEntry.ChunkTag))
+						{
+							ChunkEntry.ChunkTag = "";
+						}
+
+						// read devices
+						string RawDevices;
+						if (Properties.TryGetValue("Devices", out RawDevices))
+						{
+							string[] Devices;
+							if (ConfigHierarchy.TryParse(RawDevices, out Devices) && Devices != null && Devices.Length > 0)
+							{
+								ChunkEntry.Detail_Devices = string.Join(";", Devices);
+							}
+						}
+
+						// read stage ids
+						string RawStageIds;
+						if (Properties.TryGetValue("StageIds", out RawStageIds))
+						{
+							string[] StageIds;
+							if (ConfigHierarchy.TryParse(RawStageIds, out StageIds) && StageIds != null && StageIds.Length > 0)
+							{
+								foreach (string StageId in StageIds)
+								{
+									// look up the CultureId from the given CultureToStage Id
+									string CultureId;
+									if (!CultureIds.TryGetValue(StageId, out CultureId))
+									{
+										CultureId = StageId;
+									}
+									ChunkEntry.Detail_Languages += CultureId + ";";
+								}
+							}
+							ChunkEntry.Detail_Languages = ChunkEntry.Detail_Languages.TrimEnd(';');
+						}
+
+						// read launch flag
+						string RequiredForLaunch;
+						if (!Properties.TryGetValue("bRequiredForLaunch", out RequiredForLaunch) || !bool.TryParse(RequiredForLaunch, out ChunkEntry.Detail_bRequiredForLaunch))
+						{
+							ChunkEntry.Detail_bRequiredForLaunch = false;
+						}
+
+						IntelligentInstallChunks.Add(ChunkEntry);
+					}
+				}
+
+				return IntelligentInstallChunks;
+			}
+		}
+
+
+		protected class GDKBundleSettings : BundleUtils.BundleSettings
+		{
+			public int ChunkID { get; set; }
+			public string PlatformChunkName { get; set; }
+		}
+
+		// Chunk 2 maps to UE chunk 0 and also includes all content not mapped explicitly to another chunk index.
+		const int OSChunkIdOffset = 2; // We need an offset to map from UE chunk to GDK chunk because a GDK chunk ID of 0 is not valid. This is 2 for historical CUv1 reasons (and even CUv3 doesn't support chunk ID reassignment without causing a full re-download)
+
+		/// The CustomChunkMapping table is used to assign files to specific chunks using wildcard path matching.
+		/// +CustomChunkMapping=(Pattern="shootergame/content/movies/*",ChunkId=1)
+		/// +CustomChunkMapping=(Pattern="shootergame/content/movies/movie1.mp4",ChunkId=1)
+		protected struct CustomChunkMapping
+		{
+			public string Pattern;
+			public int ChunkId;
+
+			internal static IEnumerable<CustomChunkMapping> GetFromConfig(ProjectParams Params, DeploymentContext SC, string IniSection_TargetSettings)
+			{
+				const string Key = "CustomChunkMapping";
+
+				ConfigHierarchy EngineIni = ConfigCache.ReadHierarchy(ConfigHierarchyType.Engine, DirectoryReference.FromFile(SC.RawProjectPath), SC.StageTargetPlatform.PlatformType, SC.CustomConfig);
+
+				var Section = EngineIni.FindSection(IniSection_TargetSettings);
+				if (Section != null && Section.TryGetValues(Key, out IReadOnlyList<string> ChunkMappingList))
+				{
+					foreach (var ChunkMapping in ChunkMappingList)
+					{
+						if (!ConfigHierarchy.TryParse(ChunkMapping, out Dictionary<string, string> Properties))
+							throw new AutomationException("Failed to parse properties from chunk mapping.");
+
+						CustomChunkMapping Pair;
+						if (!TryGetValue(Properties, "Pattern", out Pair.Pattern))
+							throw new AutomationException($"Failed to retrieve Pattern value from {Key} in config section \"{IniSection_TargetSettings}\".");
+
+						if (!TryGetValue(Properties, "ChunkId", out Pair.ChunkId))
+							throw new AutomationException($"Failed to retrieve ChunkId value from {Key} in config section \"{IniSection_TargetSettings}\".");
+
+						yield return Pair;
+					}
+				}
+			}
+		}
+
+
+
+
+
+
+
+
+		class PackageFileDesc
+		{
+			public StagedFileReference File;
+			public string RelativeSourceDir;
+			public int ChunkId;
+			public bool bAllowInManifest;
+		};
+
+
+
+		private void GeneratePackageMetadataInternal(ProjectParams Params, DeploymentContext SC, IEnumerable<StagedFileReference> FilesToInclude, IEnumerable<StagedFileReference> GeneratedFilesToInclude, DirectoryReference LayoutDirectory, ILogger Logger)
+		{
+			IReadOnlyDictionary<string, GDKBundleSettings> InstallBundles = null;
+			if (BundleUtils.HasPlatformBundleSource(SC.RawProjectPath.Directory, SC.StageTargetPlatform.PlatformType, SC.CustomConfig))
+			{
+				Logger.LogInformation("Using Intelligent Delivery Chunk mappings from install bundle config.");
+
+				BundleUtils.LoadBundleConfig(SC.RawProjectPath.Directory, SC.StageTargetPlatform.PlatformType, SC.CustomConfig, out InstallBundles, (Settings, InBundleConfig, BundleSection) =>
+				{
+					if (InBundleConfig.GetInt32(BundleSection, "PlatformChunkID", out int ChunkID))
+					{
+						Settings.ChunkID = ChunkID;
+					}
+					else
+					{
+						Settings.ChunkID = 0;
+					}
+
+					if (InBundleConfig.GetString(BundleSection, "PlatformChunkName", out string ChunkTag))
+					{
+						Settings.PlatformChunkName = ChunkTag;
+					}
+					else
+					{
+						Settings.PlatformChunkName = "";
+					}
+				});
+			}
+
+			var IntelligentInstallChunks = IntelligentInstallChunk.GetFromConfig(Params, SC, IniSection_TargetSettings).ToList();
+			var IntelligentInstallFeatures = IntelligentInstallFeature.GetFromConfig(Params, SC, IniSection_TargetSettings).ToList();
+			var IntelligentInstallRecipes = IntelligentInstallRecipe.GetFromConfig(Params, SC, IniSection_TargetSettings).ToList();
+
+			ConfigHierarchy EngineIni = ConfigCache.ReadHierarchy(ConfigHierarchyType.Engine, DirectoryReference.FromFile(SC.RawProjectPath), SC.StageTargetPlatform.PlatformType, SC.CustomConfig);
+			EngineIni.GetBool(IniSection_TargetSettings, "bUseFeaturesAndRecipes", out bool bUseFeaturesAndRecipes);
+
+
+
+			DirectoryReference.CreateDirectory(LayoutDirectory);
+			FileReference PackageManifestFile = FileReference.Combine(LayoutDirectory, "package.manifest");
+			FileReference LayoutXmlFile = FileReference.Combine(LayoutDirectory, "layout.xml");
+
+			string RelativeLayoutDir = SC.StageDirectory.MakeRelativeTo(LayoutDirectory);
+
+
+			// Build a map of all the files to include in the package, and which chunks they belong to.
+			IReadOnlyDictionary<StagedFileReference, PackageFileDesc> FileChunkMapping;
+			{
+				var LocalFileChunkMapping = new Dictionary<StagedFileReference, PackageFileDesc>();
+				FileChunkMapping = LocalFileChunkMapping;
+
+				foreach (var File in FilesToInclude)
+				{
+					// Assume all files are in chunk 0, except for pak files with a suffixed chunk number (i.e. pakchunkN-platform.pak).
+					var PakFileMatch = Regex.Match(File.Name, @".*[/\\]pakchunk([0-9]+)[a-zA-Z]*-[a-zA-Z0-9]+\.(pak|utoc||uondemandtoc|ucas)");
+					int ChunkId = PakFileMatch.Success ? int.Parse(PakFileMatch.Groups[1].Value) : 0;
+
+					LocalFileChunkMapping.Add(File, new PackageFileDesc()
+					{
+						File = File,
+						RelativeSourceDir = Path.Combine(RelativeLayoutDir, File.Directory.Name),
+						ChunkId = ChunkId,
+						bAllowInManifest = true,
+					});
+				}
+
+				// add all the generated files to include
+				if (GeneratedFilesToInclude != null)
+				{
+					foreach (var File in GeneratedFilesToInclude)
+					{
+						LocalFileChunkMapping.Add(File, new PackageFileDesc()
+						{
+							File = File,
+							RelativeSourceDir = "./" + File.Directory.Name,
+							ChunkId = 0,
+							bAllowInManifest = false,
+						});
+					}
+				}
+
+				// Add the package.manifest files to the chunk mapping
+				{
+					StagedFileReference StagedManifestFile = StagedFileReference.Combine(PackageManifestFile.MakeRelativeTo(LayoutDirectory));
+					LocalFileChunkMapping.Add(StagedManifestFile, new PackageFileDesc()
+					{
+						File = StagedManifestFile,
+						RelativeSourceDir = "./" + StagedManifestFile.Directory.Name,
+						ChunkId = 0,
+						bAllowInManifest = false,
+					});
+				}
+			}
+
+			// apply appropriate chunk remapping
+			Dictionary<Int32, HashSet<string>> BundleDefinedImplicitChunkTags = new();  // Maps from chunk id to tag for bundle when there's no explict IntelligentDeliveryChunk definition
+			if (InstallBundles != null)
+			{
+				HashSet<string> LaunchFeatures = [];
+
+				// Override the mappings we built above for anything that matches the regex
+				foreach (var Pair in FileChunkMapping)
+				{
+					string Filename = Path.GetFileName(Pair.Key.Name);
+					GDKBundleSettings Bundle = BundleUtils.MatchBundleSettings(Filename, InstallBundles);
+					if (Bundle != null)
+					{
+						Logger.LogInformation("Mapping {Filename} to Chunk {ChunkId} for Bundle {BundleName}", Filename, Bundle.ChunkID, Bundle.Name);
+						Pair.Value.ChunkId = Bundle.ChunkID;
+
+						if (Bundle.ChunkID != -1 && Bundle.PlatformChunkName != "")
+						{
+							string ExistingTag = IntelligentInstallChunks.FirstOrDefault(X =>
+								(X.ChunkId - OSChunkIdOffset) == Bundle.ChunkID &&
+								!X.Tags.Contains(Bundle.PlatformChunkName, StringComparer.OrdinalIgnoreCase) &&
+								!X.CultureIDs.Contains(Bundle.PlatformChunkName, StringComparer.OrdinalIgnoreCase))
+								?.ChunkTag;
+							if (ExistingTag != null)
+							{
+								throw new AutomationException($"Bundle '{Bundle.Name}' with PlatformChunkId={Bundle.ChunkID} has PlatformChunkName={Bundle.PlatformChunkName} but that tag is not referenced in tags or mapped CultureIds for IntelligentDeliveryChunks chunkId={Bundle.ChunkID} {ExistingTag}"); // @todo: could merge these tags, but that'd be a more invasive refactor & isn't supported on some other platforms
+							}
+
+							if (!BundleDefinedImplicitChunkTags.ContainsKey(Bundle.ChunkID))
+							{
+								BundleDefinedImplicitChunkTags.Add(Bundle.ChunkID, new());
+							}
+							BundleDefinedImplicitChunkTags[Bundle.ChunkID].Add(Bundle.PlatformChunkName);
+
+							// create an implicit feature if necessary
+							if (bUseFeaturesAndRecipes)
+							{
+								IntelligentInstallChunk ExistingChunk = IntelligentInstallChunks.FirstOrDefault( X => (X.ChunkId - OSChunkIdOffset) == Bundle.ChunkID );
+								bool bIsLanguageChunk = (ExistingChunk != null) && ExistingChunk.Detail_Languages.Any();
+								bool bIsLaunchChunk = (ExistingChunk != null) && ExistingChunk.Detail_bRequiredForLaunch;
+
+								if (!bIsLanguageChunk)
+								{
+									// find the associated feature, or create it if it doesn't exist
+									// note that we do this even for launch chunks because the bundle system needs all bundles to have a Feature so they can be mounted
+									IntelligentInstallFeature Feature = IntelligentInstallFeatures.FirstOrDefault(X => X.Tags.Contains(Bundle.PlatformChunkName, StringComparer.OrdinalIgnoreCase));
+									if (Feature == null)
+									{
+										IntelligentInstallFeatures.Add( Feature = new IntelligentInstallFeature
+										{
+											Id = Bundle.PlatformChunkName,
+											Tags = [Bundle.PlatformChunkName],
+											bHidden = true,
+										});
+									}
+									else
+									{
+										Feature.AddTag(Bundle.PlatformChunkName);
+									}
+
+									if (bIsLaunchChunk)
+									{
+										LaunchFeatures.Add(Feature.Id);
+									}
+								}
+							}
+						}
+					}
+					else
+					{
+						Logger.LogInformation("Mapping {Filename} to Chunk {ChunkId} for unknown bundle", Filename, Pair.Value.ChunkId);
+					}
+				}
+
+				// Create a default Recipe that contains the launch Features, if there isn't one already
+				if (bUseFeaturesAndRecipes && IntelligentInstallRecipes.Count == 0 && LaunchFeatures.Count > 0)
+				{
+					Logger.LogInformation("Creating default Recipe for launch Features {features}", string.Join(',',LaunchFeatures));
+
+					IntelligentInstallRecipes.Add(
+						new IntelligentInstallRecipe
+						{
+							Id = "default",
+							FeatureIds = [.. LaunchFeatures],
+						}
+					);			
+				}
+			}
+			else
+			{
+				// If we have custom chunk mappings, override the mappings we built above for anything that matches the wildcard patterns.
+				foreach (var Mapping in CustomChunkMapping.GetFromConfig(Params, SC, IniSection_TargetSettings))
+				{
+					string FullPathPattern = CombinePaths(SC.StageDirectory.FullName, Mapping.Pattern);
+
+					var MatchedFiles = (FileFilter.FindWildcardIndex(FullPathPattern) != -1)
+						? FileFilter.ResolveWildcard(FullPathPattern)
+						: new List<FileReference>() { new FileReference(FullPathPattern) };
+
+					foreach (var MatchedFile in MatchedFiles)
+					{
+						var FileRef = new StagedFileReference(MatchedFile.MakeRelativeTo(SC.StageDirectory));
+						if (FileChunkMapping.ContainsKey(FileRef))
+						{
+							FileChunkMapping[FileRef].ChunkId = Mapping.ChunkId;
+						}
+					}
+				}
+			}
+
+			// Generate a contiguous list of chunks up to the maximum ID listed in either the FileChunkMappings or IntelligentInstallChunks
+			int MaxChunkId = (new[] { 0 })
+				.Concat(FileChunkMapping.Values.Select(c => c.ChunkId))
+				.Concat(IntelligentInstallChunks.Select(c => c.ChunkId - OSChunkIdOffset))
+				.Where(c => c >= 0)
+				.Distinct()
+				.Max();
+			var AllChunkIds = Enumerable.Range(0, MaxChunkId + 1).ToList();
+			AllChunkIds.Sort();
+
+			// Generate a lookup table of all chunk details for all chunks
+			var ChunkInfo = AllChunkIds.ToDictionary(
+				ChunkId => ChunkId,
+				ChunkId => (from Mapping in IntelligentInstallChunks
+							where Mapping.ChunkId - OSChunkIdOffset == ChunkId
+							select Mapping)
+							.FirstOrDefault() ?? new IntelligentInstallChunk
+							{
+								ChunkId = ChunkId + OSChunkIdOffset,
+								ChunkTag = BundleDefinedImplicitChunkTags.TryGetValue(ChunkId, out HashSet<string> Tags) ? string.Join(';', Tags) : ""
+							}
+			);
+
+			// Create a collection of all chunks that are required for launch
+			var InitialChunkIds = (from Mapping in IntelligentInstallChunks
+								   where Mapping.Detail_bRequiredForLaunch
+								   select Mapping.ChunkId - OSChunkIdOffset).Distinct().ToList();
+			InitialChunkIds.Sort();
+			if (InitialChunkIds.Count == 0)
+			{
+				int MinChunkId = AllChunkIds.Min();
+				if (IntelligentInstallChunks.Count > 0)
+				{
+					Logger.LogTrace($"No initial chunk ID specified. Assuming chunk {MinChunkId} is the initial chunk.");
+				}
+				InitialChunkIds.Add(MinChunkId);
+			}
+
+			// Gather the tags that can be used in the layout.xml
+			var AllowedLayoutTags = ChunkInfo.Values.SelectMany(Chunk => Chunk.Tags).ToHashSet(StringComparer.InvariantCultureIgnoreCase);
+			if (IntelligentInstallFeatures.Any())
+			{
+				// do not allow tags that are not referenced by a feature. These extra tags are useful for the manifest generated below, but will cause packaging failures in layout.xml
+				var ReferencedTags = IntelligentInstallFeatures
+					.Select(Feature => Feature.Tags)
+					.SelectMany(X => X)
+					.ToHashSet(StringComparer.InvariantCultureIgnoreCase);
+
+				AllowedLayoutTags.IntersectWith(ReferencedTags);
+			}
+
+			// Generate the package.manifest file
+			{
+				static string ToManifestFile(StagedFileReference File) => File.Name.Replace('/', '\\');
+
+				var StagedFilesByChunk =
+					(from File in FileChunkMapping
+					 where File.Value.bAllowInManifest && File.Value.ChunkId >= 0
+					 orderby File.Value.ChunkId ascending
+					 group File by File.Value.ChunkId)
+						.ToDictionary(g => g.Key, g => g.ToList());
+
+				var ManifestObj = new Dictionary<string, object>()
+				{
+					["chunks"] = StagedFilesByChunk
+						.Select(Chunk => new Dictionary<string, object>()
+						{
+							["id"] = Chunk.Key,
+							["Tag"] = ChunkInfo[Chunk.Key].ChunkTag,
+							["Devices"] = ChunkInfo[Chunk.Key].Detail_Devices,
+							["Languages"] = ChunkInfo[Chunk.Key].Detail_Languages,
+							["IsInitial"] = InitialChunkIds.Contains(Chunk.Key),
+
+							["files"] = Chunk.Value.Select(File => ToManifestFile(File.Key))
+						})
+				};
+
+				if (IntelligentInstallFeatures.Any())
+				{
+					ManifestObj["features"] = IntelligentInstallFeatures
+						.Select(Feature => new Dictionary<string, object>()
+						{
+							["Id"] = Feature.Id,
+							["Tags"] = string.Join(';', Feature.Tags)
+						});
+				}
+
+				var JsonOptions = new JsonSerializerOptions { WriteIndented = true };
+				var Contents = JsonSerializer.Serialize(ManifestObj, JsonOptions);
+				File.WriteAllText(PackageManifestFile.FullName, Contents, new UTF8Encoding(encoderShouldEmitUTF8Identifier: false));
+			}
+
+
+
+			// Generate layout.xml file
+			{
+				bool bHasRecipes = IntelligentInstallFeatures.Any() || IntelligentInstallRecipes.Any();
+
+				string ToSourcePath(string SrcDir) => @$"{SrcDir.Replace('/', '\\')}";
+				string ToDestPath(StagedFileReference File) => @$"\{File.Directory.Name.Replace('/', '\\')}";
+
+				var StagedFilesByChunk =
+					(from File in FileChunkMapping
+					 where File.Value.ChunkId >= 0
+					 orderby File.Value.ChunkId ascending
+					 orderby InitialChunkIds.Contains(File.Value.ChunkId) descending // launch chunks first
+					 group File by File.Value.ChunkId)
+						.ToDictionary(g => g.Key, g => g.ToList());
+
+				var LayoutXml = new XElement("Package",
+					IntelligentInstallRecipes.Any() ? new XElement("Recipes",
+						from Recipe in IntelligentInstallRecipes
+						select new XElement("Recipe",
+							new XAttribute("Id", Recipe.Id),
+							opt_XAttribute("IncludedFeatures", Recipe.FeatureIds),
+							opt_XAttribute("Devices", Recipe.Devices),
+							opt_XAttribute("StoreId", Recipe.StoreIds)
+						)) : null,
+
+					IntelligentInstallFeatures.Any() ? new XElement("Features",
+						from Feature in IntelligentInstallFeatures.Where( F => F.Tags.Length > 0 ) // skip features with no tags
+						select new XElement("Feature",
+							new XAttribute("Id", Feature.Id),
+							opt_XAttribute("Tags", Feature.Tags),
+							opt_XAttribute("Hidden", Feature.bHidden),
+							opt_XAttribute("DisplayName", Feature.DisplayNameResourceId),
+							opt_XAttribute("Image", Feature.ImageFileName)
+						)) : null,
+
+					from Chunk in StagedFilesByChunk
+					select new XElement("Chunk",
+						new XAttribute("Id", Chunk.Key + OSChunkIdOffset),
+						opt_XAttribute("OnDemand", ChunkInfo[Chunk.Key].IsOnDemand && !bHasRecipes && !InitialChunkIds.Contains(Chunk.Key)), // OnDemand tags only apply when not using Features & Recipes. Launch chunks can't be OnDemand
+						opt_XAttribute("Languages", ChunkInfo[Chunk.Key].Detail_Languages),
+						opt_XAttribute("Devices", ChunkInfo[Chunk.Key].Detail_Devices),
+						opt_XAttribute("Tags", ChunkInfo[Chunk.Key].Tags.Intersect(AllowedLayoutTags).ToArray()),                         // Filter chunk's tags to only include the permitted ones
+						opt_XAttribute("Marker", "Launch", bAdditionalPrerequisite: (Chunk.Key == InitialChunkIds.Max())),                   // Last initial chunk gets the Launch marker
+						from FilePair in Chunk.Value
+						select new XElement("FileGroup",
+							new XAttribute("Include", Path.GetFileName(FilePair.Key.Name)),
+							new XAttribute("SourcePath", ToSourcePath(FilePair.Value.RelativeSourceDir)),
+							new XAttribute("DestinationPath", ToDestPath(FilePair.Value.File))
+						)
+					)
+				);
+
+				LayoutXml.Save(LayoutXmlFile.FullName);
+			}
+		}
+
+
+
+		/// <summary>
+		/// Generates the layout.xml and associated files required for packaging in the given layout directory
+		/// </summary>
+		public void GeneratePackageMetadata(ProjectParams Params, DeploymentContext SC, DirectoryReference PackageMetadataDir, UnrealTargetConfiguration? OnlyConfig = null)
+		{
+			FileUtils.ForceDeleteDirectory(PackageMetadataDir);
+			FileUtils.CreateDirectoryTree(PackageMetadataDir);
+
+			// Gather all the files we'll be including in the package
+			var FilesUnion = SC.FilesToStage.NonUFSFiles.Union(SC.FilesToStage.NonUFSSystemFiles);
+			if (!Params.UsePak(Platform.GetPlatform(PlatformType)))
+			{
+				// Don't include UFS files in the layout.xml if we're cooking on-the-fly or using ZenStore.
+				// UFS files are provided by the connection to the server.
+				if (!Params.CookOnTheFly && !Params.ShouldTreatAsFileServer(SC))
+				{
+					FilesUnion = FilesUnion.Union(SC.FilesToStage.UFSFiles);
+				}
+
+				FilesUnion = FilesUnion.Union(SC.CrashReporterUFSFiles);
+			}
+			else
+			{
+				var DLCName = Params.HasDLCName ? Params.DLCFile.GetFileNameWithoutExtension() : null;
+
+				var PaksDir = Params.HasDLCName
+					? Path.Combine(SC.StageDirectory.FullName, SC.ShortProjectName, "Plugins", DLCName, "Content", "Paks", SC.FinalCookPlatform)
+					: Path.Combine(SC.StageDirectory.FullName, SC.ShortProjectName, "Content", "Paks");
+				var FilesInPaksDir = (new[] { "*.pak", "*.ucas", "*.utoc", "*.uondemandtoc", "*.sig" }).SelectMany(Pattern => FindFiles(Pattern, false, PaksDir));
+
+				FilesUnion = FilesUnion.Union(
+					from f in FilesInPaksDir
+					let Reference = new FileReference(Path.Combine(SC.StageDirectory.FullName, f))
+					select new KeyValuePair<StagedFileReference, FileReference>(new StagedFileReference(Reference.MakeRelativeTo(SC.StageDirectory)), Reference)
+				);
+
+				// crash reporter is handled as a special case in CopyBuildToStagingDirectory
+				var CrashReporterPak = FileReference.Combine(SC.RuntimeRootDir, "Engine", "Programs", "CrashReportClient", "Content", "Paks", "CrashReportClient.pak");
+				if (SC.CrashReporterUFSFiles.Any() && FileReference.Exists(CrashReporterPak))
+				{
+					FilesUnion = FilesUnion.Append(new KeyValuePair<StagedFileReference, FileReference>(new StagedFileReference(CrashReporterPak.MakeRelativeTo(SC.StageDirectory)), CrashReporterPak));
+				}
+			}
+
+			// Only include the debug files in non-distribution builds, and when not running UAT with the "-nodebuginfo" flag.
+			// We add debug files here regardless of Params.Distribution because we may be generating multiple configurations,
+			// some allowing debug files, some don't.
+			if (!Params.NoDebugInfo && !Params.SeparateDebugInfo && !Params.Distribution)
+			{
+				FilesUnion = FilesUnion.Union(SC.FilesToStage.NonUFSDebugFiles);
+			}
+
+			// generate a filtered list of the files that should be included in the package
+			IEnumerable<StagedFileReference> FilesToInclude = FilesUnion.Where(X =>
+				ShouldPackageFile(X.Value, Params, SC, OnlyConfig))
+				.Select(f => f.Key);
+
+			IEnumerable<StagedFileReference> GeneratedFilesToInclude = [];
+
+			// Generate MicrosoftGame.config if we are not generating metadata in the staging root (if we are we'll want to re-use the staging manifest)
+			if (PackageMetadataDir != SC.StageDirectory)
+			{
+				// exclude any staged manifest resources as we're generating new ones
+				IEnumerable<StagedFileReference> ExistingStagedManifestFiles = GetGDKManifestFilesInDirectory(SC.StageDirectory)
+					.Select(File => StagedFileReference.Combine(File.MakeRelativeTo(SC.StageDirectory)));
+				FilesToInclude = FilesToInclude.Except(ExistingStagedManifestFiles);
+
+				// generate the manifest and gather the files
+				Dictionary<UnrealTargetConfiguration, string> ExecutablePairs = GDKAutomationUtils.GetGameConfigExecutables(Params, SC, Purpose: GDKAutomationUtils.GameConfigPurpose.ForStaging, IniSection_TargetSettings, OnlyConfig);
+				GDKGameConfigGenerator.Create(PlatformType, Logger, PackageMetadataDir, SC.ShortProjectName, Params.RawProjectPath, ExecutablePairs, bDeleteIntermediates: true, DLCFile: Params.DLCFile, CustomConfig: SC.CustomConfig);
+				GeneratedFilesToInclude = GeneratedFilesToInclude.Union( GetGDKManifestFilesInDirectory(PackageMetadataDir)
+					.Select(File => StagedFileReference.Combine(File.MakeRelativeTo(PackageMetadataDir)))
+				);
+			}
+
+			// Add any platform specific files that need to be included in the package
+			GeneratedFilesToInclude = GeneratedFilesToInclude.Union(GetPlatformFilesToPackage(Params, SC, PackageMetadataDir, OnlyConfig));
+
+
+			// Generate the packaging metadata
+			GeneratePackageMetadataInternal(Params, SC, FilesToInclude, GeneratedFilesToInclude, PackageMetadataDir, Logger);
+		}
+
+		protected virtual IEnumerable<StagedFileReference> GetPlatformFilesToPackage(ProjectParams Params, DeploymentContext SC, DirectoryReference PackageMetadataDir, UnrealTargetConfiguration? OnlyConfig)
+		{
+			return [];
+		}
+
+
+		#region helper functions
+
+		/// attempt to parse out a typed value from the given dictionary & key
+		protected static bool TryGetValue<T>(Dictionary<string, string> Dictionary, string Key, out T Value)
+		{
+			try
+			{
+				Value = (T)TypeDescriptor.GetConverter(typeof(T)).ConvertFromInvariantString(Dictionary[Key]);
+				return true;
+			}
+			catch
+			{
+				Value = default(T);
+				return false;
+			}
+		}
+
+		/// attempt to parse out an array of strings from the given dictionary & key
+		protected static bool TryGetValue(Dictionary<string, string> Dictionary, string Key, out string[] Value)
+		{
+			if (TryGetValue(Dictionary, Key, out string RawValue))
+			{
+				if (ConfigHierarchy.TryParse(RawValue, out Value))
+				{
+					return true;
+				}
+			}
+
+			Value = Array.Empty<string>();
+			return false;
+		}
+
+		/// create an XAttribute as long as the value is not null and the additionl prerequisite is true
+		protected static XAttribute opt_XAttribute(string Key, string OptValue, bool bAdditionalPrerequisite = true)
+		{
+			if (!bAdditionalPrerequisite || string.IsNullOrEmpty(OptValue))
+			{
+				return null;
+			}
+			return new XAttribute(Key, OptValue);
+		}
+
+		/// create an XAttribute of semicolon-separated values as long as there are values
+		protected static XAttribute opt_XAttribute(string Key, string[] Values)
+		{
+			if (Values == null || Values.Length == 0)
+			{
+				return null;
+			}
+			return new XAttribute(Key, string.Join(';', Values));
+		}
+
+		/// create an XAttribute as long as the value is true
+		protected static XAttribute opt_XAttribute(string Key, bool OptValue)
+		{
+			if (!OptValue)
+			{
+				return null;
+			}
+			return new XAttribute(Key, OptValue);
+		}
+
+		#endregion
+
+		#endregion
+	}
+}
diff --git a/Engine/Source/Programs/UnrealBuildTool/Platform/Microsoft/GDK/GDKContentOnlyProjectSupport.cs b/Engine/Source/Programs/UnrealBuildTool/Platform/Microsoft/GDK/GDKContentOnlyProjectSupport.cs
new file mode 100644
--- /dev/null
+++ b/Engine/Source/Programs/UnrealBuildTool/Platform/Microsoft/GDK/GDKContentOnlyProjectSupport.cs
@@ -0,0 +1,194 @@
+// Copyright Epic Games, Inc. All Rights Reserved.
+
+using EpicGames.Core;
+using Microsoft.Extensions.Logging;
+using System;
+using System.IO;
+using System.Linq;
+using UnrealBuildBase;
+
+namespace UnrealBuildTool
+{
+	/// <summary>
+	/// Helper functionality for easier F5 debugging for blueprint-only projects
+	/// </summary>
+	public static class GDKContentOnlyProjectSupport
+	{
+		private class Config
+		{
+			/// <summary>
+			/// Path to the .uproject file to use when F5 debugging a content-only project via the UE5 engine project
+			/// </summary>
+			[XmlConfigFile(Category = "GDKPlatform")]
+			public string ContentOnlyDebugProject = "";
+
+			public Config()
+			{
+				XmlConfig.ApplyTo(this);
+			}
+		}
+
+		/// <summary>
+		/// Get the configured content-only debug project, if any
+		/// </summary>
+		public static string ContentOnlyDebugProject
+		{
+			get
+			{
+				Config Config = new();
+				return Config.ContentOnlyDebugProject;
+			}
+		}
+
+		/// <summary>
+		/// Helper function to get the content-only debug project for the given target, if applicable and supported
+		/// </summary>
+		private static string? GetProjectFileFromTargetInternal(ReadOnlyTargetRules Target, ILogger Logger)
+		{
+			string ProjectFileName = ContentOnlyDebugProject;
+
+			// early out if this isn't required
+			if (Target.ProjectFile != null || Target.Type == TargetType.Program || String.IsNullOrEmpty(ProjectFileName) || !Environment.CommandLine.Contains("-frommsbuild", StringComparison.OrdinalIgnoreCase))
+			{
+				return null;
+			}
+
+			// find & verify the uproject name
+			if (String.IsNullOrEmpty(Path.GetExtension(ProjectFileName)))
+			{
+				ProjectFileName += ".uproject";
+			}
+			if (String.IsNullOrEmpty(Path.GetDirectoryName(ProjectFileName)))
+			{
+				FileReference? FoundProject = NativeProjects.EnumerateProjectFiles(Logger).FirstOrDefault(X => X.GetFileName().Equals(ProjectFileName, StringComparison.OrdinalIgnoreCase));
+				if (FoundProject != null)
+				{
+					ProjectFileName = FoundProject.FullName;
+				}
+			}
+			else if (!Path.IsPathRooted(ProjectFileName))
+			{
+				ProjectFileName = Path.GetFullPath(ProjectFileName, Unreal.RootDirectory.FullName);
+			}
+			if (!File.Exists(ProjectFileName))
+			{
+				throw new BuildException($"ContentOnlyDebugProject not found: {ProjectFileName}");
+			}
+
+			return ProjectFileName;
+		}
+
+		/// <summary>
+		/// Stores the path to the content-only debug project in the target receipt
+		/// </summary>
+		public static bool ModifyTargetReceipt(ReadOnlyTargetRules Target, TargetReceipt Receipt, ILogger Logger)
+		{
+			string? ProjectFileName = GetProjectFileFromTargetInternal(Target, Logger);
+			if (ProjectFileName == null)
+			{
+				return false;
+			}
+
+			// add the receipt property
+			Receipt.AdditionalProperties.Add(new ReceiptProperty("ContentOnlyDebugProject", ProjectFileName));
+			Logger.LogInformation("*** GDK F5 debugging for content-only projects is enabled via ContentOnlyDebugProject. Using project {ProjectFileName} ***", ProjectFileName);
+			return true;
+		}
+
+		/// <summary>
+		/// Retrieves the path to the project from the given target receipt
+		/// </summary>
+		public static FileReference? GetProjectFileFromTargetReceipt(TargetReceipt Receipt, ILogger Logger)
+		{
+			FileReference? ProjectFile = Receipt.ProjectFile;
+
+			// No project file was specified - this might be the UE5 project & someone wants to F5 debug it.
+			// We need to know the project so the MicrosoftGame.config can be generated with the correct package data
+			if (ProjectFile == null && Receipt.TargetType != TargetType.Program && Environment.CommandLine.Contains("-frommsbuild", StringComparison.OrdinalIgnoreCase))
+			{
+				string? ProjectFilePath = Receipt.AdditionalProperties.FirstOrDefault(X => X.Name == "ContentOnlyDebugProject")?.Value;
+				ProjectFile = GetVerifiedProjectFile(ProjectFilePath, Logger, true);
+				if (ProjectFile == null)
+				{
+					return null;
+				}
+			}
+
+			return ProjectFile;
+		}
+
+		private static FileReference? GetVerifiedProjectFile(string? ProjectFilePath, ILogger Logger, bool bUseLog)
+		{
+			FileReference? ProjectFile = null;
+
+			if (ProjectFilePath == null)
+			{
+				if (bUseLog)
+				{
+					// ideally this would be a build failure for F5 debugging - unfortunately this same code is run via UnrealVS batch build and we can't distingish it
+					Logger.LogError("{BuildConfigXmlFileName}: warning : *** No project specified for the engine solution - this is required for F5 debugging so that the correct MicrosoftGame.config can be generated. Missing <ContentOnlyDebugProject> *** \nPlease add this to your BuildConfiguration.xml:\n\n<GDKPlatform>\n\t<ContentOnlyDebugProject>Path/To/Project.uproject</ContentOnlyDebugProject>\n</GDKPlatform>\n", GetBuildConfigXmlFileName());
+				}
+			}
+			else
+			{
+				if (!Path.IsPathRooted(ProjectFilePath))
+				{
+					ProjectFilePath = Path.Combine(Unreal.RootDirectory.FullName, ProjectFilePath);
+				}
+
+				ProjectFile = new FileReference(ProjectFilePath);
+				if (!FileReference.Exists(ProjectFile))
+				{
+					throw new BuildException($"*** Cannot find ContentOnlyDebugProject {ProjectFile} ***");
+				}
+
+				// see if this project has a temporary target.cs file
+				DirectoryReference ProjectDir = ProjectFile.Directory;
+				DirectoryReference TempSourceDir = DirectoryReference.Combine(ProjectDir, "Intermediate", "Source");
+				if (DirectoryReference.Exists(TempSourceDir) && DirectoryReference.EnumerateFiles(TempSourceDir, "*.target.cs", SearchOption.TopDirectoryOnly).Any())
+				{
+					throw new BuildException($"{GetBuildConfigXmlFileName()}: error : *** Content-only project {ProjectFile.GetFileName()} has a temporary target.cs. You need to include temp targets in your solution and debug {ProjectFile.GetFileNameWithoutAnyExtensions()} directly.\nPlease add this to your BuildConfiguration.xml and regenerate project files: ***\n<ProjectFileGenerator>\n\t<bIncludeTempTargets>true</bIncludeTempTargets>\n</ProjectFileGenerator>\n");
+				}
+			}
+
+			return ProjectFile;
+		}
+
+		private static string GetBuildConfigXmlFileName()
+		{
+			// find a BuildConfiguration.xml, favoring the global appdata one
+			DirectoryReference? AppDataFolder = DirectoryReference.GetSpecialFolder(Environment.SpecialFolder.ApplicationData);
+			XmlConfig.InputFile? BuildConfig = XmlConfig.InputFiles.FirstOrDefault(X => AppDataFolder != null && X.Location.IsUnderDirectory(AppDataFolder)) ?? XmlConfig.InputFiles.FirstOrDefault();
+
+			// ideally we have a path to a usable BuildConfiguration.xml, but any 'file name' will do; it will get shown in VS's Error List window.
+			return BuildConfig?.Location.FullName ?? "BuildConfiguration.xml";
+		}
+
+		/// <summary>
+		/// Retrieves the path to the project from the given target
+		/// </summary>
+		public static FileReference? GetProjectFileFromTarget(ReadOnlyTargetRules Target, ILogger Logger, bool bUseLog = true)
+		{
+			FileReference? ProjectFile = Target.ProjectFile;
+
+			// No project file was specified - this might be the UE5 project & someone wants to F5 debug it.
+			// We need to know the project so the MicrosoftGame.config can be generated with the correct package data
+			if (ProjectFile == null && Target.Type != TargetType.Program && Environment.CommandLine.Contains("-frommsbuild", StringComparison.OrdinalIgnoreCase))
+			{
+				string? ProjectFilePath = GetProjectFileFromTargetInternal(Target, Logger);
+				ProjectFile = GetVerifiedProjectFile(ProjectFilePath, Logger, bUseLog);
+				if (ProjectFile == null)
+				{
+					return null;
+				}
+
+				if (bUseLog)
+				{
+					Logger.LogInformation("*** GDK F5 debugging for content-only projects is enabled via ContentOnlyDebugProject. Using project {ProjectFileName} ***", ProjectFilePath);
+				}
+			}
+
+			return ProjectFile;
+		}
+	}
+}
diff --git a/Engine/Source/Programs/UnrealBuildTool/Platform/Microsoft/GDK/GDKEditionAudit.cs b/Engine/Source/Programs/UnrealBuildTool/Platform/Microsoft/GDK/GDKEditionAudit.cs
new file mode 100644
--- /dev/null
+++ b/Engine/Source/Programs/UnrealBuildTool/Platform/Microsoft/GDK/GDKEditionAudit.cs
@@ -0,0 +1,196 @@
+// Copyright Epic Games, Inc. All Rights Reserved.
+
+using EpicGames.Core;
+using Microsoft.Extensions.Logging;
+using System;
+using System.Collections.Generic;
+using System.IO;
+using System.Linq;
+
+namespace UnrealBuildTool
+{
+	/// <summary>
+	/// Helper functionality for optionally auditing the GDK edition on any used lib and dll files
+	/// </summary>
+	public static class GDKEditionAudit
+	{
+		private class Config
+		{
+			/// <summary>
+			/// Whether to verify that any libs the application links to are using the same GDK edition
+			/// </summary>
+			[XmlConfigFile(Category = "GDKPlatform")]
+			public bool bVerifyLibGDKEditions = false;
+
+			/// <summary>
+			/// Whether to verify that any DLLs the application links to are using the same GDK edition
+			/// </summary>
+			[XmlConfigFile(Category = "GDKPlatform")]
+			public bool bVerifyDLLGDKEditions = false;
+
+			public Config()
+			{
+				XmlConfig.ApplyTo(this);
+
+				bVerifyLibGDKEditions |= Environment.GetCommandLineArgs().Contains("-VerifyLibGDKEditions", StringComparer.InvariantCultureIgnoreCase);
+				bVerifyDLLGDKEditions |= Environment.GetCommandLineArgs().Contains("-VerifyDLLGDKEditions", StringComparer.InvariantCultureIgnoreCase);
+			}
+		};
+
+		/// <summary>
+		/// Submission validator xml report can suggest multiple GDK editions for a single executable - the additional ones come from the libs that are linked to.
+		/// This function checks all the libs we are planning to link against and warns if they were built against a different GDK
+		/// </summary>
+		internal static void VerifyLinkEnvironment(VCEnvironment EnvVars, LinkEnvironment LinkEnvironment, bool bBuildImportLibraryOnly, ILogger Logger)
+		{
+			Config Config = new();
+			if (LinkEnvironment.bIsBuildingLibrary || bBuildImportLibraryOnly || !Config.bVerifyLibGDKEditions)
+			{
+				return;
+			}
+
+			Logger.LogInformation("Verifying GDK editions from linked libs...");
+
+			// locate dumpbin
+			string DumpBinPath = DirectoryReference.Combine(EnvVars.CompilerDir, "bin", "Hostx64", "x64", "dumpbin.exe").ToString();
+			if (!File.Exists(DumpBinPath))
+			{
+				Logger.LogWarning("Cannot verify the GDK edition of any libs because dumpbin cannot be found at {DumpBinPath}", DumpBinPath);
+				return;
+			}
+
+			// collect as many libraries as we can
+			List<FileReference> AllLibraries = new();
+			foreach (FileReference Library in LinkEnvironment.Libraries)
+			{
+				if (FileReference.Exists(Library))
+				{
+					AllLibraries.Add(Library);
+				}
+			}
+			foreach (string SystemLibrary in LinkEnvironment.SystemLibraries)
+			{
+				FileReference Library = new(SystemLibrary);
+				if (FileReference.Exists(Library))
+				{
+					AllLibraries.Add(Library);
+				}
+				else
+				{
+					foreach (DirectoryReference SystemLibraryDir in LinkEnvironment.SystemLibraryPaths)
+					{
+						Library = FileReference.Combine(SystemLibraryDir, SystemLibrary);
+						if (FileReference.Exists(Library))
+						{
+							AllLibraries.Add(Library);
+							break;
+						}
+					}
+				}
+			}
+
+			// check the GDK edition for all libraries and log out any libraries that have a different GDK edition (if they have one)
+			string? CurrentGDKEdition = UEBuildPlatformSDK.GetSDKForPlatform(LinkEnvironment.Platform.ToString())!.GetInstalledVersion();
+			foreach (FileReference Library in AllLibraries)
+			{
+				// skip libs that should be excluded
+				if (LinkEnvironment.ExcludedLibraries.Contains(Library.GetFileNameWithoutExtension()))
+				{
+					continue;
+				}
+
+				// run DUMPBIN /SYMBOLS <lib> and parse the output for a _XBLD_GRDK_EDITION_* declaration
+				string StdOutString = Utils.RunLocalProcessAndReturnStdOut(DumpBinPath, String.Format("/symbols \"{0}\"", Library.FullName), null, out int ExitCode);
+				if (ExitCode != 0)
+				{
+					Logger.LogError("{Error}", StdOutString);
+					continue;
+				}
+				string[] Results = StdOutString.Split(new char[] { '\r', '\n' }, StringSplitOptions.RemoveEmptyEntries);
+				string? GDKEdition = System.Linq.Enumerable.FirstOrDefault(Results, s => s.Contains("_XBLD_GRDK_EDITION_", StringComparison.Ordinal));
+				if (GDKEdition != null)
+				{
+					GDKEdition = GDKEdition.Substring(GDKEdition.Length - 6, 6); //take last 6 characters   _XBLD_GRDK_EDITION_YYMMQQ -> YYMMQQ
+					if (GDKEdition != CurrentGDKEdition)
+					{
+						Logger.LogWarning("Library {LibraryFullName} was built against GDK edition {GDKEdition}. The current GDK edition is {CurrentGDKEdition}.", Library.FullName, GDKEdition, CurrentGDKEdition);
+					}
+				}
+			}
+		}
+
+		/// <summary>
+		/// Optionally verify that any DLLs the application links to are using the same GDK edition
+		/// </summary>
+		internal static void VerifyRuntimeDependencies(VCEnvironment EnvVars, UnrealTargetPlatform Platform, IEnumerable<FileReference> RuntimeFileDependencies, ILogger Logger)
+		{
+			Config Config = new();
+			if (!Config.bVerifyDLLGDKEditions)
+			{
+				return;
+			}
+			Logger.LogInformation("Verifying GDK editions from DLL dependencies...");
+
+			// locate dumpbin
+			string DumpBinPath = DirectoryReference.Combine(EnvVars.CompilerDir, "bin", "Hostx64", "x64", "dumpbin.exe").ToString();
+			if (!File.Exists(DumpBinPath))
+			{
+				Logger.LogWarning("Cannot verify the GDK edition of any DLLs because dumpbin cannot be found at {DumpBinPath}", DumpBinPath);
+				return;
+			}
+
+			// iterate all DLLs
+			string? CurrentGDKEdition = UEBuildPlatformSDK.GetSDKForPlatform(Platform.ToString())!.GetInstalledVersion();
+			foreach (FileReference FileDependency in RuntimeFileDependencies)
+			{
+				// skip anything that's not a DLL
+				if (!FileDependency.HasExtension("dll"))
+				{
+					continue;
+				}
+
+				try
+				{
+					// run DUMPBIN /SECTION:.xbld /RAWDATA:1,10000 to get the .xbld section contents. The 10000 hopefully means the entire block will be printed one one line so it's easier to parse
+					string StdOutString = Utils.RunLocalProcessAndReturnStdOut(DumpBinPath, String.Format($"/NOLOGO /SECTION:.xbld /RAWDATA:1,10000 \"{FileDependency.FullName}\""), null, out int ExitCode);
+					if (ExitCode != 0)
+					{
+						Logger.LogError("{Error}", StdOutString);
+						continue;
+					}
+					string[] Results = StdOutString.Split(new char[] { '\r', '\n' }, StringSplitOptions.RemoveEmptyEntries);
+
+					// parse the result.
+					int StartIndex = Results.FindIndex(s => s.StartsWith("RAW DATA #", StringComparison.Ordinal));                                                        // Step 1: find the start of the raw data
+					if (StartIndex != -1 && StartIndex < Results.Length)
+					{
+						string[] RawData = Results[StartIndex + 1].Split(new char[] { ':' }, StringSplitOptions.RemoveEmptyEntries);                                        // Step 2: raw data is formatted <base address>:<hex data>  - separate it
+						if (RawData.Length == 2)
+						{
+							byte[] SectionBytes = RawData[1].Split(new char[] { ' ' }, StringSplitOptions.RemoveEmptyEntries).Select(s => Convert.ToByte(s, 16)).ToArray(); // Step 3: get each byte of the hex data
+							string[] XBLDStrings = System.Text.Encoding.ASCII.GetString(SectionBytes).Split(new char[] { '\0' }, StringSplitOptions.RemoveEmptyEntries);       // Step 4: convert the bytes into string table. The bytes contain a series of null-terminated strings
+
+							// find the string we're interested in
+							List<string> FoundGDKEditions = new();
+							foreach (string XBLDString in XBLDStrings)
+							{
+								if (XBLDString.StartsWith("_xbld_edition_sdktype=", StringComparison.Ordinal))
+								{
+									string GDKEdition = XBLDString.Substring(XBLDString.Length - 6, 6); // check last 6 characters - YYMMQQ
+									if (GDKEdition != CurrentGDKEdition && !FoundGDKEditions.Contains(GDKEdition))
+									{
+										FoundGDKEditions.Add(GDKEdition);
+										Logger.LogWarning("Runtime dependency {FileDependency} was built against GDK edition {GDKEdition}. The current GDK edition is {CurrentGDKEdition}", FileDependency, GDKEdition, CurrentGDKEdition);
+									}
+								}
+							}
+						}
+					}
+				}
+				catch (Exception)
+				{
+				}
+			}
+		}
+	}
+}
diff --git a/Engine/Source/Programs/UnrealBuildTool/Platform/Microsoft/GDK/GDKGameConfigGenerator.cs b/Engine/Source/Programs/UnrealBuildTool/Platform/Microsoft/GDK/GDKGameConfigGenerator.cs
new file mode 100644
--- /dev/null
+++ b/Engine/Source/Programs/UnrealBuildTool/Platform/Microsoft/GDK/GDKGameConfigGenerator.cs
@@ -0,0 +1,902 @@
+// Copyright Epic Games, Inc. All Rights Reserved.
+
+using EpicGames.Core;
+using Microsoft.Extensions.Logging;
+using System;
+using System.Collections.Generic;
+using System.IO;
+using System.Linq;
+using System.Threading.Tasks;
+using System.Xml.Linq;
+using UnrealBuildBase;
+
+namespace UnrealBuildTool.GDK
+{
+	/// <summary>
+	/// GDK-specific appx manifest generator (MicrosoftGame.config)
+	/// </summary>
+	public abstract class GDKGameConfigGenerator : AppXManifestGeneratorBase
+	{
+
+		/// Platform-specific target device family
+		protected virtual string? TargetDeviceFamily => null;
+
+		/// Platform-specific architecture
+		protected virtual string DefaultArchitecture => "neutral";
+
+		/// DLC file to use, if any
+		protected FileReference? DLCFile;
+
+		/// Cached properties for the current DLC
+		protected Dictionary<string, string>? DLCConfig;
+
+		/// Helper to check if we are dealing with DLC
+		protected bool IsDLC => (DLCFile != null);
+
+		/// Whether we have a bootstrap exe (Windows only) - affects dependencies
+		protected bool NoBootstrapExe = false;
+
+		/// Curent GDKEdition
+		protected static int GDKEdition => GDKExports.GetGDKVersionNumber() ?? 0;
+
+		/// <summary>
+		/// Create a manifest generator for the given platform variant.
+		/// </summary>
+		protected GDKGameConfigGenerator(UnrealTargetPlatform InPlatform, ILogger InLogger)
+			: base(InPlatform, InLogger)
+		{
+		}
+
+		/// <summary>
+		/// Create a manifest for the given platform and return the list of modified files
+		/// </summary>
+		public static List<string>? Create(UnrealTargetPlatform Platform, ILogger Logger, DirectoryReference InOutputPath, string? InTargetName, FileReference? InProjectFile, Dictionary<UnrealTargetConfiguration, string> InExecutablePairs, bool bDeleteIntermediates = false, FileReference? DLCFile = null, string CustomConfig = "", bool NoBootstrapExe = true)
+		{
+			GDKGameConfigGenerator Generator = GDKGameConfigGeneratorFactory.CreateGenerator(Platform, Logger);
+			Generator.DLCFile = DLCFile;
+			Generator.CustomConfig = CustomConfig;
+			Generator.NoBootstrapExe = NoBootstrapExe;
+
+			List<string>? ModifiedFiles = Generator.CreateManifest("MicrosoftGame.config", InOutputPath, InTargetName, InProjectFile, InExecutablePairs);
+
+			return ModifiedFiles;
+		}
+
+		/// <summary>
+		/// Get the Game element
+		/// </summary>
+		protected virtual XElement GetGame(Dictionary<UnrealTargetConfiguration, string> InExecutablePairs, out string IdentityName)
+		{
+			int ManifestVersion = GetManifestVersion();
+
+			XElement Game = new(XName.Get("Game"),
+				new XAttribute("configVersion", ManifestVersion.ToString()),
+				GetIdentity(out IdentityName),
+				GetShellVisuals(),
+				GetResources()
+				);
+
+			if (!IsDLC)
+			{
+				Game.Add(
+					GetExecutableList(InExecutablePairs),
+					GetProtocolList(InExecutablePairs)
+				);
+
+				string? TitleId = GetConfigString("TitleId", null, null);
+				if (TitleId != null && TitleId.Length > 0)
+				{
+					Game.Add(new XElement(XName.Get("TitleId"), TitleId));
+				}
+
+				string? MSAAppId = GetConfigString("MSAAppId", null, null);
+				if (MSAAppId != null && MSAAppId.Length > 0)
+				{
+					Game.Add(new XElement(XName.Get("MSAAppId"), MSAAppId));
+				}
+
+				string? StoreId = GetConfigString("StoreId", null, null);
+				if (StoreId != null && StoreId.Length > 0)
+				{
+					Game.Add(new XElement(XName.Get("StoreId"), StoreId));
+				}
+
+				// The latest guidance from Microsoft says this flag should always be included in the manifest even if it's false, and most titles should keep it false.
+				bool bRequiresXboxLive = GetConfigBool("bRequireXboxLive", null);
+				Game.Add(new XElement(XName.Get("RequiresXboxLive"), bRequiresXboxLive));
+
+				bool bUseSimplifiedUserModel = GetConfigBool("bUseSimplifiedUserModel", null);
+				if (bUseSimplifiedUserModel)
+				{
+					Game.Add(new XElement(XName.Get("AdvancedUserModel"), "false"));
+				}
+			}
+			else
+			{
+				if (DLCConfig!.TryGetValue("StoreId", out string? StoreId) && !String.IsNullOrEmpty(StoreId))
+				{
+					Game.Add(new XElement(XName.Get("StoreId"), StoreId));
+				}
+
+				string? OwnerStoreId = GetConfigString("StoreId", null, null);
+				if (OwnerStoreId != null && OwnerStoreId.Length > 0)
+				{
+					Game.Add(new XElement(XName.Get("AllowedProducts"),
+						new XElement(XName.Get("AllowedProduct"), OwnerStoreId)));
+				}
+
+				Game.Add(new XElement(XName.Get("TargetDeviceFamilyForDLC"), TargetDeviceFamily));
+			}
+
+			return Game;
+		}
+
+		/// <summary>
+		/// Get the executable list element
+		/// </summary>
+		protected XElement GetExecutableList(Dictionary<UnrealTargetConfiguration, string> InExecutablePairs)
+		{
+			XElement ExecutableList = new(XName.Get("ExecutableList"));
+
+			if (InExecutablePairs.Count < 1)
+			{
+				return ExecutableList;
+			}
+
+			// Need exactly one exe not marked as dev only in manifest.  Most optimized exe gets non-dev.  Dependent on UnrealTargetConfiguration being ordered by optimization level.
+			UnrealTargetConfiguration MaxTargetConfiguration = InExecutablePairs.Keys.Max();
+			foreach (KeyValuePair<UnrealTargetConfiguration, string> ExecutablePair in InExecutablePairs)
+			{
+				XElement Executable = GetExecutable(ExecutablePair.Key, ExecutablePair.Value, ExecutablePair.Key != MaxTargetConfiguration);
+				ExecutableList.Add(Executable);
+			}
+
+			return ExecutableList;
+		}
+
+		private string GetAppId(UnrealTargetConfiguration InTargetConfig)
+		{
+			bool bIsDevBuild = (InTargetConfig != UnrealTargetConfiguration.Shipping);
+			bool bSkipAppIdSuffixForShippingConfig = GetConfigBool("bSkipAppIdSuffixForShippingConfig", null, false);
+			bool bIncludeConfigSuffix = bIsDevBuild || !bSkipAppIdSuffixForShippingConfig;
+
+			string? PackageBaseName = ValidateProjectBaseName(ReadIniString("ProjectName", IniSection_GeneralProjectSettings, "UEGame"));
+			string AppId = "App" + PackageBaseName + (bIncludeConfigSuffix ? InTargetConfig.ToString() : String.Empty);
+
+			return AppId;
+		}
+
+		/// <summary>
+		/// Get the executable element for the given configuration
+		/// </summary>
+		protected XElement GetExecutable(UnrealTargetConfiguration InTargetConfig, string InExecutablePath, bool IsDevOnly)
+		{
+			bool bIsDevBuild = (InTargetConfig != UnrealTargetConfiguration.Shipping);
+
+			string RelativeExePath;
+			if (Path.IsPathRooted(InExecutablePath))
+			{
+				RelativeExePath = Utils.MakePathRelativeTo(InExecutablePath, Path.Combine(InExecutablePath, "../../../.."));
+			}
+			else
+			{
+				RelativeExePath = InExecutablePath;
+			}
+
+			string AppId = GetAppId(InTargetConfig);
+
+			XElement Executable = new(XName.Get("Executable"),
+				new XAttribute("Name", RelativeExePath),
+				new XAttribute("Id", AppId));
+
+			string DisplayNameDefault = ReadIniString(IniSection_GeneralProjectSettings, "ProjectName", GetDefaultDisplayName());
+			if (bIsDevBuild)
+			{
+				string DisplayNameSuffix = " (" + InTargetConfig.ToString() + ")";
+				string DisplayNameResource = AppXResources!.AddResourceString("ExeDisplayName" + InTargetConfig.ToString(), "ApplicationDisplayName", DisplayNameDefault, DisplayNameSuffix);
+				Executable.Add(new XAttribute("OverrideDisplayName", DisplayNameResource));
+			}
+			else
+			{
+				string DisplayNameResource = AppXResources!.AddResourceString("AppDisplayName", "ApplicationDisplayName", DisplayNameDefault);
+				Executable.Add(new XAttribute("OverrideDisplayName", DisplayNameResource));
+			}
+
+			if (TargetDeviceFamily != null)
+			{
+				Executable.Add(new XAttribute("TargetDeviceFamily", TargetDeviceFamily));
+			}
+
+			if (IsDevOnly)
+			{
+				Executable.Add(new XAttribute("IsDevOnly", "true"));
+			}
+
+			return Executable;
+		}
+
+		/// <summary>
+		/// Get the protocol list element
+		/// </summary>
+		protected XElement? GetProtocolList(Dictionary<UnrealTargetConfiguration, string> InExecutablePairs)
+		{
+			XElement? ProtocolList = null;
+
+			// only supported in March GDK onwards
+			if (GDKEdition >= 230300)
+			{
+				// if there's multiple executables in the build, need to specify which one to associate with the protocol. it must use back-slashes
+				string? DefaultExecutable = null;
+				if (InExecutablePairs.Count > 1)
+				{
+					UnrealTargetConfiguration MaxTargetConfiguration = InExecutablePairs.Keys.Max(); // select the most optimized exe
+					DefaultExecutable = InExecutablePairs[MaxTargetConfiguration].Replace('/', '\\');
+				}
+
+				List<string> ProtocolNames = new();
+
+				// read protocols
+				if (EngineIni!.GetArray(IniSection_PlatformTargetSettings, "UriActivationProtocols", out List<string>? UriActivationProtocols) &&
+					UriActivationProtocols.Count > 0)
+				{
+					foreach (string RawUriActivationProtocol in UriActivationProtocols)
+					{
+						if (ConfigHierarchy.TryParse(RawUriActivationProtocol, out Dictionary<string, string>? UriActivationProtocol) && UriActivationProtocol.ContainsKey("RegisteredProtocol"))
+						{
+#pragma warning disable CA1308 // Normalize strings to uppercase - protocol must be lowercase
+							string RegisteredProtocol = UriActivationProtocol["RegisteredProtocol"]
+								.Split(':')[0]      // ignore any suffix in case someone has added a :// etc.
+								.ToLowerInvariant()
+							;
+#pragma warning restore CA1308
+							ProtocolNames.Add(RegisteredProtocol);
+						}
+					}
+				}
+
+				// build the protocol list element
+				if (ProtocolNames.Count > 0)
+				{
+					ProtocolList = new XElement(XName.Get("ProtocolList"));
+					foreach (string ProtocolName in ProtocolNames)
+					{
+						XElement Protocol = new(XName.Get("Protocol"),
+							new XAttribute(XName.Get("Name"), ProtocolName)
+						);
+						if (DefaultExecutable != null)
+						{
+							Protocol.Add(new XAttribute(XName.Get("Executable"), DefaultExecutable));
+						}
+
+						ProtocolList.Add(Protocol);
+					}
+				}
+			}
+
+			return ProtocolList;
+		}
+
+		
+
+		/// <summary>
+		/// Get the shell visuals element
+		/// </summary>
+		protected XElement GetShellVisuals()
+		{
+#pragma warning disable CA1308 // Normalize strings to uppercase - external format may be case-sensitive
+			string ForegroundTextValue = GetConfigString("ApplicationForegroundColor", null, "dark").ToLowerInvariant();
+			string BackgroundColorValue = GetConfigColor("ApplicationBackgroundColor", "#000040").ToLowerInvariant();
+#pragma warning restore CA1308
+
+			string PublisherDisplayName = GetConfigString("PublisherDisplayName", "CompanyName", "NoPublisher");
+			string DefaultDisplayName = GetDefaultDisplayName(); //NB. this cannot be localized (Partner Center refuses the package otherwise)
+
+			string DescriptionDefault = ReadIniString(IniSection_GeneralProjectSettings, "Description", "");
+			string DescriptionResource = AppXResources!.AddResourceString("AppDescription", "ApplicationDescription", DescriptionDefault);
+
+			if (!AppXResources!.AddResourceBinaryFileReference("StoreLogo.png"))
+			{
+				Logger.LogWarning("Unable to stage Store logo.");
+			}
+			if (!AppXResources!.AddResourceBinaryFileReference("Logo.png"))
+			{
+				Logger.LogError("Unable to stage application logo.");
+			}
+			if (!AppXResources!.AddResourceBinaryFileReference("SmallLogo.png"))
+			{
+				Logger.LogError("Unable to stage application small logo.");
+			}
+			if (!AppXResources!.AddResourceBinaryFileReference("SplashScreen.png"))
+			{
+				Logger.LogError("Unable to stage splash screen image.");
+			}
+			if (!AppXResources!.AddResourceBinaryFileReference("Square480x480Logo.png"))
+			{
+				Logger.LogError("Unable to stage application large logo.");
+			}
+
+			XElement ShellVisuals = new(XName.Get("ShellVisuals"),
+				new XAttribute(XName.Get("DefaultDisplayName"), DefaultDisplayName),
+				new XAttribute(XName.Get("PublisherDisplayName"), PublisherDisplayName),
+				new XAttribute(XName.Get("StoreLogo"), BuildResourceSubPath + "\\StoreLogo.png"),
+				new XAttribute(XName.Get("Square150x150Logo"), BuildResourceSubPath + "\\Logo.png"),
+				new XAttribute(XName.Get("Square44x44Logo"), BuildResourceSubPath + "\\SmallLogo.png"),
+				new XAttribute(XName.Get("Square480x480Logo"), BuildResourceSubPath + "\\Square480x480Logo.png"),
+				new XAttribute(XName.Get("Description"), DescriptionResource),
+				new XAttribute(XName.Get("ForegroundText"), ForegroundTextValue),
+				new XAttribute(XName.Get("BackgroundColor"), BackgroundColorValue),
+				new XAttribute(XName.Get("SplashScreenImage"), BuildResourceSubPath + "\\SplashScreen.png")
+			);
+
+			return ShellVisuals;
+		}
+
+		/// <summary>
+		/// Get the default display name for the project
+		/// </summary>
+		protected string GetDefaultDisplayName()
+		{
+			if (IsDLC)
+			{
+				if (DLCConfig!.TryGetValue("DefaultDisplayName", out string? DefaultDisplayName) && !String.IsNullOrEmpty(DefaultDisplayName))
+				{
+					return DefaultDisplayName;
+				}
+
+				Logger.LogWarning("Default display name for {DLC} is not specified", DLCFile!.GetFileNameWithoutExtension());
+				return DLCFile!.GetFileNameWithoutAnyExtensions();
+			}
+			else
+			{
+				string DefaultName = (ProjectFile != null) ? ProjectFile.GetFileNameWithoutAnyExtensions() : (TargetName ?? "UnrealGame");
+
+				string DefaultDisplayName = GetConfigString("DefaultDisplayName", "ProjectName", DefaultName);
+
+				return DefaultDisplayName;
+			}
+		}
+
+		/// <summary>
+		/// Gets the package identity name string for the main package, disregarding any current DLC
+		/// </summary>
+		/// <returns></returns>
+		protected string GetMainPackageIdentityPackageName()
+		{
+			return base.GetIdentityPackageName();
+		}
+
+		/// <summary>
+		/// Get the package identity name string
+		/// </summary>
+		protected override string GetIdentityPackageName()
+		{
+			if (IsDLC)
+			{
+				if (DLCConfig!.TryGetValue("PackageName", out string? PackageName) && !String.IsNullOrEmpty(PackageName))
+				{
+					return PackageName;
+				}
+
+				Logger.LogWarning("Package name for {DLC} is not specified", DLCFile!.GetFileNameWithoutExtension());
+				return DLCFile!.GetFileNameWithoutAnyExtensions();
+			}
+
+			return base.GetIdentityPackageName();
+		}
+
+		/// <summary>
+		/// Get the package version string
+		/// </summary>
+		protected override string? GetIdentityVersionNumber()
+		{
+			if (IsDLC)
+			{
+				if (DLCConfig!.TryGetValue("PackageVersion", out string? PackageVersion) && !String.IsNullOrEmpty(PackageVersion))
+				{
+					// If specified in the project settings attempt to retrieve the current build number and increment the version number by that amount, accounting for overflows
+					if (EngineIni!.GetBool(IniSection_PlatformTargetSettings, "bIncludeEngineVersionInPackageVersion", out bool bIncludeEngineVersionInPackageVersion) && bIncludeEngineVersionInPackageVersion)
+					{
+						PackageVersion = IncludeBuildVersionInPackageVersion(PackageVersion);
+					}
+
+					return PackageVersion;
+				}
+			}
+
+			return base.GetIdentityVersionNumber();
+		}
+
+		/// <summary>
+		/// Manifest version. 0 = original, 1 = March 2022 GDK and above. 
+		/// Note: Version 0 has been deprecated as of October 2023 GDK
+		/// </summary>
+		protected int GetManifestVersion()
+		{
+			// As of October 2023 GDK, configVersion=0 has been deprecated. As a result
+			// the minimum expected manifest version is 1. 
+			const int MinManifestVersion = 1;
+
+			int Result = MinManifestVersion; 
+
+			// If developers need to use a deprecated manifest version, they can use this config
+			// key to override it. 
+			string? ManifestVersionDeprecated = GetConfigString("ManifestVersionDeprecated", null);
+			if (!String.IsNullOrEmpty(ManifestVersionDeprecated))
+			{
+				ManifestVersionDeprecated = ManifestVersionDeprecated.Replace("Version", "", StringComparison.Ordinal);
+				if (Int32.TryParse(ManifestVersionDeprecated, out int ResultOverride))
+				{
+					Result = ResultOverride;
+				}
+				else
+				{
+					Logger.LogError("Could not parse ManifestVersionDeprecated value '{ManifestVersionDeprecated}'", ManifestVersionDeprecated);
+				}
+			}
+
+			return Result;
+		}
+
+		/// <summary>
+		/// Perform any additional initialization once all parameters and configuration are ready
+		/// </summary>
+		protected override void PostConfigurationInit()
+		{
+			base.PostConfigurationInit();
+
+			// find and cache the configuration for the selected DLC
+			DLCConfig = null;
+			if (DLCFile != null)
+			{
+				string DLCName = DLCFile!.GetFileNameWithoutExtension();
+
+				if (EngineIni!.GetArray(IniSection_PlatformTargetSettings, "DLCPlugins", out List<string>? DLCPlugins))
+				{
+					foreach (string DLCPlugin in DLCPlugins)
+					{
+						if (ConfigHierarchy.TryParse(DLCPlugin, out Dictionary<string, string>? Properties)
+							&& Properties.TryGetValue("DLCName", out string? ThisDLCName)
+							&& DLCName == ThisDLCName)
+						{
+							// we have found the configuration for this DLC plugin
+							DLCConfig = Properties;
+							break;
+						}
+					}
+				}
+
+				// create a dummy configuration if we didn't find one
+				if (DLCConfig == null)
+				{
+					string DummyStoreId = String.Concat("0123456789bcdfghjklmnpqrstvwxzBCDFGHJKLMNPQRSTVWXZ".AsEnumerable().OrderBy(x => Guid.NewGuid()).Take(12)); // 12 random items from valid character list (real one can have duplicate chars, but this is just a dummy value)
+					Logger.LogWarning("No DLC configuration for {DLCName}. Creating temporary defaults, including StoreId={DummyStoreId}", DLCName, DummyStoreId);
+					DLCConfig = new Dictionary<string, string>()
+					{
+						{ "StoreId",             DummyStoreId },
+						{ "DefaultDisplayName",  DLCName },
+						{ "PackageName",         DLCName },
+						{ "PackageVersion",      "1.0.0.0" },
+					};
+				}
+			}
+		}
+
+		/// <summary>
+		/// Create all the localization data. Returns whether there is any localization data
+		/// </summary>
+		protected override bool BuildLocalizationData()
+		{
+			bool bResult = base.BuildLocalizationData();
+
+			// replace base strings with DLC strings
+			if (DLCConfig != null)
+			{
+				// reset per-culture strings and make sure the default culture entry exists
+				AppXResources!.ClearStrings();
+
+				// add all default strings
+				if (DLCConfig.TryGetValue("DefaultStringResources", out string? DefaultStringResources) && ConfigHierarchy.TryParse(DefaultStringResources, out Dictionary<string, string>? DefaultStrings))
+				{
+					AppXResources!.AddDefaultStrings(DefaultStrings);
+				}
+
+				// add per culture strings
+				if (DLCConfig.TryGetValue("PerCultureResources", out string? PerCultureResources) && ConfigHierarchy.TryParse(PerCultureResources, out string[]? PerCultureResourcesArray))
+				{
+					foreach (string CultureResources in PerCultureResourcesArray)
+					{
+						if (!ConfigHierarchy.TryParse(CultureResources, out Dictionary<string, string>? CultureProperties)
+							|| !CultureProperties.ContainsKey("CultureStringResources")
+							|| !CultureProperties.ContainsKey("StageId"))
+						{
+							Logger.LogWarning("Invalid per-culture resource value: {Culture}", CultureResources);
+							continue;
+						}
+
+						string StageId = CultureProperties["StageId"];
+						if (String.IsNullOrEmpty(StageId))
+						{
+							Logger.LogWarning("Missing StageId value: {Culture}", CultureResources);
+							continue;
+						}
+
+						if (!UEStageIdToAppXCultureId.TryGetValue(StageId, out string? CultureId))
+						{
+							Logger.LogWarning("Unknown Culture {ID} in {Culture}. Check main PerCultureResources", StageId, CultureResources);
+							continue;
+						}
+
+						// read culture strings
+						if (!ConfigHierarchy.TryParse(CultureProperties["CultureStringResources"], out Dictionary<string, string>? CultureStringResources))
+						{
+							Logger.LogError("Invalid culture string resources: \"{Culture}\". Unable to add resource entry.", CultureResources);
+							continue;
+						}
+						AppXResources!.AddCultureStrings(CultureId, CultureStringResources);
+					}
+				}
+			}
+
+			// add the feature display name
+			if (EngineIni!.GetArray(IniSection_PlatformTargetSettings, "IntelligentDeliveryFeatures", out List<string>? Features) && Features != null && Features.Count > 0)
+			{
+				foreach (string Feature in Features)
+				{
+					if (ConfigHierarchy.TryParse(Feature, out Dictionary<string, string>? FeatureProperties))
+					{
+						// verify feature id
+						if (!FeatureProperties.TryGetValue("Id", out string? Id) || Id.Length == 0)
+						{
+							throw new InvalidDataException("Missing or invalid Id in feature");
+						}
+						string ResourceEntryName = "Feature_" + Id;
+
+						// read the default display name and skip if it isn't set
+						if (FeatureProperties.TryGetValue("DefaultDisplayName", out string? DefaultDisplayName) && !String.IsNullOrEmpty(DefaultDisplayName))
+						{
+							AppXResources!.AddDefaultString(ResourceEntryName, DefaultDisplayName);
+						}
+						else
+						{
+							continue;
+						}
+
+						// read the per-culture resources for this feature
+						if (FeatureProperties.TryGetValue("PerCultureResources", out string? PerCultureResourcesProperty))
+						{
+							if (!ConfigHierarchy.TryParse(PerCultureResourcesProperty, out string[]? PerCultureResources))
+							{
+								Logger.LogWarning("Invalid per-culture resources value in feature {Feature} : {Culture}", Id, PerCultureResourcesProperty);
+								continue;
+							}
+
+							// parse all per culture resources
+							foreach (string CultureResources in PerCultureResources)
+							{
+								if (!ConfigHierarchy.TryParse(CultureResources, out Dictionary<string, string>? CultureProperties)
+									|| !CultureProperties.ContainsKey("StageId")
+									|| !CultureProperties.ContainsKey("DisplayName"))
+								{
+									Logger.LogWarning("Invalid per-culture resource value in feature {Feature} : {Culture}", Id, CultureResources);
+									continue;
+								}
+
+								string StageId = CultureProperties["StageId"];
+								if (String.IsNullOrEmpty(StageId))
+								{
+									Logger.LogWarning("Missing StageId value in feature {Feature}: {Culture}", Id, CultureProperties);
+									continue;
+								}
+
+								if (UEStageIdToAppXCultureId.TryGetValue(StageId, out string? CultureId))
+								{
+									string DisplayName = CultureProperties["DisplayName"];
+									if (!String.IsNullOrEmpty(DisplayName))
+									{
+										AppXResources!.AddCultureString(CultureId, ResourceEntryName, DisplayName);
+									}
+								}
+								else
+								{
+									throw new InvalidDataException($"Unknown StageId {StageId} in feature {Id}. Make sure there is an entry in the main PerCultureResources");
+								}
+							}
+						}
+					}
+				}
+			}
+
+			return bResult;
+		}
+
+		/// <summary>
+		/// Return the entire manifest element
+		/// </summary>
+		protected override XElement GetManifest(Dictionary<UnrealTargetConfiguration, string> InExecutablePairs, out string IdentityName)
+		{
+			return GetGame(InExecutablePairs, out IdentityName);
+		}
+
+		/// <summary>
+		/// 
+		/// </summary>
+		protected override void ProcessManifest(Dictionary<UnrealTargetConfiguration, string> InExecutablePairs, string ManifestName, string ManifestTargetPath, string ManifestIntermediatePath)
+		{
+			AddResourcesForPackaging();
+
+			base.ProcessManifest(InExecutablePairs, ManifestName, ManifestTargetPath, ManifestIntermediatePath);
+		}
+
+		/// <summary>
+		/// Register the locations where resource binary files can be found
+		/// </summary>
+		protected override void PrepareResourceBinaryPaths()
+		{
+			// check for DLC images first
+			if (DLCFile != null)
+			{
+				AppXResources!.ProjectBinaryResourceDirectories.Add(DirectoryReference.Combine(DLCFile.Directory, "Build", Platform.ToString(), BuildResourceSubPath));
+				AppXResources!.ProjectBinaryResourceDirectories.Add(DirectoryReference.Combine(DLCFile.Directory, "Platforms", Platform.ToString(), "Build", BuildResourceSubPath));
+			}
+
+			base.PrepareResourceBinaryPaths();
+		}
+
+		/// <summary>
+		/// Perform any platform-specific processing on the manifest before it is saved
+		/// </summary>
+		protected void AddResourcesForPackaging()
+		{
+			// see if features and recipes are being used
+			EngineIni!.GetBool(IniSection_PlatformTargetSettings, "bUseFeaturesAndRecipes", out bool bUseFeaturesAndRecipes);
+			if (!bUseFeaturesAndRecipes)
+			{
+				return;
+			}
+
+			// get the list of all features
+			EngineIni.GetArray(IniSection_PlatformTargetSettings, "IntelligentDeliveryFeatures", out List<string>? Features);
+			if (Features == null || Features.Count == 0)
+			{
+				return;
+			}
+
+			// parse all features
+			foreach (string Feature in Features)
+			{
+				if (ConfigHierarchy.TryParse(Feature, out Dictionary<string, string>? Properties))
+				{
+					// verify feature id
+					if (!Properties.TryGetValue("Id", out string? Id) || Id.Length == 0)
+					{
+						throw new InvalidDataException("Missing or invalid Id in feature");
+					}
+
+					// copy the feature image. Image file name must match the one in GDKPackageGenerator.GenerateFeaturesAndRecipes
+					// note that localized feature images are not currently supported by the GDK
+					string ImageFilename = "Feature_" + Id + ".png";
+					bool bHasFeatureImage = AppXResources!.DoesDefaultResourceBinaryFileExist(ImageFilename, AllowEngineFallback: false);
+					if (bHasFeatureImage)
+					{
+						if (!AppXResources!.AddResourceBinaryFileReference(ImageFilename, AllowEngineFallback: false))
+						{
+							Logger.LogError("Unable to feature image {Image} for feature {Id}.", ImageFilename, Id);
+						}
+					}
+
+					// add the string entry. note that the prefix must match the one in GDKPackageGenerator.GenerateFeaturesAndRecipes
+					string ResourceEntryName = "Feature_" + Id;
+					if (AppXResources!.HasString(ResourceEntryName))
+					{
+						AppXResources!.AddResourceString(ResourceEntryName, ResourceEntryName, "");
+					}
+				}
+			}
+		}
+
+		#region packaging resource query helper
+
+		/// <summary>
+		/// Helper to find whether there is a packaging resource file available
+		/// </summary>
+		public static bool DoesDefaultResourceBinaryFileExistForPlatform(UnrealTargetPlatform Platform, ILogger Logger, string ImageFilename, bool bAllowEngineFallback = true)
+		{
+			GDKGameConfigGenerator Generator = GDKGameConfigGeneratorFactory.CreateGenerator(Platform,Logger);
+			Generator.AppXResources = new(Logger, Generator.GetMakePriBinaryPath());
+			Generator.PrepareResourceBinaryPaths();
+
+			return Generator.AppXResources!.DoesDefaultResourceBinaryFileExist(ImageFilename, bAllowEngineFallback);
+		}
+
+		#endregion
+
+		#region Package name derivation code
+
+		private void GetManifestElementsForProject(DirectoryReference InProjectDir, string? InTargetName, out string Name, out string? Version, out string Architecture, out string ResourceId, out string Publisher)
+		{
+			// Load up INI settings. We'll use engine settings to retrieve the manifest configuration, but these may reference
+			// values in either game or engine settings, so we'll keep both.
+			GameIni = ConfigCache.ReadHierarchy(ConfigHierarchyType.Game, InProjectDir, Platform, CustomConfig);
+			EngineIni = ConfigCache.ReadHierarchy(ConfigHierarchyType.Engine, InProjectDir, Platform, CustomConfig);
+			TargetName = InTargetName;
+
+			Name = GetIdentityPackageName();
+			Version = GetIdentityVersionNumber();
+			Architecture = DefaultArchitecture;
+			ResourceId = "";
+			Publisher = GetIdentityPublisherName();
+		}
+
+		/// <summary>
+		/// Returns the package full name for the given project
+		/// </summary>
+		public static string GetPackageFullNameFromProject(UnrealTargetPlatform InPlatform, ILogger Logger, DirectoryReference InProjectDir, string? InTargetName, string InCustomConfig)
+		{
+			GDKGameConfigGenerator ManifestGenerator = GDKGameConfigGeneratorFactory.CreateGenerator(InPlatform, Logger);
+			ManifestGenerator.CustomConfig = InCustomConfig;
+
+			ManifestGenerator.GetManifestElementsForProject(InProjectDir, InTargetName, out string Name, out string? Version, out string Architecture, out string ResourceId, out string Publisher);
+
+			return String.Format("{0}_{1}_{2}_{3}_{4}",
+				Name,
+				Version,
+				Architecture,
+				ResourceId,
+				GDKGameConfigInfo.EncodePublisherNameForPackageName(Publisher)
+				);
+		}
+
+		/// <summary>
+		/// Returns the package family name for the given project
+		/// </summary>
+		public static string GetPackageFamilyNameFromProject(UnrealTargetPlatform InPlatform, ILogger Logger, DirectoryReference InProjectDir, string? InTargetName, string InCustomConfig)
+		{
+			GDKGameConfigGenerator ManifestGenerator = GDKGameConfigGeneratorFactory.CreateGenerator(InPlatform, Logger);
+			ManifestGenerator.CustomConfig = InCustomConfig;
+
+			ManifestGenerator.GetManifestElementsForProject(InProjectDir, InTargetName, out string Name, out _, out _, out _, out string Publisher);
+
+			return String.Format("{0}_{1}",
+				Name,
+				GDKGameConfigInfo.EncodePublisherNameForPackageName(Publisher)
+				);
+		}
+
+		/// <summary>
+		/// Returns the AUMID for the given project and configuration
+		/// </summary>
+		public static string GetAUMIDFromProject(UnrealTargetPlatform InPlatform, ILogger Logger, DirectoryReference InProjectDir, UnrealTargetConfiguration InTargetConfig, string? InTargetName, string InCustomConfig)
+		{
+			GDKGameConfigGenerator ManifestGenerator = GDKGameConfigGeneratorFactory.CreateGenerator(InPlatform, Logger);
+			ManifestGenerator.CustomConfig = InCustomConfig;
+
+			ManifestGenerator.GetManifestElementsForProject(InProjectDir, InTargetName, out string Name, out _, out _, out _, out string Publisher);
+
+			string AppId = ManifestGenerator.GetAppId(InTargetConfig);
+
+			return $"{Name}_{GDKGameConfigInfo.EncodePublisherNameForPackageName(Publisher)}!{AppId}";
+		}
+
+		#endregion
+	}
+
+	partial struct GDKGameConfigGeneratorFactory
+	{
+		private static readonly Dictionary<UnrealTargetPlatform, Type> PlatformGenerators = new();
+
+		/// <summary>
+		/// Get the manifest generator for the given platform
+		/// </summary>
+		public static GDKGameConfigGenerator CreateGenerator(UnrealTargetPlatform Platform, ILogger? Logger)
+		{
+			if (!PlatformGenerators.TryGetValue(Platform, out Type? GeneratorType) || GeneratorType == null)
+			{
+				throw new BuildException($"Platform {Platform} does not support GDK manifest generation");
+			}
+
+			GDKGameConfigGenerator? Generator = (GDKGameConfigGenerator?)Activator.CreateInstance(GeneratorType, Logger);
+			if (Generator == null)
+			{
+				throw new BuildException($"Could not instantiate the GDK manifest generator for {Platform}");
+			}
+
+			return Generator!;
+		}
+
+		private static Type RegisterGenerator(UnrealTargetPlatform Platform, Type GeneratorType)
+		{
+			if (!GeneratorType.IsSubclassOf(typeof(GDKGameConfigGenerator)))
+			{
+				throw new BuildException($"{GeneratorType} is not a GDKGameConfigGenerator");
+			}
+
+			PlatformGenerators.Add(Platform, GeneratorType);
+			return GeneratorType;
+		}
+	}
+
+	/// <summary>
+	/// ToolMode to create a MicrosoftGame.config for a single executable pair
+	/// </summary>
+	[ToolMode("GDKGameConfigGenerator", ToolModeOptions.BuildPlatforms)]
+	internal class GDKGameConfigGeneratorMode : ToolMode
+	{
+#pragma warning disable IDE0044 // Make field readonly - these private fields are set by the command line.
+		[CommandLine("-Executable", Required = false)]
+		string? Executable = null;
+
+		[CommandLine("-Configuration", Required = false)]
+		UnrealTargetConfiguration? Configuration = null;
+
+		[CommandLine("-Project", Required = false)]
+		FileReference? ProjectFile = null;
+
+		[CommandLine("-Platform", Required = true)]
+		UnrealTargetPlatform? Platform = null;
+
+		[CommandLine("-Target", Required = true)]
+		string? TargetName = null;
+
+		[CommandLine("-CustomConfig", Required = false)]
+		string CustomConfig = "";
+
+		[CommandLine("-OutputDirectory", Required = true)]
+		DirectoryReference? OutputDirectory = null;
+#pragma warning restore IDE0044
+
+		/// <inheritdoc/>
+		public override Task<int> ExecuteAsync(CommandLineArguments Arguments, ILogger Logger)
+		{
+			Arguments.ApplyTo(this);
+			Arguments.CheckAllArgumentsUsed();
+
+			Dictionary<UnrealTargetConfiguration, string> ExecutablePairs;
+			if (Executable == null && Configuration == null)
+			{
+				ExecutablePairs = [];
+			}
+			else
+			{
+				ExecutablePairs = new() { { (UnrealTargetConfiguration)Configuration!, (string)Executable! } };
+			}
+
+			GDKGameConfigGenerator.Create((UnrealTargetPlatform)Platform!, Logger, OutputDirectory!, TargetName, ProjectFile, ExecutablePairs, CustomConfig: CustomConfig);
+
+			return Task.FromResult(0);
+		}
+
+		/// <summary>
+		/// Helper function to create an action that will run this ToolMode
+		/// </summary>
+		public static Action CreateAction(TargetMakefileBuilder MakefileBuilder, UnrealTargetPlatform Platform, DirectoryReference OutputDirectory, string TargetName, FileReference? ProjectFile, UnrealTargetConfiguration? Configuration = null, string? ExecutableFileName = null, string CustomConfig = "")
+		{
+			Action GenAction = MakefileBuilder.CreateAction(ActionType.PostBuildStep);
+			GenAction.CommandPath = Unreal.DotnetPath;
+			GenAction.CommandArguments = $"\"{Unreal.UnrealBuildToolDllPath}\" -Mode=GDKGameConfigGenerator -Platform={Platform} -Target={TargetName} -OutputDirectory=\"{OutputDirectory}\"";
+			GenAction.WorkingDirectory = Unreal.EngineSourceDirectory;
+			GenAction.ProducedItems.Add(FileItem.GetItemByFileReference(FileReference.Combine(OutputDirectory, "MicrosoftGame.config")));
+			GenAction.ProducedItems.Add(FileItem.GetItemByFileReference(FileReference.Combine(OutputDirectory, "resources.pri")));
+			GenAction.DeleteItems.UnionWith(GenAction.ProducedItems);
+			GenAction.bCanExecuteRemotely = false;
+
+			if (Configuration != null && ExecutableFileName != null)
+			{
+				GenAction.CommandArguments += $" -Executable =\"{ExecutableFileName}\" -Configuration={Configuration}";
+			}
+
+			if (!String.IsNullOrEmpty(CustomConfig))
+			{
+				GenAction.CommandArguments += $" -CustomConfig={CustomConfig}";
+			}
+
+			if (ProjectFile != null)
+			{
+				GenAction.CommandArguments += $" -Project=\"{ProjectFile}\"";
+			}
+
+			MakefileBuilder.AddOutputFiles(GenAction.ProducedItems);
+
+			return GenAction;
+		}
+	}
+}
diff --git a/Engine/Source/Programs/UnrealBuildTool/Platform/Microsoft/GDK/GDKGameConfigUtils.cs b/Engine/Source/Programs/UnrealBuildTool/Platform/Microsoft/GDK/GDKGameConfigUtils.cs
new file mode 100644
--- /dev/null
+++ b/Engine/Source/Programs/UnrealBuildTool/Platform/Microsoft/GDK/GDKGameConfigUtils.cs
@@ -0,0 +1,544 @@
+// Copyright Epic Games, Inc. All Rights Reserved.
+
+using EpicGames.Core;
+using System;
+using System.Collections.Generic;
+using System.IO;
+using System.Linq;
+using System.Text;
+using System.Text.Json;
+using System.Text.RegularExpressions;
+using System.Xml;
+
+namespace UnrealBuildTool.GDK
+{
+	/// <summary>
+	/// Helper class used to manage GDK manifest identities
+	/// </summary>
+	public class GDKGameConfigInfo
+	{
+		/// Regular expression for standard Micrososoft full package name:  Name_Version_ProcessorArchitecture_ResourceId_PublisherId
+		public const string PackageIdentityRegEx = @"([A-Za-z0-9\-\.]+)_(\d+\.\d+.\d+.\d+)_(x86|x64|arm|neutral)_([A-Za-z0-9\-\.]*)_([A-Za-z0-9]+)";
+
+		/// Regular expression for standard Microsoft AUMID: Name_PublisherId!AppName
+		public const string PackageAUMIDRegEx = @"([A-Za-z0-9\-\.]+)_([A-Za-z0-9]+)!(\S+)";
+
+		/// Architecture to use when none is specified
+		public const string DefaultArchitecture = "neutral";
+
+		/// Version to use when none is specified
+		public const string DefaultVersion = "0.0.0.0";
+
+		/// The package identity. e.g. "GameName"
+		public string Identity { get; protected set; }
+
+		/// The package version. e.g. "1.0.0.0"
+		public string Version { get; protected set; }
+
+		/// The package Architecture. e.g. "neutral", "x64" etc.
+		public string Architecture { get; protected set; }
+
+		/// The package resource id. (not used on GDK)
+		public string ResourceId { get; protected set; }
+
+		/// The package publisher id. e.g. "z844xmwvfmpty"
+		public string PublisherId { get; protected set; }
+
+		/// The full Package Name. e.g. "GameName_1.0.0.0_neutral__z844xmwvfmpty"
+		public string PackageFullName => Identity + "_" + Version + "_" + Architecture + "_" + ResourceId + "_" + PublisherId;
+
+		/// The package family name. e.g. "GameName_z844xmwvfmpty"
+		public string PackageFamilyName => Identity + "_" + PublisherId;
+
+		/// only valid if we are created from a MicrosoftGame.config
+		public readonly Dictionary<string, string>? IdToExecuatable;
+
+		/// only valid if we are created from a MicrosoftGame.config or via ParseAppList
+		public List<string>? AUMIDs => InternalAUMIDs;
+		private readonly List<string>? InternalAUMIDs;
+
+		/// <summary>
+		/// Create empty manifest identity
+		/// </summary>
+		public GDKGameConfigInfo()
+		{
+			Identity = "";
+			Version = DefaultVersion;
+			Architecture = DefaultArchitecture;
+			ResourceId = "";
+			PublisherId = "";
+		}
+
+		/// <summary>
+		/// Create a manifest based on a given package full name. e.g. "GameName_1.0.0.0_neutral__z844xmwvfmpty"
+		/// </summary>
+		public GDKGameConfigInfo(string InPackageFullName) : this()
+		{
+			if (!SetFromPackageFullName(InPackageFullName))
+			{
+				throw new Exception($"Invalid Package Full Name : {InPackageFullName}");
+			}
+		}
+
+		/// <summary>
+		/// Create a manifest identity from the given manifest file
+		/// </summary>
+		public GDKGameConfigInfo(FileReference ManifestFile) : this()
+		{
+			try
+			{
+				// read the Identity
+				XmlDocument XmlDoc = new();
+				XmlDoc.Load(ManifestFile.FullName);
+				XmlNodeList IdentityNodes = XmlDoc.GetElementsByTagName("Identity");
+				if (IdentityNodes.Count == 1)
+				{
+					XmlElement XmlIdentity = (XmlElement)IdentityNodes[0]!;
+					SetFromIdentity(XmlIdentity);
+				}
+				else
+				{
+					throw new Exception($"Found more than one Identity node in manifest {ManifestFile.FullName}");
+				}
+
+				// read the executables (NB. MicrosoftGame.config only... old AppXManifest.xml uses "Applications")
+				XmlNodeList ExecutableList = XmlDoc.GetElementsByTagName("ExecutableList");
+				if (ExecutableList.Count > 0)
+				{
+					IdToExecuatable = new Dictionary<string, string>();
+					InternalAUMIDs = new List<string>();
+
+					foreach (XmlNode? Executable in ExecutableList[0]!.ChildNodes)
+					{
+						if (Executable != null)
+						{
+							string Name = Executable.Attributes!["Name"]!.Value;
+							string Id = Executable.Attributes["Id"]!.Value;
+							IdToExecuatable.Add(Id, Name);
+							InternalAUMIDs.Add($"{PackageFamilyName}!{Id}");
+						}
+					}
+				}
+			}
+			catch
+			{
+				throw new Exception($"Unable to parse Xml manifest {ManifestFile.FullName}");
+			}
+		}
+
+		internal GDKGameConfigInfo(JsonElement JsonPkg) : this()
+		{
+			try
+			{
+				if (JsonPkg.TryGetStringProperty("PackageFullName", out string? PackageFullName) || JsonPkg.TryGetStringProperty("FullName", out PackageFullName))
+				{
+					if (SetFromPackageFullName(PackageFullName))
+					{
+						if (JsonPkg.TryGetProperty("Applications", out JsonElement JsonApplications))
+						{
+							IdToExecuatable = new Dictionary<string, string>();
+							InternalAUMIDs = new List<string>();
+
+							foreach (JsonElement JsonApp in JsonApplications.EnumerateArray())
+							{
+								if (JsonApp.TryGetStringProperty("Aumid", out string? AUMID))
+								{
+									InternalAUMIDs.Add(AUMID);
+
+									if (JsonApp.TryGetStringProperty("Executable", out string? Executable) && JsonApp.TryGetStringProperty("Id", out string? Id))
+									{
+										IdToExecuatable.Add(Id, Executable);
+									}
+								}
+							}
+						}
+					}
+				}
+			}
+			catch (Exception)
+			{
+			}
+		}
+
+		internal bool SetFromPackageFullName(string InPackageFullName)
+		{
+			Match M = Regex.Match(InPackageFullName, PackageIdentityRegEx, RegexOptions.IgnoreCase);
+			if (M.Success)
+			{
+				Identity = M.Groups[1].Value;
+				Version = M.Groups[2].Value;
+				Architecture = M.Groups[3].Value;
+				ResourceId = M.Groups[4].Value;
+				PublisherId = M.Groups[5].Value;
+			}
+
+			return M.Success;
+		}
+
+		internal bool SetFromIdentity(XmlElement XmlIdentity)
+		{
+			Identity = XmlIdentity.HasAttribute("Name") ? XmlIdentity.GetAttribute("Name") : "";
+			Version = XmlIdentity.HasAttribute("Version") ? XmlIdentity.GetAttribute("Version") : "";
+			ResourceId = XmlIdentity.HasAttribute("ResourceId") ? XmlIdentity.GetAttribute("ResourceId") : "";
+
+			Architecture = DefaultArchitecture;
+			if (XmlIdentity.HasAttribute("ProcessorArchitecture"))
+			{
+				Architecture = XmlIdentity.GetAttribute("ProcessorArchitecture");
+			}
+			else
+			{
+				XmlNodeList XmlProcessorArchitectureNodes = XmlIdentity.OwnerDocument.GetElementsByTagName("ProcessorArchitecture");
+				if (XmlProcessorArchitectureNodes.Count == 1)
+				{
+					Architecture = XmlProcessorArchitectureNodes[0]!.InnerText;
+				}
+			}
+
+			string Publisher = XmlIdentity.HasAttribute("Publisher") ? XmlIdentity.GetAttribute("Publisher") : "";
+			PublisherId = EncodePublisherNameForPackageName(Publisher);
+
+			return true;
+		}
+
+		/// <summary>
+		/// Query whether there are executables available in this manifest
+		/// </summary>
+		public bool HasExecutables => IdToExecuatable != null;
+
+		/// <summary>
+		/// Look up the relative path to the given executable, if it is mentioned
+		/// </summary>
+		/// <returns></returns>
+		public string? GetExecutableRelativePath(string ExecutableName)
+		{
+			if (IdToExecuatable != null)
+			{
+				string ExecutableFileName = Path.GetFileName(ExecutableName);
+				foreach (KeyValuePair<string, string> Pair in IdToExecuatable)
+				{
+					if (Path.GetFileName(Pair.Value).Equals(ExecutableFileName, StringComparison.OrdinalIgnoreCase))
+					{
+						return Pair.Value;
+					}
+				}
+			}
+
+			return null;
+		}
+
+		/// <summary>
+		/// Returns the AUMID for the given manifest and executable file name. e.g. GameName_z844xmwvfmpty!AppGameNameDevelopment
+		/// </summary>
+		public string? GetAUMIDForExecutable(string ExecutableName)
+		{
+			if (IdToExecuatable != null)
+			{
+				string ExecutableFileName = Path.GetFileName(ExecutableName);
+				foreach (KeyValuePair<string, string> Pair in IdToExecuatable)
+				{
+					if (Path.GetFileName(Pair.Value).Equals(ExecutableFileName, StringComparison.OrdinalIgnoreCase))
+					{
+						return $"{PackageFamilyName}!{Pair.Key}";
+					}
+				}
+			}
+
+			return null;
+		}
+
+		/// <summary>
+		/// Returns the AUMID that matches the given configuration
+		/// </summary>
+		public string GetAUMIDForConfiguration(UnrealTargetConfiguration Configuration)
+		{
+			if (IdToExecuatable != null && AUMIDs != null)
+			{
+				return GetAUMIDForConfiguration(AUMIDs, Configuration);
+			}
+			else
+			{
+				throw new Exception("executable list is not available");
+			}
+		}
+
+		/// <summary>
+		/// Returns the executable that matches the given configuration
+		/// </summary>
+		public string GetExecutableForConfiguration(UnrealTargetConfiguration Configuration)
+		{
+			if (IdToExecuatable == null)
+			{
+				throw new Exception("executable list is not available");
+			}
+			else if (IdToExecuatable.Values.Count == 0)
+			{
+				throw new Exception("no executables");
+			}
+
+			string? Executable = IdToExecuatable.Values.FirstOrDefault(X => X.EndsWith(Configuration.ToString(), StringComparison.Ordinal));
+			// fall back to any executable not tagged as Dev/Test/Shipping
+			Executable ??= IdToExecuatable.Values.FirstOrDefault(X =>
+				!X.EndsWith(@"Development", StringComparison.Ordinal) &&
+				!X.EndsWith(@"Test", StringComparison.Ordinal) &&
+				!X.EndsWith(@"Shipping", StringComparison.Ordinal));
+
+			if (Executable == null)
+			{
+				throw new Exception($"Unable to find executable for configuration {Configuration} from {String.Join(',', IdToExecuatable.Values)}");
+			}
+
+			return Executable!;
+		}
+
+		/// <summary>
+		/// Returns the executable that matches the given AUMID
+		/// </summary>
+		public string? GetExecutableForAUMID(string AUMID)
+		{
+			if (IdToExecuatable != null)
+			{
+				foreach (KeyValuePair<string, string> Pair in IdToExecuatable)
+				{
+					string ThisAUMID = $"{PackageFamilyName}!{Pair.Key}";
+					if (AUMID.Equals(ThisAUMID, StringComparison.OrdinalIgnoreCase))
+					{
+						return Pair.Value;
+					}
+				}
+			}
+
+			return null;
+		}
+
+		/// <summary>
+		/// Standard ToString override
+		/// </summary>
+		public override string ToString()
+		{
+			return PackageFullName;
+		}
+
+		#region Publisher encoding helper functions from Microsoft
+
+		/// <summary>
+		/// Converts a human-readable Publisher Name into the Publisher Id used by MS app manifests
+		/// </summary>
+		public static string EncodePublisherNameForPackageName(string PublisherName)
+		{
+			byte[] PublisherBytes = Encoding.Unicode.GetBytes(PublisherName);
+			byte[] PublisherBytesHash = System.Security.Cryptography.SHA256.HashData(PublisherBytes);
+
+			byte[] FirstBytes = new byte[8];
+			Array.Copy(PublisherBytesHash, 0, FirstBytes, 0, 8);
+			return GetBase32CharsForPackageName(FirstBytes);
+		}
+
+		private static string GetBase32CharsForPackageName(byte[] Bytes)
+		{
+			StringBuilder Result = new();
+
+			for (int ByteIndex = 0; ByteIndex < Bytes.Length; ByteIndex += 5)
+			{
+				byte FirstByte = Bytes[ByteIndex];
+				byte SecondByte = (ByteIndex + 1) < Bytes.Length ? Bytes[ByteIndex + 1] : (byte)0;
+
+				Result.Append(ValueToDigitForPackageName((byte)((FirstByte & 0xF8) >> 3)));
+				Result.Append(ValueToDigitForPackageName((byte)(((FirstByte & 0x07) << 2) | ((SecondByte & 0xC0) >> 6))));
+
+				if (ByteIndex + 1 < Bytes.Length)
+				{
+					byte thirdByte = (ByteIndex + 2) < Bytes.Length ? Bytes[ByteIndex + 2] : (byte)0;
+
+					Result.Append(ValueToDigitForPackageName((byte)((SecondByte & 0x3E) >> 1)));
+					Result.Append(ValueToDigitForPackageName((byte)(((SecondByte & 0x01) << 4) | ((thirdByte & 0xF0) >> 4))));
+
+					if (ByteIndex + 2 < Bytes.Length)
+					{
+						byte fourthByte = (ByteIndex + 3) < Bytes.Length ? Bytes[ByteIndex + 3] : (byte)0;
+
+						Result.Append(ValueToDigitForPackageName((byte)(((thirdByte & 0x0F) << 1) | ((fourthByte & 0x80) >> 7))));
+
+						if (ByteIndex + 3 < Bytes.Length)
+						{
+							byte fifthByte = (ByteIndex + 4) < Bytes.Length ? Bytes[ByteIndex + 4] : (byte)0;
+
+							Result.Append(ValueToDigitForPackageName((byte)((fourthByte & 0x7C) >> 2)));
+							Result.Append(ValueToDigitForPackageName((byte)(((fourthByte & 0x03) << 3) | ((fifthByte & 0xE0) >> 5))));
+
+							if (ByteIndex + 4 < Bytes.Length)
+							{
+								Result.Append(ValueToDigitForPackageName((byte)(fifthByte & 0x1F)));
+							}
+						}
+					}
+				}
+			}
+
+			return Result.ToString();
+		}
+
+		private static char ValueToDigitForPackageName(byte Value)
+		{
+			if (Value >= 0x20)
+			{
+				throw new ArgumentOutOfRangeException(nameof(Value));
+			}
+
+			const string Base32DigitList = "0123456789abcdefghjkmnpqrstvwxyz";
+			return Base32DigitList[Value];
+		}
+
+		#endregion
+
+		/// <summary>
+		/// Returns the full package name from the given manifest. e.g.  GameName_1.0.0.0_neutral__z844xmwvfmpty
+		/// </summary>
+		public static string GetPackageFullNameFromManifest(string PathToManifest)
+		{
+			GDKGameConfigInfo Manifest = new(new FileReference(PathToManifest));
+			return Manifest.PackageFullName;
+		}
+
+		/// <summary>
+		/// Returns the package family name from the given manifest. e.g. GameName_z844xmwvfmpty
+		/// </summary>
+		public static string GetPackageFamilyNameFromManifest(string PathToManifest)
+		{
+			GDKGameConfigInfo Manifest = new(new FileReference(PathToManifest));
+			return Manifest.PackageFamilyName;
+		}
+
+		/// <summary>
+		/// Returns the AUMID for the given manifest and executable file name. e.g. GameName_z844xmwvfmpty!AppGameNameDevelopment
+		/// </summary>
+		public static string GetAUMIDFromManifest(string PathToManifest, string ExecutableName, string? AlternativeExecutableName = null)
+		{
+			GDKGameConfigInfo Manifest = new(new FileReference(PathToManifest));
+			string? AUMID = Manifest.GetAUMIDForExecutable(ExecutableName);
+			if (AUMID == null && AlternativeExecutableName != null)
+			{
+				AUMID = Manifest.GetAUMIDForExecutable(AlternativeExecutableName);
+			}
+			if (AUMID == null)
+			{
+				throw new Exception($"Unable to find AUMID for Executable {ExecutableName} {AlternativeExecutableName ?? ""} in manifest {PathToManifest}");
+			}
+
+			return AUMID;
+		}
+
+		/// <summary>
+		/// Returns the AUMID that matches the given configuration from a list of AUMIDs
+		/// </summary>
+		public static string GetAUMIDForConfiguration(IEnumerable<string> AUMIDs, UnrealTargetConfiguration Configuration)
+		{
+			if (!AUMIDs.Any())
+			{
+				throw new Exception("No AUMIDs");
+			}
+
+			string? AUMID = AUMIDs.FirstOrDefault(A => A.EndsWith(Configuration.ToString(), StringComparison.Ordinal));
+			// fall back on any AUMIDs that are not tagged as Dev/Test/shipping
+			AUMID ??= AUMIDs.FirstOrDefault(A =>
+				!A.EndsWith(@"Development", StringComparison.Ordinal) &&
+				!A.EndsWith(@"Test", StringComparison.Ordinal) &&
+				!A.EndsWith(@"Shipping", StringComparison.Ordinal));
+
+			if (AUMID == null)
+			{
+				throw new Exception($"Unable to find AUMID for configuration {Configuration} from {String.Join(',', AUMIDs)}");
+			}
+
+			return AUMID;
+		}
+	}
+
+	/// <summary>
+	/// Structure containing information about an installed GDK application
+	/// </summary>
+	public class GDKInstalledApp
+	{
+		/// <summary>
+		/// The type of application
+		/// </summary>
+		public enum AppType
+		{
+			/// Unknown
+			Unknown,
+
+			/// Loose folder
+			Staged,
+
+			/// Installed package
+			Packaged,
+
+			/// Network share
+			NetworkShare,
+		};
+
+		internal GDKInstalledApp(JsonElement JsonPkg)
+		{
+			GameConfig = new GDKGameConfigInfo(JsonPkg);
+
+			DisplayName = JsonPkg.TryGetStringProperty("DisplayName", out string? Temp) ? Temp : "unknown";
+			DeployPath = JsonPkg.TryGetStringProperty("InstallPath", out Temp) || JsonPkg.TryGetStringProperty("DeployPath", out Temp) || JsonPkg.TryGetStringProperty("Network", out Temp) ? Temp : "";
+
+			Type = AppType.Unknown;
+			IsRegistered = false;
+			if (JsonPkg.TryGetStringProperty("DeployType", out Temp))
+			{
+				string[] PackageTypes = { "MSIXVC", "Package", "Install" };
+				string[] StagedTypes = { "Loose", "Folder" };
+				string[] NetworkTypes = { "Network" };
+				if (PackageTypes.Any(X => Temp.Contains(X, StringComparison.OrdinalIgnoreCase)))
+				{
+					Type = AppType.Packaged;
+				}
+				else if (StagedTypes.Any(X => Temp.Contains(X, StringComparison.OrdinalIgnoreCase)))
+				{
+					Type = AppType.Staged;
+				}
+				else if (NetworkTypes.Any(X => Temp.Contains(X, StringComparison.OrdinalIgnoreCase)))
+				{
+					Type = AppType.NetworkShare;
+				}
+				IsRegistered = !Temp.Contains("Unregistered", StringComparison.OrdinalIgnoreCase);
+			}
+		}
+
+		/// core MicrosoftGame.config information
+		public GDKGameConfigInfo GameConfig { get; private set; }
+
+		/// deployment type - staged or packaged
+		public AppType Type { get; private set; }
+
+		/// registration state
+		public bool IsRegistered { get; private set; }
+
+		/// display name
+		public string DisplayName { get; private set; }
+
+		/// deployment path
+		public string DeployPath { get; private set; }
+
+		/// <inheritdoc/>
+		public override string ToString()
+		{
+			return $"{GameConfig.PackageFullName} ({Type})";
+		}
+
+		/// <summary>
+		/// Returns a list of all applications from the given app list /json output
+		/// </summary>
+		/// <param name="JsonAppList"></param>
+		/// <returns>List of all applications</returns>
+		public static IEnumerable<GDKInstalledApp> ParseAppList(JsonDocument JsonAppList)
+		{
+			foreach (JsonElement JsonPkg in JsonAppList.RootElement.GetProperty("Packages").EnumerateArray())
+			{
+				yield return new GDKInstalledApp(JsonPkg);
+			}
+		}
+	}
+}
diff --git a/Engine/Source/Programs/UnrealBuildTool/Platform/Microsoft/GDK/MSGaming/MSGamingExports.cs b/Engine/Source/Programs/UnrealBuildTool/Platform/Microsoft/GDK/MSGaming/MSGamingExports.cs
new file mode 100644
--- /dev/null
+++ b/Engine/Source/Programs/UnrealBuildTool/Platform/Microsoft/GDK/MSGaming/MSGamingExports.cs
@@ -0,0 +1,192 @@
+// Copyright Epic Games, Inc. All Rights Reserved.
+
+using EpicGames.Core;
+using Microsoft.Win32;
+using System;
+using System.Collections.Generic;
+using System.Linq;
+using System.Runtime.Versioning;
+
+namespace UnrealBuildTool
+{
+	/// <summary>
+	/// Public MSGaming functions exposed to UAT
+	/// </summary>
+	public static class MSGamingExports
+	{
+		#region custom install actions
+
+		/// Type of custom install action
+		public enum CustomInstallActionType
+		{
+			/// command for installing
+			Install,
+			/// command for repairing
+			Repair,
+			/// command for uninstalling
+			Uninstall,
+		};
+
+#pragma warning disable CA1034 // Nested types should not be visible - preserve API of APIs exposed to UAT
+		/// a custom install action
+		public struct CustomInstallAction
+		{
+			/// location to copy the custom installs to when staging
+			private const string DefaultTargetRootFolder = "CustomInstall";
+
+			/// name of the install action
+			public string Name;
+
+			/// source folder, relative to engine root
+			public DirectoryReference SourceFolder;
+
+			/// target folder, relative to targe root folder
+			public string TargetFolder;
+
+			/// a command to run and its arguments
+			public struct InstallCommand
+			{
+				/// command to run
+				public FileReference Command;
+
+				/// arguments to pass to the given command
+				public string Arguments;
+			};
+
+			/// the commands to run for each action
+			public Dictionary<CustomInstallActionType, InstallCommand> Commands;
+
+			/// read the given config key and return the actions
+			public static IEnumerable<CustomInstallAction> GetFromConfig(ConfigHierarchy EngineIni, string IniSection)
+			{
+				List<CustomInstallAction> CustomInstallActions = new();
+
+				if (EngineIni!.TryGetValuesGeneric(IniSection, "CustomInstallActions", out CustomInstallAction[]? CustomInstallActionsArray) && CustomInstallActionsArray != null)
+				{
+					// set the default target folder
+					for (int i = 0; i < CustomInstallActionsArray.Length; i++)
+					{
+						if (String.IsNullOrEmpty(CustomInstallActionsArray[i].TargetFolder))
+						{
+							CustomInstallActionsArray[i].TargetFolder = CustomInstallActionsArray[i].Name;
+						}
+						else if (CustomInstallActionsArray[i].TargetFolder == ".")
+						{
+							CustomInstallActionsArray[i].TargetFolder = "";
+						}
+					}
+
+					CustomInstallActions.AddRange(CustomInstallActionsArray);
+				}
+
+				// inject the GameInput dependency if requested
+				if (EngineIni.GetBool(IniSection, "bAlwaysIncludeGameInputDependency", out bool bAlwaysIncludeGameInputDependency) && bAlwaysIncludeGameInputDependency)
+				{
+					DirectoryReference GameInputRedistDir = DirectoryReference.Combine(UnrealBuildBase.Unreal.EngineDirectory, @"Extras\Redist\en-us\");
+					FileReference GameInputRedist = FileReference.Combine(GameInputRedistDir, "GameInputRedist.msi");
+
+					InstallCommand MSICommand = new()
+					{
+						Command = GameInputRedist,
+						Arguments = ""
+					};
+
+					CustomInstallAction GameInputAction = new();
+					GameInputAction.Name = "GameInput";
+					GameInputAction.SourceFolder = GameInputRedistDir;
+					GameInputAction.TargetFolder = "GameInput";
+					GameInputAction.Commands = new()
+					{
+						{ CustomInstallActionType.Install, MSICommand },
+						{ CustomInstallActionType.Uninstall, MSICommand },
+						{ CustomInstallActionType.Repair, MSICommand }
+					};
+
+					CustomInstallActions.Add(GameInputAction);
+				}
+
+				return CustomInstallActions.ToArray();
+			}
+
+			/// the target root folder for staging custom install actions, relative to the package root
+			public static string GetTargetRootFolder(ConfigHierarchy EngineIni, string IniSection)
+			{
+				if (EngineIni.TryGetValue(IniSection, "CustomInstallStagingRoot", out string? TargetRootFolder ) && !String.IsNullOrEmpty(TargetRootFolder))
+				{
+					return TargetRootFolder;
+				}
+				else
+				{
+					return DefaultTargetRootFolder;
+				}
+			}
+		}
+#pragma warning restore CA1034
+
+		#endregion
+
+		/// <summary>
+		/// Determine whether Windows is in developer mode. PC GDK packages can't be sideloaded without it
+		/// </summary>
+		[SupportedOSPlatform("windows")]
+		public static bool IsWindowsDeveloperMode()
+		{
+			try
+			{
+				// normal windows key
+				using RegistryKey? Key = Registry.LocalMachine.OpenSubKey(@"SOFTWARE\Microsoft\Windows\CurrentVersion\AppModelUnlock");
+				if (Key != null && Convert.ToInt64(Key.GetValue("AllowDevelopmentWithoutDevLicense", 0L)) != 0)
+				{
+					return true;
+				}
+
+				// group policy key
+				using RegistryKey? GroupKey = Registry.LocalMachine.OpenSubKey(@"SOFTWARE\Policies\Microsoft\Windows\Appx");
+				if (GroupKey != null && Convert.ToInt64(GroupKey.GetValue("AllowDevelopmentWithoutDevLicense", 0L)) != 0)
+				{
+					return true;
+				}
+			}
+			catch (Exception)
+			{
+			}
+
+			return false;
+		}
+
+		/// <summary>
+		/// Determine if the Microsoft.GamingServices package is installed.
+		/// </summary>
+		[SupportedOSPlatform("windows")]
+		public static bool IsMSGamingRuntimeInstalled()
+		{
+			try
+			{
+				using RegistryKey? Key = Registry.CurrentUser.OpenSubKey(@"SOFTWARE\Classes\Local Settings\Software\Microsoft\Windows\CurrentVersion\AppModel\Repository\Packages");
+				if (Key != null)
+				{
+					return Key.GetSubKeyNames().Any( X => X.StartsWith("Microsoft.GamingServices_", StringComparison.Ordinal));
+				}
+			}
+			catch (Exception)
+			{
+			}
+		
+			return false;
+		}
+
+		/// <summary>
+		/// Determine whether the local version of Windows can support GDK
+		/// </summary>
+		[SupportedOSPlatform("windows")]
+		public static bool IsWindowsVersionSupported()
+		{
+			if (System.Environment.OSVersion.Version.Major < 10 || (System.Environment.OSVersion.Version.Major == 10 && System.Environment.OSVersion.Version.Build < 18362)) //https://en.wikipedia.org/wiki/Template:Windows_10_versions
+			{
+				return false;
+			}
+
+			return true;
+		}
+	}
+}
diff --git a/Engine/Source/Programs/UnrealBuildTool/Platform/Microsoft/GDK/MSGaming/MSGamingGameConfigGenerator.cs b/Engine/Source/Programs/UnrealBuildTool/Platform/Microsoft/GDK/MSGaming/MSGamingGameConfigGenerator.cs
new file mode 100644
--- /dev/null
+++ b/Engine/Source/Programs/UnrealBuildTool/Platform/Microsoft/GDK/MSGaming/MSGamingGameConfigGenerator.cs
@@ -0,0 +1,258 @@
+// Copyright Epic Games, Inc. All Rights Reserved.
+
+using EpicGames.Core;
+using Microsoft.Extensions.Logging;
+using System;
+using System.Collections.Generic;
+using System.IO;
+using System.Linq;
+using System.Xml.Linq;
+using UnrealBuildBase;
+using static UnrealBuildTool.MSGamingExports;
+
+namespace UnrealBuildTool.GDK
+{
+	partial struct GDKGameConfigGeneratorFactory
+	{
+#pragma warning disable CA1823 // Avoid unused private fields - instantiation is side-effectful, preferable to a static ctor
+#pragma warning disable IDE0052
+		private static readonly Type MSGamingGen = RegisterGenerator(UnrealTargetPlatform.Win64, typeof(MSGamingGameConfigGenerator));
+#pragma warning restore CA1823
+#pragma warning restore IDE0052
+	}
+
+	class MSGamingGameConfigGenerator : GDKGameConfigGenerator
+	{
+		protected override string TargetDeviceFamily => "PC";
+		protected override string DefaultArchitecture => "x64";
+		protected override string IniSection_PlatformTargetSettings => "/Script/MSGamingSupport.MSGamingSettings";
+		protected override string? IniSection_GeneralPlatformSettings => "/Script/WindowsTargetPlatform.WindowsTargetSettings"; // for querying DefaultGraphicsRHI etc.
+
+		protected virtual string IniPlatformName => "Windows";
+
+		/// <summary>
+		/// Create a manifest generator for Win64 using MS Gaming Runtime and MS Game Store
+		/// </summary>
+		public MSGamingGameConfigGenerator(ILogger InLogger)
+			: base(UnrealTargetPlatform.Win64, InLogger)
+		{
+		}
+
+		/// <summary>
+		/// Create a manifest generator for other platforms
+		/// </summary>
+		internal MSGamingGameConfigGenerator(UnrealTargetPlatform InPlatform, ILogger InLogger)
+			: base(InPlatform, InLogger)
+		{
+		}
+
+		/// <summary>
+		/// Register the locations where resource binary files can be found
+		/// </summary>
+		protected override void PrepareResourceBinaryPaths()
+		{
+			base.PrepareResourceBinaryPaths();
+
+			// check for DLC images first
+			if (DLCFile != null)
+			{
+				AppXResources!.ProjectBinaryResourceDirectories.Add(DirectoryReference.Combine(DLCFile.Directory, "Build", IniPlatformName, "MSGaming", BuildResourceSubPath));
+				AppXResources!.ProjectBinaryResourceDirectories.Add(DirectoryReference.Combine(DLCFile.Directory, "Platforms", IniPlatformName, "Build", "MSGaming", BuildResourceSubPath));
+			}
+
+			// allow projects to keep their resources in an MSGaming subfolder to help with project organization
+			if (ProjectFile != null)
+			{
+				AppXResources!.ProjectBinaryResourceDirectories.Add(DirectoryReference.Combine(ProjectFile.Directory, "Build", IniPlatformName, "MSGaming", BuildResourceSubPath));
+				AppXResources!.ProjectBinaryResourceDirectories.Add(DirectoryReference.Combine(ProjectFile.Directory, "Platforms", IniPlatformName, "Build", "MSGaming", BuildResourceSubPath));
+			}
+
+			// default images are stored under MSGamingSupport plugin folder
+			DirectoryReference MSGamingSupportResourcePath = DirectoryReference.Combine(Unreal.EngineDirectory, "Build", EngineResourceSubPath);
+			PluginInfo? MSGamingSupportPlugin = Plugins.GetPlugin("MSGamingSupport");
+			if (MSGamingSupportPlugin != null)
+			{
+				MSGamingSupportResourcePath = DirectoryReference.Combine( MSGamingSupportPlugin.File.Directory, "Build", EngineResourceSubPath);
+			}
+			AppXResources!.EngineFallbackBinaryResourceDirectories.Add(MSGamingSupportResourcePath);
+
+			// fallback path because Program targets do not load plugins & this is still under the platforms folder for now
+			AppXResources!.EngineFallbackBinaryResourceDirectories.Add(DirectoryReference.Combine(Unreal.EngineDirectory, @"Platforms\GDK\Plugins\MSGamingSupport\Build\DefaultImages"));
+		}
+
+		protected override XElement GetGame(Dictionary<UnrealTargetConfiguration, string> InExecutablePairs, out string IdentityName)
+		{
+			XElement Game = base.GetGame(InExecutablePairs, out IdentityName);
+
+			Game.Add(GetDesktopRegistration());
+
+			if (!IsDLC)
+			{
+				if (EngineIni!.GetBool(IniSection_PlatformTargetSettings, "bMSAFullTrust", out bool bMSAFullTrust) && bMSAFullTrust)
+				{
+					Game.Add(new XElement(XName.Get("MSAFullTrust"), "true"));
+				}
+			}
+
+			return Game;
+		}
+
+		private XElement GetDesktopRegistration()
+		{
+			XElement DesktopRegistration = new(XName.Get("DesktopRegistration"),
+				new XElement(XName.Get("ProcessorArchitecture"), DefaultArchitecture)
+				);
+
+			if (!IsDLC)
+			{
+				int ManifestVersion = GetManifestVersion();
+
+				DesktopRegistration.Add(
+					new XElement(XName.Get("MultiplayerProtocol"), "true"),
+					GetDependencyList(),
+					GetCustomInstallActions(),
+					GetWindowsOSVersion()
+					);
+
+				// Note: Manifest Version 0 has been deprecated as of October 2023 GDK and is not supported,
+				// but keeping this here in case Manifest version is overridden.  
+				if (ManifestVersion == 0) // ModFolder is deprecated in Manifest Version 1
+				{
+					if (EngineIni!.GetBool(IniSection_PlatformTargetSettings, "bUseModFolder", out bool bUseModFolder) && bUseModFolder)
+					{
+						if (EngineIni.GetString(IniSection_PlatformTargetSettings, "ModFolder", out string ModFolder) && !String.IsNullOrEmpty(ModFolder))
+						{
+							DesktopRegistration.Add(new XElement(XName.Get("ModFolder"),
+								new XAttribute(XName.Get("Name"), ModFolder)
+								)
+							);
+						}
+					}
+				}
+			}
+			else
+			{
+				string MainPackageIdentityName = GetMainPackageIdentityPackageName();
+
+				DesktopRegistration.Add(new XElement(XName.Get("MainPackageDependency"),
+					new XAttribute(XName.Get("Name"), MainPackageIdentityName)
+					)
+				);
+			}
+
+			return DesktopRegistration;
+		}
+
+		private XElement? GetDependencyList()
+		{
+			XElement DependencyList = new(XName.Get("DependencyList"));
+
+			if (NoBootstrapExe)
+			{
+				// Without a bootstrapper, need Windows to install VC runtime libraries during package install
+				// this will cause VC DLL loads to be redirected from %System32% to %ProgramFiles%\WindowsApps\Microsoft.VCLibs.140.00.UWPDesktop_14.0.32530.0_x64__8wekyb3d8bbwe
+				DependencyList.Add(new XElement(XName.Get("KnownDependency"), new XAttribute(XName.Get("Name"), "VC14")));
+			}
+
+			bool bNeedsDX11 = (GetConfigString("DefaultGraphicsRHI", null) == "DefaultGraphicsRHI_DX11") ||
+							  (GetConfigBool("bAlwaysIncludeDX11Dependency", null));
+			if (bNeedsDX11)
+			{
+				// this will cause DX DLL loads to be redirected from %System32% to %ProgramFiles%\WindowsApps\Microsoft.DirectXRuntime_9.29.1974.0_x64__8wekyb3d8bbwe
+				Log.TraceInformationOnce("Adding KnownDependency on DX11 redist in MicrosoftGame.config");
+				DependencyList.Add(new XElement(XName.Get("KnownDependency"), new XAttribute(XName.Get("Name"), "DX11")));
+			}
+
+			return DependencyList.HasElements ? DependencyList : null;
+		}
+
+		private XElement? GetCustomInstallActions()
+		{
+			IEnumerable<CustomInstallAction> CustomInstallActions = CustomInstallAction.GetFromConfig(EngineIni!, IniSection_PlatformTargetSettings);
+			if (!CustomInstallActions.Any())
+			{
+				return null;
+			}
+
+			string TargetRoot = CustomInstallAction.GetTargetRootFolder(EngineIni!, IniSection_PlatformTargetSettings);
+
+			List<XElement> InstallActionList = new();
+			List<XElement> RepairActionList = new();
+			List<XElement> UninstallActionList = new();
+			foreach (CustomInstallAction Action in CustomInstallActions)
+			{
+				foreach (KeyValuePair<CustomInstallActionType, CustomInstallAction.InstallCommand> Command in Action.Commands)
+				{
+					List<XAttribute> Attribs = new List<XAttribute>
+					{
+						new("File",      Path.Combine( Action.TargetFolder, Command.Value.Command.MakeRelativeTo(Action.SourceFolder))),
+						new("Name",      Action.Name),
+					};
+					if (!String.IsNullOrEmpty(Command.Value.Arguments))
+					{
+						Attribs.Add(new("Arguments", Command.Value.Arguments));
+					}
+					switch (Command.Key)
+					{
+						case CustomInstallActionType.Install: InstallActionList.Add(new XElement("InstallAction", Attribs)); break;
+						case CustomInstallActionType.Repair: RepairActionList.Add(new XElement("RepairAction", Attribs)); break;
+						case CustomInstallActionType.Uninstall: UninstallActionList.Add(new XElement("UninstallAction", Attribs)); break;
+					}
+				}
+			}
+
+			XElement CustomInstallActionsElement = new("CustomInstallActions",
+				new XElement("Folder", new XText(TargetRoot)),
+				InstallActionList.Any() ? new XElement("InstallActionList", InstallActionList) : null,
+				RepairActionList.Any() ? new XElement("RepairActionList", RepairActionList) : null,
+				UninstallActionList.Any() ? new XElement("UninstallActionList", UninstallActionList) : null
+			);
+			return CustomInstallActionsElement;
+		}
+
+		private XElement? GetWindowsOSVersion()
+		{
+			EngineIni!.GetString(IniSection_PlatformTargetSettings, "RequiredMinimumWindowsVersion", out string RequiredMinVersion);
+			EngineIni!.GetString(IniSection_PlatformTargetSettings, "SuggestedMinimumWindowsVersion", out string SuggestedMinVersion);
+			EngineIni!.GetString(IniSection_PlatformTargetSettings, "RecommendedWindowsVersion", out string RecommendedVersion);
+
+			XElement WindowsOsVersion = new(XName.Get("WindowsOsVersion"));
+
+			if (!String.IsNullOrWhiteSpace(RequiredMinVersion))
+			{
+				WindowsOsVersion.SetAttributeValue("RequiredMinimum", RequiredMinVersion);
+			}
+
+			if (!String.IsNullOrWhiteSpace(SuggestedMinVersion))
+			{
+				WindowsOsVersion.SetAttributeValue("SuggestedMinimum", SuggestedMinVersion);
+			}
+
+			if (!String.IsNullOrWhiteSpace(RecommendedVersion))
+			{
+				WindowsOsVersion.SetAttributeValue("Recommended", RecommendedVersion);
+			}
+
+			// The <WindowsOsVersion> element is optional. If none of the values were configured, omit it.
+			return WindowsOsVersion.HasAttributes ? WindowsOsVersion : null;
+		}
+
+		protected override string? IncludeBuildVersionInPackageVersion(string? VersionNumber)
+		{
+			// GDK on Windows reserves the lowest version element (Revision) for internal use so package versioning needs to be a little different
+
+			if (VersionNumber != null && BuildVersion.TryRead(BuildVersion.GetDefaultFileName(), out BuildVersion? BuildVersionForPackage) && BuildVersionForPackage.Changelist != 0)
+			{
+				// Break apart the version number into individual elements
+				string[] SplitVersionString = VersionNumber.Split('.');
+				VersionNumber = String.Format("{0}.{1}.{2}.{3}",
+					SplitVersionString[0],
+					BuildVersionForPackage.Changelist / 10000,
+					BuildVersionForPackage.Changelist % 10000,
+					0); //must be zero
+			}
+
+			return VersionNumber;
+		}
+	}
+}
diff --git a/Engine/Source/Programs/UnrealBuildTool/Platform/Microsoft/GDK/MSGaming/MSGamingToolChain.cs b/Engine/Source/Programs/UnrealBuildTool/Platform/Microsoft/GDK/MSGaming/MSGamingToolChain.cs
new file mode 100644
--- /dev/null
+++ b/Engine/Source/Programs/UnrealBuildTool/Platform/Microsoft/GDK/MSGaming/MSGamingToolChain.cs
@@ -0,0 +1,103 @@
+// Copyright Epic Games, Inc. All Rights Reserved.
+
+using EpicGames.Core;
+using Microsoft.Extensions.Logging;
+using System;
+using System.Collections.Generic;
+using System.Linq;
+using UnrealBuildTool.GDK;
+
+namespace UnrealBuildTool
+{
+	class MSGamingToolChain : VCToolChain
+	{
+		protected virtual bool SupportsContentOnlyProjects => false; // there is no content-only project support for Win64+GDK because it requires bespoke target.cs files
+
+		public MSGamingToolChain(ReadOnlyTargetRules Target, ILogger Logger)
+			: base(Target, Logger)
+		{
+		}
+
+		public override void GetVersionInfo(List<string> Lines)
+		{
+			base.GetVersionInfo(Lines);
+
+			string? GDKEdition = GDKExports.GetSDKVersion();
+			if (GDKEdition != null)
+			{
+				string CurrentGDKDir = GDKExports.GetCurrentGSDKDir();
+				Lines.Add($"Using GDK version {GDKEdition} ({CurrentGDKDir}).");
+			}
+			else
+			{
+				Lines.Add("GameDKCoreLatest or GameDKLatest environment variable not found - is the GDK installed?");
+			}
+		}
+
+		protected override void AppendLinkArguments(LinkEnvironment LinkEnvironment, List<string> Arguments)
+		{
+			base.AppendLinkArguments(LinkEnvironment, Arguments);
+
+			// do not strip the PDB path from the executable if we want to debug it (otherwise it can't find the symbols when attempting to debug an installed package build)
+			if (LinkEnvironment.bCreateDebugInfo && LinkEnvironment.Configuration != CppConfiguration.Shipping)
+			{
+				Arguments.Remove("/PDBALTPATH:%_PDB%");
+			}
+		}
+
+		public override void FinalizeOutput(ReadOnlyTargetRules Target, TargetMakefileBuilder MakefileBuilder, DirectoryReference ExeDir)
+		{
+			base.FinalizeOutput(Target, MakefileBuilder, ExeDir);
+
+			// Look up the project
+			FileReference? ProjectFile = Target.ProjectFile;
+			if (SupportsContentOnlyProjects)
+			{
+				ProjectFile = GDKContentOnlyProjectSupport.GetProjectFileFromTarget(Target, Logger) ?? ProjectFile;
+			}
+
+			// Create a placeholder MicrosoftGame.config to be used during F5 debugging, and for Program targets. It does not reference executables because it could be shared between several configurations.
+			// This allows the game to run while using assets from the Saved/Cooked folder, without needing -basedir= to point to a Staged build
+			if ((Environment.CommandLine.Contains("-frommsbuild", StringComparison.OrdinalIgnoreCase) && ProjectFile != null) || Target.Type == TargetType.Program)
+			{
+				DirectoryReference ExecutableDirectory = MakefileBuilder.Makefile.ExecutableFile.Directory;
+				Action GenAction = GDKGameConfigGeneratorMode.CreateAction(MakefileBuilder, Target.Platform, ExecutableDirectory, Target.Name, ProjectFile, CustomConfig:Target.CustomConfig);
+				GenAction.CommandDescription = "Generate MicrosoftGame.config";
+				GenAction.bCanExecuteInUBA = false; // TODO: Fails with some missing file, revisit
+			}
+		}
+
+		protected override void ModifyFinalLinkArguments(LinkEnvironment LinkEnvironment, List<string> Arguments, bool bBuildImportLibraryOnly)
+		{
+			base.ModifyFinalLinkArguments(LinkEnvironment, Arguments, bBuildImportLibraryOnly);
+
+			// optionally audit the GDK editions of any libs we link to
+			GDKEditionAudit.VerifyLinkEnvironment(EnvVars, LinkEnvironment, bBuildImportLibraryOnly, Logger);
+		}
+
+		public override void PrepareRuntimeDependencies(List<RuntimeDependency> RuntimeDependencies, Dictionary<FileReference, FileReference> TargetFileToSourceFile, DirectoryReference ExeDir)
+		{
+			base.PrepareRuntimeDependencies(RuntimeDependencies, TargetFileToSourceFile, ExeDir);
+
+			// optionally audit the GDK editions of any libs we depend on
+			GDKEditionAudit.VerifyRuntimeDependencies(EnvVars, Target.Platform, TargetFileToSourceFile.Values.Distinct(), Logger);
+		}
+
+		public override void SetUpGlobalEnvironment(ReadOnlyTargetRules Target, CppCompileEnvironment GlobalCompileEnvironment, LinkEnvironment GlobalLinkEnvironment)
+		{
+			base.SetUpGlobalEnvironment(Target, GlobalCompileEnvironment, GlobalLinkEnvironment);
+
+			DirectoryReference platformGDK = new(GDKExports.GetCurrentGSDKDir());
+			GlobalCompileEnvironment.RootPaths.AddFolderName(CppRootPathFolder.PlatformSDK, $"{Target.Platform}SDK");
+			GlobalLinkEnvironment.RootPaths.AddFolderName(CppRootPathFolder.PlatformSDK, $"{Target.Platform}SDK");
+			GlobalCompileEnvironment.RootPaths[CppRootPathFolder.PlatformSDK] = platformGDK;
+			GlobalLinkEnvironment.RootPaths[CppRootPathFolder.PlatformSDK] = platformGDK;
+
+			string? currentGDKVersion = GDKExports.GetSDKVersion();
+			if (currentGDKVersion != null)
+			{
+				EpicGames.UBA.Utils.RegisterPathHash(platformGDK.FullName, currentGDKVersion);
+			}
+		}
+	}
+}
\ No newline at end of file
