Checking CoreDeviceService thread usage and killing/restarting process if too high to avoid hangs on mac agents.

#Synced-CL 50214590

diff --git a/Engine/Source/Programs/AutomationTool/Gauntlet/Platform/IOS/Gauntlet.TargetDeviceIOS.cs b/Engine/Source/Programs/AutomationTool/Gauntlet/Platform/IOS/Gauntlet.TargetDeviceIOS.cs
--- a/Engine/Source/Programs/AutomationTool/Gauntlet/Platform/IOS/Gauntlet.TargetDeviceIOS.cs
+++ b/Engine/Source/Programs/AutomationTool/Gauntlet/Platform/IOS/Gauntlet.TargetDeviceIOS.cs
@@ -256,6 +256,21 @@ namespace Gauntlet
 				{
 					AppiumContainer.Dispose();
 				}
+
+				// Clean up any hanging devicectl/xcrun processes to prevent CoreDeviceService thread buildup
+				try
+				{
+					// Kill any xcrun processes that may be hanging
+					AppleBuild.ExecuteCommand("killall", "xcrun").WaitForExit();
+
+					// Give processes time to clean up
+					Thread.Sleep(500);
+				}
+				catch (Exception Ex)
+				{
+					Log.Verbose("Exception during cleanup in Dispose: {0}", Ex.Message);
+				}
+
 				DisposedValue = true;
 			}
 		}
@@ -267,6 +282,80 @@ namespace Gauntlet
 		}
 		#endregion
 
+		/// <summary>
+		/// Checks CoreDeviceService thread count and kills it if too high to prevent agent failure
+		/// </summary>
+		/// <returns>True if CoreDeviceService was killed and restarted</returns>
+		private static bool CheckAndCleanupCoreDeviceService(int MaxThreadCount = 500)
+		{
+			try
+			{
+				// Get thread count for CoreDeviceService
+				IProcessResult ThreadCheckResult = AppleBuild.ExecuteCommand("sh", "-c \"ps -M $(pgrep -x CoreDeviceService 2>/dev/null) 2>/dev/null | tail -n +2 | wc -l\"");
+				ThreadCheckResult.WaitForExit();
+
+				if (ThreadCheckResult.ExitCode == 0 && !string.IsNullOrWhiteSpace(ThreadCheckResult.Output))
+				{
+					string ThreadCountStr = ThreadCheckResult.Output.Trim();
+					if (int.TryParse(ThreadCountStr, out int ThreadCount))
+					{
+						Log.Verbose("CoreDeviceService thread count: {0}", ThreadCount);
+
+						if (ThreadCount > MaxThreadCount)
+						{
+							Log.Warning("CoreDeviceService has {0} threads (exceeds limit of {1}). Killing to prevent agent failure...", ThreadCount, MaxThreadCount);
+
+							// Try to kill CoreDeviceService daemon with SIGKILL - it will automatically restart
+							// First try with sudo for system daemon privileges
+							IProcessResult KillResult = AppleBuild.ExecuteCommand("sudo", "killall -9 CoreDeviceService");
+							KillResult.WaitForExit();
+
+							if (KillResult.ExitCode != 0)
+							{
+								Log.Warning("sudo killall failed with exit code {0}, trying without sudo: {1}", KillResult.ExitCode, KillResult.Output);
+								// Try without sudo but with -9 signal
+								KillResult = AppleBuild.ExecuteCommand("killall", "-9 CoreDeviceService");
+								KillResult.WaitForExit();
+								Thread.Sleep(5000);
+							}
+
+							if (KillResult.ExitCode == 0)
+							{
+								Log.Info("CoreDeviceService killed successfully, waiting for restart...");
+								// Wait for service to fully terminate and restart
+								Thread.Sleep(3000);
+
+								// Verify it restarted
+								IProcessResult VerifyResult = AppleBuild.ExecuteCommand("pgrep", "-x CoreDeviceService");
+								VerifyResult.WaitForExit();
+								if (VerifyResult.ExitCode == 0)
+								{
+									Log.Info("CoreDeviceService restarted successfully");
+									return true; // Service was killed and restarted
+								}
+								else
+								{
+									Log.Warning("CoreDeviceService may not have restarted, continuing anyway...");
+									return true; // Still treated as killed
+								}
+							}
+							else
+							{
+								Log.Warning("Failed to kill CoreDeviceService (exit code {0}): {1}", KillResult.ExitCode, KillResult.Output);
+								return false;
+							}
+						}
+					}
+				}
+			}
+			catch (Exception Ex)
+			{
+				Log.Verbose("Exception checking CoreDeviceService thread count: {0}", Ex.Message);
+			}
+
+			return false; // No cleanup needed or failed
+		}
+
 		#region ITargetDevice
 		public bool Connect()
 		{
@@ -1040,8 +1129,11 @@ namespace Gauntlet
 				AppleBuild.ExecuteCommand("killall", "xcrun").WaitForExit();
 				AppleBuild.ExecuteCommand("killall", "ios-deploy").WaitForExit();
 				AppleBuild.ExecuteCommand("killall", "lldb").WaitForExit();
-				AppleBuild.ExecuteCommand("killall", "usbmuxd").WaitForExit();
-				AppleBuild.ExecuteCommand("launchctl", "kickstart -k system/com.apple.usbmuxd").WaitForExit();
+				//AppleBuild.ExecuteCommand("killall", "usbmuxd").WaitForExit();
+				//AppleBuild.ExecuteCommand("launchctl", "kickstart -k system/com.apple.usbmuxd").WaitForExit();
+
+				// Check CoreDeviceService thread count and clean up if needed
+				CheckAndCleanupCoreDeviceService();
 
 				// Gather the iOS devices are connected to this UAT host
 				ConnectedDevices = GetConnectedDevices();
@@ -1054,6 +1146,33 @@ namespace Gauntlet
 		}
 
 		private static Dictionary<string, ConnectedDevice> GetConnectedDevices()
+		{
+			// Check CoreDeviceService thread count before making devicectl calls
+			bool ServiceWasKilled = CheckAndCleanupCoreDeviceService();
+
+			Dictionary<string, ConnectedDevice> Devices = GetConnectedDevicesInternal();
+
+			// If CoreDeviceService was just killed and no devices found, wait longer and retry once
+			if (ServiceWasKilled && (Devices == null || !Devices.Any()))
+			{
+				Log.Warning("No devices found immediately after killing CoreDeviceService. Waiting 5 seconds and retrying...");
+				Thread.Sleep(5000);
+				Devices = GetConnectedDevicesInternal();
+
+				if (Devices == null || !Devices.Any())
+				{
+					Log.Warning("Still no devices found after retry.");
+				}
+				else
+				{
+					Log.Info("Found {0} device(s) after retry.", Devices.Count);
+				}
+			}
+
+			return Devices;
+		}
+
+		private static Dictionary<string, ConnectedDevice> GetConnectedDevicesInternal()
 		{
 			Dictionary<string, ConnectedDevice> Devices = new();
 
