fix case-sensitive device name matching in remote deployment support

#rnx

#Synced-CL 50642529

#Synced-CL 50642529

diff --git a/Engine/Source/Programs/AutomationTool/Win/RemoteWinSupport.cs b/Engine/Source/Programs/AutomationTool/Win/RemoteWinSupport.cs
--- a/Engine/Source/Programs/AutomationTool/Win/RemoteWinSupport.cs
+++ b/Engine/Source/Programs/AutomationTool/Win/RemoteWinSupport.cs
@@ -7,6 +7,7 @@ using System;
 using System.Collections.Generic;
 using System.Diagnostics;
 using System.IO;
+using System.Linq;
 using UnrealBuildTool;
 
 using static AutomationTool.CommandUtils;
@@ -39,7 +40,7 @@ public static class RemoteWinSupport
 
 	public static bool IsRemoteDevice( string DeviceName, UnrealTargetPlatform Platform )
 	{
-		return GetAllRemoteDeviceNames(Platform).Contains(DeviceName);
+		return GetAllRemoteDeviceNames(Platform).Contains(DeviceName, StringComparer.OrdinalIgnoreCase);
 	}
 	
 	public static List<string> GetAllRemoteDeviceNames(UnrealTargetPlatform Platform)
diff --git a/Engine/Source/Programs/AutomationTool/Win/WinPlatform.Automation.cs b/Engine/Source/Programs/AutomationTool/Win/WinPlatform.Automation.cs
--- a/Engine/Source/Programs/AutomationTool/Win/WinPlatform.Automation.cs
+++ b/Engine/Source/Programs/AutomationTool/Win/WinPlatform.Automation.cs
@@ -57,7 +57,7 @@ public class Win64Platform : Platform
 		{
 			SteamDeckSupport.Deploy(UnrealTargetPlatform.Win64, Params, SC);
 		}
-		else if (Params.Devices.Count > 0 && GetDevices().FirstOrDefault(x => Params.DeviceNames.Contains(x.Id))?.Type == "RemoteWin")
+		else if (Params.Devices.Count > 0 && GetDevices().FirstOrDefault(x => Params.DeviceNames.Contains(x.Id, StringComparer.OrdinalIgnoreCase))?.Type == "RemoteWin")
 		{
 			RemoteWinSupport.Deploy(PlatformType, Params, SC);
 		}
@@ -69,7 +69,7 @@ public class Win64Platform : Platform
 		{
 			return SteamDeckSupport.RunClient(UnrealTargetPlatform.Win64, ClientRunFlags, ClientApp, ClientCmdLine, Params);
 		}
-		else if (Params.Devices.Count > 0 && GetDevices().FirstOrDefault(x => Params.DeviceNames.Contains(x.Id))?.Type == "RemoteWin")
+		else if (Params.Devices.Count > 0 && GetDevices().FirstOrDefault(x => Params.DeviceNames.Contains(x.Id, StringComparer.OrdinalIgnoreCase))?.Type == "RemoteWin")
 		{
 			return RemoteWinSupport.RunClient(PlatformType, ClientRunFlags, ClientApp, ClientCmdLine, Params);
 		}
