Turnkey - Improved text on prompt asking user to create a clientspec for AutoSDK.

#Synced-CL 50781450

diff --git a/Engine/Source/Programs/AutomationTool/Turnkey/Providers/PerforceCopyProvider.cs b/Engine/Source/Programs/AutomationTool/Turnkey/Providers/PerforceCopyProvider.cs
--- a/Engine/Source/Programs/AutomationTool/Turnkey/Providers/PerforceCopyProvider.cs
+++ b/Engine/Source/Programs/AutomationTool/Turnkey/Providers/PerforceCopyProvider.cs
@@ -270,7 +270,7 @@ namespace Turnkey
 				{
 					TurnkeyUtils.Log("Unable to find a clientspec for the perforce operation {0}, looking for a depot client", Operation);
 
-					bool bResponse = TurnkeyUtils.GetUserConfirmation("Would you like to create one?", false);
+					bool bResponse = TurnkeyUtils.GetUserConfirmation("Unable to find an existing Perforce Workspace for the Turnkey AutoSDK.\nWould you like to create one?", false);
 					if (bResponse == false)
 					{
 						// make sure to try again next time
