- Convert pakchunkN files to local path separator when reading them in so path matching works on all platforms

diff --git a/Engine/Source/Programs/AutomationTool/Scripts/CopyBuildToStagingDirectory.Automation.cs b/Engine/Source/Programs/AutomationTool/Scripts/CopyBuildToStagingDirectory.Automation.cs
--- a/Engine/Source/Programs/AutomationTool/Scripts/CopyBuildToStagingDirectory.Automation.cs
+++ b/Engine/Source/Programs/AutomationTool/Scripts/CopyBuildToStagingDirectory.Automation.cs
@@ -5349,7 +5349,9 @@ namespace AutomationScripts
 						if (!FileNameToChunks.TryGetValue(FileName, out Chunks))
 						{
 							Chunks = new List<ChunkDefinition>();
-							FileNameToChunks.Add(FileName, Chunks);
+							// make sure the paths match the host platform for later comparison
+							string Fixed = FileName.Replace('/', Path.DirectorySeparatorChar).Replace('\\', Path.DirectorySeparatorChar);
+							FileNameToChunks.Add(Fixed, Chunks);
 						}
 						Chunks.Add(Chunk);
 					}
