Report inactivity timeout in Engine Test report

#rnx

#Synced-CL 50818512

diff --git a/Engine/Source/Programs/AutomationTool/Gauntlet/Unreal/Automation/UE.Automation.cs b/Engine/Source/Programs/AutomationTool/Gauntlet/Unreal/Automation/UE.Automation.cs
--- a/Engine/Source/Programs/AutomationTool/Gauntlet/Unreal/Automation/UE.Automation.cs
+++ b/Engine/Source/Programs/AutomationTool/Gauntlet/Unreal/Automation/UE.Automation.cs
@@ -751,7 +751,7 @@ namespace UE
 		/// <param name="JsonTestPassResults"></param>
 		private void UpdateTestStateOnCriticalFailure(UnrealAutomatedTestPassResults JsonTestPassResults)
 		{
-			bool HasTimeout = RoleResults != null && RoleResults.Where(R => R.ProcessResult == UnrealProcessResult.TimeOut).Any();
+			bool HasTimeout = GetTestResult() == TestResult.TimedOut || (RoleResults != null && RoleResults.Where(R => R.ProcessResult == UnrealProcessResult.TimeOut).Any());
 			string HordeArtifactPath = GetCachedConfiguration().HordeArtifactPath;
 			var MainRole = GetCachedConfiguration().GetMainRequiredRole();
 			// Get any critical error and push it to json report and resave it.
@@ -780,7 +780,14 @@ namespace UE
 					string ErrorMessage = null;
 					if (HasTimeout)
 					{
-						ErrorMessage = String.Format("Session reached timeout after {0} seconds.", MaxDuration);
+						if (RoleResults != null && RoleResults.Where(R => R.ProcessResult == UnrealProcessResult.TimeOut).Any())
+						{
+							ErrorMessage = String.Format("Session reached timeout after {0:0.00} minutes.", MaxDuration / 60);
+						}
+						else
+						{
+							ErrorMessage = String.Format("No activity observed in last {0:0.00} minutes. Test aborted.", IdleTimeoutSec / 60);
+						}
 					}
 					else
 					{
