Revert back to BPS 1.7.5

#rnx

diff --git a/Engine/Source/Programs/AutomationTool/AutomationUtils/MCPPublic.cs b/Engine/Source/Programs/AutomationTool/AutomationUtils/MCPPublic.cs
--- a/Engine/Source/Programs/AutomationTool/AutomationUtils/MCPPublic.cs
+++ b/Engine/Source/Programs/AutomationTool/AutomationUtils/MCPPublic.cs
@@ -644,7 +644,7 @@ namespace EpicGames.MCP.Automation
 			/// <summary>
 			/// The current live, tested build.
 			/// </summary>
-			Live,
+			Latest,
 			/// <summary>
 			/// The latest published build, may be untested.
 			/// </summary>
@@ -665,7 +665,10 @@ namespace EpicGames.MCP.Automation
 			Online_v151,
 			Online_v152,
 			Online_v160,
+			Online_v175,
+			Online_v180,
 
+			Live = Online_v175,
 			Online_Live = Live
 		}
 
