Refactoring TargetDeviceIOS device pairing logic to mitigate ios connection issues.
#Synced-CL 50214590

diff --git a/Engine/Source/Programs/AutomationTool/Gauntlet/Platform/IOS/Gauntlet.IOSAppium.cs b/Engine/Source/Programs/AutomationTool/Gauntlet/Platform/IOS/Gauntlet.IOSAppium.cs
--- a/Engine/Source/Programs/AutomationTool/Gauntlet/Platform/IOS/Gauntlet.IOSAppium.cs
+++ b/Engine/Source/Programs/AutomationTool/Gauntlet/Platform/IOS/Gauntlet.IOSAppium.cs
@@ -155,14 +155,19 @@ namespace Gauntlet
 		/// </summary>
 		private AppiumLocalService Server = null;
 
-		private Thread ScreenshotThread = null;
+		private Thread MonitorThread = null;
 
-		private AutoResetEvent ScreenshotTrigger = null;
+		private AutoResetEvent MonitorTrigger = null;
 
 		/// <summary>
 		/// UUID of the device being tests
 		/// </summary>
 		private string UUID = null;
+
+		/// <summary>
+		/// Package name of the app being tested
+		/// </summary>
+		private string PackageName = null;
 		static AppiumContainer()
 		{
 			if (Globals.Params.ParseParam("NoAppium"))
@@ -226,6 +231,8 @@ namespace Gauntlet
 		{
 			Stop();
 
+			this.PackageName = PackageName;
+
 			// Install WebdriverAgent
 			try
 			{
@@ -287,8 +294,8 @@ namespace Gauntlet
 
 			if (!Globals.Params.ParseParam("NoAppiumScreenshots") && AppiumConfig.ScreenshotPeriod > 0)
 			{
-				ScreenshotThread = new Thread(StartScreenshotLoop);
-				ScreenshotThread.Start();
+				MonitorThread = new Thread(StartMonitorLoop);
+				MonitorThread.Start();
 			}
 		}
 
@@ -306,9 +313,9 @@ namespace Gauntlet
 				Server = null;
 			}
 
-			if (ScreenshotThread != null)
+			if (MonitorThread != null)
 			{
-				StopScreenshotLoop();
+				StopMonitorLoop();
 			}
 
 			// TODO: Terminate WDA app - parse devicectl for pid, then kill?
@@ -405,50 +412,79 @@ namespace Gauntlet
 			Driver.ActivateApp(appID);
 		}
 
-		public void CancelAppleSignOnAlert(string appID, int maxRetries = 5)
+		/// <summary>
+		/// Checks for and dismisses Apple Sign-On alerts. Used by the monitor loop.
+		/// Only performs 1 retry attempt.
+		/// </summary>
+		private void CheckAndDismissAlert()
 		{
 			var WebDriverWait = new WebDriverWait(new SystemClock(), Driver, TimeSpan.FromMilliseconds(5000), TimeSpan.FromMilliseconds(250));
-			for (int attempt = 0; attempt < maxRetries; attempt++)
+
+			try
 			{
-				Log.Verbose("starting alert attempt #{0}", attempt);
-				try
+				var alert = WebDriverWait.Until(drv =>
+				{
+					try
+					{
+						return drv.SwitchTo().Alert();
+					}
+					catch (NoAlertPresentException)
+					{
+						return null;
+					}
+					catch (Exception ex)
+					{
+						Log.Verbose("Alert switch attempt failed: {0}", ex.GetType().Name);
+						return null;  // Return null to keep polling
+					}
+				});
+
+				if (alert != null)
 				{
-					var alert = WebDriverWait.Until(drv =>
+					Log.Info("Found Alert, trying to dismiss");
+
+					// Retry find+click operation to handle stale elements
+					bool clicked = false;
+					for (int attempt = 0; attempt < 3 && !clicked; attempt++)
 					{
 						try
 						{
-							return drv.SwitchTo().Alert();
+							var cancelButton = Driver.FindElement(By.XPath("//XCUIElementTypeButton[@name='Cancel']"));
+							Log.Info("Found Cancel Button, clicking (attempt {0})", attempt + 1);
+							cancelButton.Click();
+							clicked = true;
+							Log.Info("Successfully clicked Cancel button");
+						}
+						catch (NoSuchElementException)
+						{
+							Log.Info("Cancel button not found on attempt {0}", attempt + 1);
+							break; // Button doesn't exist, don't retry
 						}
-						catch (NoAlertPresentException)
+						catch (StaleElementReferenceException)
 						{
-							return null;
+							Log.Verbose("Button became stale on attempt {0}, retrying...", attempt + 1);
+							Thread.Sleep(100); // Brief delay before retry
 						}
-					});
+					}
 
-					if (alert == null)
+					if (clicked && !string.IsNullOrEmpty(PackageName))
 					{
 						Thread.Sleep(300);
-						continue;
+						Driver.ActivateApp(PackageName);
 					}
-
-					alert.Dismiss();
-					Log.Info("Appium has Dismissed an Alert.");
-					Driver.ActivateApp(appID);
-
-					return;
-				}
-				catch (WebDriverTimeoutException)
-				{
 				}
-				catch (UnhandledAlertException)
-				{
-				}
-				catch (Exception ex)
-				{
-					Log.Verbose("Other exception hit: {0}", ex.Message);
-				}
-
-				Thread.Sleep(300);
+			}
+			catch (WebDriverTimeoutException)
+			{
+				// No alert present, this is normal
+			}
+			catch (UnhandledAlertException)
+			{
+				// Alert handling already in progress
+			}
+			catch (Exception ex)
+			{
+				Log.Info("Alert check exception: {0}", ex.Message);
 			}
 		}
 
@@ -482,10 +518,10 @@ namespace Gauntlet
 			}
 		}
 
-		private void StartScreenshotLoop()
+		private void StartMonitorLoop()
 		{
 			TimeSpan Period = TimeSpan.FromSeconds(AppiumConfig.ScreenshotPeriod);
-			ScreenshotTrigger = new AutoResetEvent(false);
+			MonitorTrigger = new AutoResetEvent(false);
 
 			string ScreenshotDirectory = GetScreenshotDirectory();
 
@@ -494,28 +530,40 @@ namespace Gauntlet
 				DeleteDirectory_NoExceptions(ScreenshotDirectory);
 			}
 
-			while (!ScreenshotTrigger.WaitOne(Period))
+			while (!MonitorTrigger.WaitOne(Period))
 			{
+				// Take screenshot
 				TakeScreenshot(ScreenshotDirectory);
+
+				// Check for alerts first (more critical)
+				Log.Info("Trying to find Alert");
+				try
+				{
+					CheckAndDismissAlert();
+				}
+				catch (Exception Ex)
+				{
+					Log.Info("Alert check failed: {0}", Ex.Message);
+				}
 			}
 		}
 
 		/// <summary>
-		/// Stops the capturing of periodic screenshots, optionally moving all captured images to an output directory
+		/// Stops the monitoring loop (screenshots and alerts), optionally moving all captured images to an output directory
 		/// </summary>
 		/// <param name="OutputDirectory">Folder to copy screenshots to</param>
-		public void StopScreenshotLoop(string OutputDirectory = "")
+		public void StopMonitorLoop(string OutputDirectory = "")
 		{
-			if (ScreenshotThread == null || ScreenshotTrigger == null)
+			if (MonitorThread == null || MonitorTrigger == null)
 			{
-				Log.Verbose("Attempted to end the screenshot loop before it started. Ignoring.");
+				Log.Verbose("Attempted to end the monitor loop before it started. Ignoring.");
 				return;
 			}
 
-			ScreenshotTrigger.Set();
-			ScreenshotThread.Join();
-			ScreenshotThread = null;
-			ScreenshotTrigger = null;
+			MonitorTrigger.Set();
+			MonitorThread.Join();
+			MonitorThread = null;
+			MonitorTrigger = null;
 
 			if (!string.IsNullOrEmpty(OutputDirectory))
 			{
diff --git a/Engine/Source/Programs/AutomationTool/Gauntlet/Platform/IOS/Gauntlet.TargetDeviceIOS.cs b/Engine/Source/Programs/AutomationTool/Gauntlet/Platform/IOS/Gauntlet.TargetDeviceIOS.cs
--- a/Engine/Source/Programs/AutomationTool/Gauntlet/Platform/IOS/Gauntlet.TargetDeviceIOS.cs
+++ b/Engine/Source/Programs/AutomationTool/Gauntlet/Platform/IOS/Gauntlet.TargetDeviceIOS.cs
@@ -173,7 +173,8 @@ namespace Gauntlet
 				{
 					// If some devices are detected, but none connected, attempt connection
 					UUID = AvailableDevices.First().UUID;
-					Log.Info("Attempting to connect to {UUID}", UUID);
+
+					Log.Info("Attempting to connect to {0}", UUID);
 					if (!Connect())
 					{
 						throw new DeviceException("Failed to connect to {0}. Ensure the device is properly connected to the host.", UUID);
@@ -183,6 +184,7 @@ namespace Gauntlet
 				{
 					// Just select the first available device
 					UUID = AvailableDevices.Where(Device => Device.bConnected).First().UUID;
+
 				}
 
 				Log.Verbose("Selected device {DeviceName} as default device", UUID);
@@ -197,13 +199,12 @@ namespace Gauntlet
 						throw new DeviceException("Device with UUID {0} not found in device list", UUID);
 					}
 				}
-				if (!ConnectedDevices[UUID].bConnected)
+
+				// Connect to device (or reconnect after reboot)
+				Log.Info("Attempting to connect to {0}", Name);
+				if (!Connect())
 				{
-					Log.Info("Attempting to connect to {DeviceName}", Name);
-					if (!Connect())
-					{
-						throw new DeviceException("Failed to connect to {0}", UUID);
-					}
+					throw new DeviceException("Failed to connect to {0}", UUID);
 				}
 			}
 
@@ -257,20 +258,6 @@ namespace Gauntlet
 					AppiumContainer.Dispose();
 				}
 
-				// Clean up any hanging devicectl/xcrun processes to prevent CoreDeviceService thread buildup
-				try
-				{
-					// Kill any xcrun processes that may be hanging
-					AppleBuild.ExecuteCommand("killall", "xcrun").WaitForExit();
-
-					// Give processes time to clean up
-					Thread.Sleep(500);
-				}
-				catch (Exception Ex)
-				{
-					Log.Verbose("Exception during cleanup in Dispose: {0}", Ex.Message);
-				}
-
 				DisposedValue = true;
 			}
 		}
@@ -282,80 +269,6 @@ namespace Gauntlet
 		}
 		#endregion
 
-		/// <summary>
-		/// Checks CoreDeviceService thread count and kills it if too high to prevent agent failure
-		/// </summary>
-		/// <returns>True if CoreDeviceService was killed and restarted</returns>
-		private static bool CheckAndCleanupCoreDeviceService(int MaxThreadCount = 500)
-		{
-			try
-			{
-				// Get thread count for CoreDeviceService
-				IProcessResult ThreadCheckResult = AppleBuild.ExecuteCommand("sh", "-c \"ps -M $(pgrep -x CoreDeviceService 2>/dev/null) 2>/dev/null | tail -n +2 | wc -l\"");
-				ThreadCheckResult.WaitForExit();
-
-				if (ThreadCheckResult.ExitCode == 0 && !string.IsNullOrWhiteSpace(ThreadCheckResult.Output))
-				{
-					string ThreadCountStr = ThreadCheckResult.Output.Trim();
-					if (int.TryParse(ThreadCountStr, out int ThreadCount))
-					{
-						Log.Verbose("CoreDeviceService thread count: {0}", ThreadCount);
-
-						if (ThreadCount > MaxThreadCount)
-						{
-							Log.Warning("CoreDeviceService has {0} threads (exceeds limit of {1}). Killing to prevent agent failure...", ThreadCount, MaxThreadCount);
-
-							// Try to kill CoreDeviceService daemon with SIGKILL - it will automatically restart
-							// First try with sudo for system daemon privileges
-							IProcessResult KillResult = AppleBuild.ExecuteCommand("sudo", "killall -9 CoreDeviceService");
-							KillResult.WaitForExit();
-
-							if (KillResult.ExitCode != 0)
-							{
-								Log.Warning("sudo killall failed with exit code {0}, trying without sudo: {1}", KillResult.ExitCode, KillResult.Output);
-								// Try without sudo but with -9 signal
-								KillResult = AppleBuild.ExecuteCommand("killall", "-9 CoreDeviceService");
-								KillResult.WaitForExit();
-								Thread.Sleep(5000);
-							}
-
-							if (KillResult.ExitCode == 0)
-							{
-								Log.Info("CoreDeviceService killed successfully, waiting for restart...");
-								// Wait for service to fully terminate and restart
-								Thread.Sleep(3000);
-
-								// Verify it restarted
-								IProcessResult VerifyResult = AppleBuild.ExecuteCommand("pgrep", "-x CoreDeviceService");
-								VerifyResult.WaitForExit();
-								if (VerifyResult.ExitCode == 0)
-								{
-									Log.Info("CoreDeviceService restarted successfully");
-									return true; // Service was killed and restarted
-								}
-								else
-								{
-									Log.Warning("CoreDeviceService may not have restarted, continuing anyway...");
-									return true; // Still treated as killed
-								}
-							}
-							else
-							{
-								Log.Warning("Failed to kill CoreDeviceService (exit code {0}): {1}", KillResult.ExitCode, KillResult.Output);
-								return false;
-							}
-						}
-					}
-				}
-			}
-			catch (Exception Ex)
-			{
-				Log.Verbose("Exception checking CoreDeviceService thread count: {0}", Ex.Message);
-			}
-
-			return false; // No cleanup needed or failed
-		}
-
 		#region ITargetDevice
 		public bool Connect()
 		{
@@ -374,20 +287,108 @@ namespace Gauntlet
 					return false;
 				}
 
-				string DetailsFile = Path.GetTempFileName();
-				string DetailsCommand = string.Format("device info details -j {0}", DetailsFile);
-				IProcessResult DetailResult = ExecuteDevicectlCommand(DetailsCommand);
-				DetailResult.WaitForExit();
-				if (DetailResult.ExitCode != 0)
+				// Wait after pair for device tunnel to establish before querying details
+				Log.Verbose("Waiting 20 seconds after pair for device tunnel to establish...");
+				Thread.Sleep(20000);
+
+				// Get device details with retry logic
+				// Note: ValidateDeviceConnection() will check tunnel state from device info details response
+				JsonObject JsonDevice = null;
+				const int MaxRetries = 2;
+				for (int attempt = 1; attempt <= MaxRetries; attempt++)
 				{
-					Log.Info("Detailing {DeviceName} exited with code {ExitCode}. Skipping connection:\n{Output}", UUID, DetailResult.ExitCode, ConnectResult.Output);
-				}
+					string DetailsFile = Path.GetTempFileName();
+					string DetailsCommand = string.Format("device info details -j {0}", DetailsFile);
+					IProcessResult DetailResult = ExecuteDevicectlCommand(DetailsCommand, 180);
+					DetailResult.WaitForExit();
+
+					if (DetailResult.ExitCode != 0)
+					{
+						Log.Warning("Detailing {DeviceName} attempt {Attempt} exited with code {ExitCode}:\n{Output}",
+							UUID, attempt, DetailResult.ExitCode, DetailResult.Output);
 
-				JsonObject JsonDevice = JsonObject.Parse(File.ReadAllText(DetailsFile)).GetObjectField("result");
-				bool bIsConnected = ValidateDeviceConnection(JsonDevice, true);
+						if (attempt < MaxRetries)
+						{
+							Log.Info("Retrying device info details after 10 second wait...");
+							Thread.Sleep(10000);
+							continue;
+						}
+						else
+						{
+							Log.Error("Device info details failed after {0} attempts", MaxRetries);
+							return false;
+						}
+					}
+
+					// Validate file exists and has content before parsing
+					if (!File.Exists(DetailsFile))
+					{
+						Log.Warning("Device info details file not found at {0} (attempt {1})", DetailsFile, attempt);
+						if (attempt < MaxRetries)
+						{
+							Thread.Sleep(10000);
+							continue;
+						}
+						return false;
+					}
+
+					FileInfo fileInfo = new FileInfo(DetailsFile);
+					if (fileInfo.Length == 0)
+					{
+						Log.Warning("Device info details file is empty for {0} (attempt {1})", UUID, attempt);
+						if (attempt < MaxRetries)
+						{
+							Thread.Sleep(10000);
+							continue;
+						}
+						return false;
+					}
 
-				if (!bIsConnected)
+					// Parse JSON and validate connection immediately with fresh data
+					try
+					{
+						JsonDevice = JsonObject.Parse(File.ReadAllText(DetailsFile)).GetObjectField("result");
+						Log.Verbose("Successfully retrieved device details for {0} on attempt {1}", UUID, attempt);
+
+						// Validate connection state immediately while data is fresh
+						bool bIsConnected = ValidateDeviceConnection(JsonDevice, true);
+
+						if (bIsConnected)
+						{
+							Log.Verbose("Device {0} connection validated successfully on attempt {1}", UUID, attempt);
+							break; // Success - device is connected
+						}
+						else
+						{
+							// Connection validation failed - retry if we have attempts left
+							if (attempt < MaxRetries)
+							{
+								Log.Info("Device {0} not connected (attempt {1}), waiting 10s before retry...", UUID, attempt);
+								Thread.Sleep(10000);
+								continue;
+							}
+							else
+							{
+								Log.Error("Device {0} connection validation failed after {1} attempts", UUID, MaxRetries);
+								return false;
+							}
+						}
+					}
+					catch (Exception ex)
+					{
+						Log.Warning("Failed to parse device details JSON for {0} (attempt {1}): {2}", UUID, attempt, ex.Message);
+						if (attempt < MaxRetries)
+						{
+							Thread.Sleep(10000);
+							continue;
+						}
+						throw;
+					}
+				}
+
+				if (JsonDevice == null)
 				{
+					Log.Error("Failed to retrieve valid device details for {0} after {1} attempts", UUID, MaxRetries);
 					return false;
 				}
 
@@ -452,7 +453,7 @@ namespace Gauntlet
 			string AppOutput;
 			if (UseDeviceCtl)
 			{
-				IProcessResult AppListResult = ExecuteDevicectlCommand("device info apps");
+				IProcessResult AppListResult = ExecuteDevicectlCommand("device info apps", 180);
 				if (AppListResult.ExitCode != 0)
 				{
 					throw new DeviceException("App list exited with code {0}: {1}", AppListResult.ExitCode, AppListResult.Output);
@@ -889,12 +890,17 @@ namespace Gauntlet
 			{
 				throw new DeviceException("Failed to copy uecommandline file to device!");
 			}
-
-			if (AppiumContainer != null)
+			Log.Info("Appium Enabled: {0} on {1}", Install.Name, this);
+			if (AppiumContainer.Enabled)
 			{
+				if (AppiumContainer == null)
+				{
+					AppiumContainer = new AppiumContainer(UUID);
+				}
 				AppiumContainer.Start(IOSInstall.PackageName);
 			}
 
+
 			Log.Info("Launching {0} on {1}", Install.Name, this);
 			Log.Verbose(IOSInstall.CommandLine);
 
@@ -929,15 +935,6 @@ namespace Gauntlet
 				throw new DeviceException("Failed to launch on {0}. {1}", Name, AppProcess.Output);
 			}
 
-			if (AppiumContainer != null)
-			{
-				lock (Mutex)
-				{
-					Thread.Sleep(5000);
-					AppiumContainer.CancelAppleSignOnAlert(IOSInstall.PackageName);
-				}
-			}
-
 			return new IOSAppInstance(IOSInstall, AppProcess, IOSInstall.CommandLine);
 		}
 
@@ -1132,9 +1129,6 @@ namespace Gauntlet
 				//AppleBuild.ExecuteCommand("killall", "usbmuxd").WaitForExit();
 				//AppleBuild.ExecuteCommand("launchctl", "kickstart -k system/com.apple.usbmuxd").WaitForExit();
 
-				// Check CoreDeviceService thread count and clean up if needed
-				CheckAndCleanupCoreDeviceService();
-
 				// Gather the iOS devices are connected to this UAT host
 				ConnectedDevices = GetConnectedDevices();
 
@@ -1147,16 +1141,10 @@ namespace Gauntlet
 
 		private static Dictionary<string, ConnectedDevice> GetConnectedDevices()
 		{
-			// Check CoreDeviceService thread count before making devicectl calls
-			bool ServiceWasKilled = CheckAndCleanupCoreDeviceService();
-
 			Dictionary<string, ConnectedDevice> Devices = GetConnectedDevicesInternal();
 
-			// If CoreDeviceService was just killed and no devices found, wait longer and retry once
-			if (ServiceWasKilled && (Devices == null || !Devices.Any()))
+			if ((Devices == null || !Devices.Any()))
 			{
-				Log.Warning("No devices found immediately after killing CoreDeviceService. Waiting 5 seconds and retrying...");
-				Thread.Sleep(5000);
 				Devices = GetConnectedDevicesInternal();
 
 				if (Devices == null || !Devices.Any())
@@ -1178,47 +1166,93 @@ namespace Gauntlet
 
 			if (UseDeviceCtl)
 			{
-				string ListFile = Path.GetTempFileName();
-				string ListCommand = string.Format("list devices -j {0}", ListFile);
+				const int MaxRetries = 3; // 1 original attempt + 2 retries
+				int attemptCount = 0;
+				bool success = false;
 
-				IProcessResult ListResult = ExecuteDevicectlCommand(ListCommand, null);
-				if (ListResult.ExitCode != 0)
+				while (attemptCount < MaxRetries && !success)
 				{
-					Log.Error("List devices exited with code {0}:\n{1}", ListResult.ExitCode, ListResult.Output);
-					return null;
-				}
+					attemptCount++;
+					if (attemptCount > 1)
+					{
+						Log.Warning("Retrying devicectl list devices (attempt {0}/{1})...", attemptCount, MaxRetries);
+						Thread.Sleep(3000); // Wait 3 seconds between retries
+					}
 
-				if (!File.Exists(ListFile))
-				{
-					Log.Error("Could not find list device output file at {FilePath}", ListFile);
-					return null;
-				}
+					string ListFile = Path.GetTempFileName();
+					string ListCommand = string.Format("list devices -j {0}", ListFile);
 
-				string Text = File.ReadAllText(ListFile);
-				JsonObject[] JsonDevices = JsonObject.Parse(File.ReadAllText(ListFile))
-					.GetObjectField("result")
-					.GetObjectArrayField("devices");
+					IProcessResult ListResult = ExecuteDevicectlCommand(ListCommand, null);
+					if (ListResult.ExitCode != 0)
+					{
+						Log.Warning("List devices attempt {0} exited with code {1}:\n{2}", attemptCount, ListResult.ExitCode, ListResult.Output);
+						if (attemptCount >= MaxRetries)
+						{
+							Log.Error("List devices failed after {0} attempts", MaxRetries);
+							return null;
+						}
+						continue; // Retry
+					}
 
-				foreach (JsonObject JsonDevice in JsonDevices)
-				{
-					ConnectedDevice IOSDevice = new ConnectedDevice(string.Empty, -1, false);
+					if (!File.Exists(ListFile))
+					{
+						Log.Warning("List device output file not found at {0} (attempt {1})", ListFile, attemptCount);
+						if (attemptCount >= MaxRetries)
+						{
+							Log.Error("Could not find list device output file after {0} attempts", MaxRetries);
+							return null;
+						}
+						continue; // Retry
+					}
 
-					if (JsonDevice.TryGetObjectField("hardwareProperties", out JsonObject HardwareProperties))
+					JsonObject[] JsonDevices = JsonObject.Parse(File.ReadAllText(ListFile))
+						.GetObjectField("result")
+						.GetObjectArrayField("devices");
+
+					foreach (JsonObject JsonDevice in JsonDevices)
 					{
-						HardwareProperties.TryGetStringField("udid", out IOSDevice.UUID);
+						ConnectedDevice IOSDevice = new ConnectedDevice(string.Empty, -1, false);
+
+						if (JsonDevice.TryGetObjectField("hardwareProperties", out JsonObject HardwareProperties))
+						{
+							HardwareProperties.TryGetStringField("udid", out IOSDevice.UUID);
+						}
+
+						if (JsonDevice.TryGetObjectField("deviceProperties", out JsonObject DeviceProperties) &&
+							DeviceProperties.TryGetStringField("osVersionNumber", out string OSVersionString) &&
+							!string.IsNullOrWhiteSpace(OSVersionString))
+						{
+							IOSDevice.IOSMajorVersion = int.Parse(OSVersionString.Substring(0, OSVersionString.IndexOf('.')));
+						}
+
+						if (!string.IsNullOrWhiteSpace(IOSDevice.UUID) && IOSDevice.IOSMajorVersion != -1)
+						{
+							IOSDevice.bConnected = ValidateDeviceConnection(JsonDevice);
+							Devices.Add(IOSDevice.UUID, IOSDevice);
+						}
 					}
 
-					if (JsonDevice.TryGetObjectField("deviceProperties", out JsonObject DeviceProperties) &&
-						DeviceProperties.TryGetStringField("osVersionNumber", out string OSVersionString) &&
-						!string.IsNullOrWhiteSpace(OSVersionString))
+					// Check if we found any devices
+					if (Devices.Count == 0)
 					{
-						IOSDevice.IOSMajorVersion = int.Parse(OSVersionString.Substring(0, OSVersionString.IndexOf('.')));
+						Log.Warning("devicectl list devices returned 0 devices (attempt {0}/{1})", attemptCount, MaxRetries);
+						if (attemptCount >= MaxRetries)
+						{
+							Log.Error("No devices found after {0} attempts", MaxRetries);
+							return null;
+						}
+						Devices.Clear(); // Clear for next retry
+						continue; // Retry
 					}
 
-					if (!string.IsNullOrWhiteSpace(IOSDevice.UUID) && IOSDevice.IOSMajorVersion != -1)
+					// Success - we have devices
+					success = true;
+
+					// Log summary of found devices
+					Log.Info("Found {0} device(s) from devicectl (attempt {1}):", Devices.Count, attemptCount);
+					foreach (var device in Devices.Values)
 					{
-						IOSDevice.bConnected = ValidateDeviceConnection(JsonDevice);
-						Devices.Add(IOSDevice.UUID, IOSDevice);
+						Log.Info("  - {0}: iOS {1}, Connected={2}", device.UUID, device.IOSMajorVersion, device.bConnected);
 					}
 				}
 			}
@@ -1465,7 +1499,7 @@ namespace Gauntlet
 			AppiumContainer.Config AppiumConfig = AppiumContainer.AppiumConfig;
 			if (!Globals.Params.ParseParam("NoAppiumScreenshots") && AppiumConfig != null && AppiumConfig.ScreenshotPeriod > 0)
 			{
-				IOSDevice.AppiumContainer.StopScreenshotLoop(Path.Combine(DestinationArtifactPath, "Screenshots"));
+				IOSDevice.AppiumContainer.StopMonitorLoop(Path.Combine(DestinationArtifactPath, "Screenshots"));
 			}
 		}
 	}
