Trivial: Fixing log level errors to be warnings

#Synced-CL 50214590

diff --git a/Engine/Source/Programs/AutomationTool/Gauntlet/Platform/IOS/Gauntlet.TargetDeviceIOS.cs b/Engine/Source/Programs/AutomationTool/Gauntlet/Platform/IOS/Gauntlet.TargetDeviceIOS.cs
--- a/Engine/Source/Programs/AutomationTool/Gauntlet/Platform/IOS/Gauntlet.TargetDeviceIOS.cs
+++ b/Engine/Source/Programs/AutomationTool/Gauntlet/Platform/IOS/Gauntlet.TargetDeviceIOS.cs
@@ -315,7 +315,7 @@ namespace Gauntlet
 						}
 						else
 						{
-							Log.Error("Device info details failed after {0} attempts", MaxRetries);
+							Log.Warning("Device info details failed after {0} attempts", MaxRetries);
 							return false;
 						}
 					}
@@ -369,7 +369,7 @@ namespace Gauntlet
 							}
 							else
 							{
-								Log.Error("Device {0} connection validation failed after {1} attempts", UUID, MaxRetries);
+								Log.Warning("Device {0} connection validation failed after {1} attempts", UUID, MaxRetries);
 								return false;
 							}
 						}
@@ -388,7 +388,7 @@ namespace Gauntlet
 
 				if (JsonDevice == null)
 				{
-					Log.Error("Failed to retrieve valid device details for {0} after {1} attempts", UUID, MaxRetries);
+					Log.Warning("Failed to retrieve valid device details for {0} after {1} attempts", UUID, MaxRetries);
 					return false;
 				}
 
@@ -1188,7 +1188,7 @@ namespace Gauntlet
 						Log.Warning("List devices attempt {0} exited with code {1}:\n{2}", attemptCount, ListResult.ExitCode, ListResult.Output);
 						if (attemptCount >= MaxRetries)
 						{
-							Log.Error("List devices failed after {0} attempts", MaxRetries);
+							Log.Warning("List devices failed after {0} attempts", MaxRetries);
 							return null;
 						}
 						continue; // Retry
@@ -1199,7 +1199,7 @@ namespace Gauntlet
 						Log.Warning("List device output file not found at {0} (attempt {1})", ListFile, attemptCount);
 						if (attemptCount >= MaxRetries)
 						{
-							Log.Error("Could not find list device output file after {0} attempts", MaxRetries);
+							Log.Warning("Could not find list device output file after {0} attempts", MaxRetries);
 							return null;
 						}
 						continue; // Retry
@@ -1238,7 +1238,7 @@ namespace Gauntlet
 						Log.Warning("devicectl list devices returned 0 devices (attempt {0}/{1})", attemptCount, MaxRetries);
 						if (attemptCount >= MaxRetries)
 						{
-							Log.Error("No devices found after {0} attempts", MaxRetries);
+							Log.Warning("No devices found after {0} attempts", MaxRetries);
 							return null;
 						}
 						Devices.Clear(); // Clear for next retry
