[Gauntlet] AutoParam now supports enumerables of simple types
#tests Local testing
#Synced-CL 50524817

diff --git a/Engine/Source/Programs/AutomationTool/Gauntlet/Framework/Gauntlet.AutoParam.cs b/Engine/Source/Programs/AutomationTool/Gauntlet/Framework/Gauntlet.AutoParam.cs
--- a/Engine/Source/Programs/AutomationTool/Gauntlet/Framework/Gauntlet.AutoParam.cs
+++ b/Engine/Source/Programs/AutomationTool/Gauntlet/Framework/Gauntlet.AutoParam.cs
@@ -1,12 +1,10 @@
 ﻿// Copyright Epic Games, Inc. All Rights Reserved.
 
+using AutomationTool;
 using System;
 using System.Collections.Generic;
 using System.Linq;
 using System.Reflection;
-using System.Text;
-using System.Threading.Tasks;
-using AutomationTool;
 
 namespace Gauntlet
 {
@@ -127,7 +125,7 @@ namespace Gauntlet
 		/// <param name="ParamType"></param>
 		/// <param name="Args"></param>
 		/// <returns></returns>
-		static protected object ParaseAndCoerceParam(string Param, Type ParamType, string[] Args)
+		static protected object ParseAndCoerceParam(string Param, Type ParamType, string[] Args)
 		{
 			if (!Param.EndsWith("="))
 			{
@@ -149,12 +147,11 @@ namespace Gauntlet
 					if (ParamType.IsEnum)
 					{
 						var AllValues = Enum.GetValues(ParamType).Cast<object>();
-
 						var Enums = AllValues.Where(P => string.Equals(StringVal, P.ToString(), StringComparison.OrdinalIgnoreCase));
 
 						if (Enums.Count() == 0)
 						{
-							Log.Error("Could not convert param {0} to enum of type {1}", StringVal, ParamType);
+							throw new AutomationException("AutoParam could not convert param {0} to enum of type {1}", StringVal, ParamType);
 						}
 
 						return Enums.First();
@@ -182,6 +179,38 @@ namespace Gauntlet
 			  || type.Equals(typeof(decimal));
 		}
 
+		/// <summary>
+		/// Returns true if this type inherits from IEnumerable<>
+		/// </summary>
+		/// <param name="Type">Underlying type of the object, e.g. IEnumerable<string></param>
+		/// <param name="ElementType">Enclosed type of the generic IEnumerable, e.g. string</param>
+		/// <returns></returns>
+		private static bool IsEnumerableType(Type Type, out Type ElementType)
+		{
+			ElementType = null;
+
+			if (Type.IsArray)
+			{
+				ElementType = Type.GetElementType();
+				return true;
+			}
+			else if (Type.IsGenericType)
+			{
+				Type[] GenericArgs = Type.GetGenericArguments();
+				if (GenericArgs.Length == 1)
+				{
+					Type EnumerableType = typeof(IEnumerable<>).MakeGenericType(GenericArgs[0]);
+					if (EnumerableType.IsAssignableFrom(Type))
+					{
+						ElementType = GenericArgs[0];
+						return true;
+					}
+				}
+			}
+
+			return false;
+		}
+
 		/// <summary>
 		/// Call to process all CommandLineOption attributes on an objects members and set them based on the
 		/// provided argument list
@@ -189,7 +218,7 @@ namespace Gauntlet
 		/// <param name="Obj"></param>
 		/// <param name="Args"></param>
 		/// <param name="ApplyDefaults"></param>
-		protected static void ApplyParamsAndDefaultsInternal(object Obj, string[] Args, bool ApplyDefaults)
+		protected static void ApplyParamsAndDefaultsInternal(object Obj, string[] Args, bool ApplyDefaults, char EnumerableSeparator = ',')
 		{
 			// get all field and property members
 			var Fields = Obj.GetType().GetFields(BindingFlags.Public | BindingFlags.NonPublic | BindingFlags.Instance);
@@ -214,67 +243,22 @@ namespace Gauntlet
 				// Go through all attributes
 				foreach (object Attrib in Member.GetCustomAttributes(true))
 				{
+					// Locate AutoParam
 					if (Attrib is AutoParam)
 					{
-						// If this is a struct then we want to recurse
-						if (IsSimple(MemberType) == false)
-						{
-							object Value = null;
-
-							// Get the reference value
-							if (Member is PropertyInfo)
-							{
-								Value = ((PropertyInfo)Member).GetValue(Obj);
-							}
-							else if (Member is FieldInfo)
-							{
-								Value = ((FieldInfo)Member).GetValue(Obj);
-							}
-
-							// if null create a new one (e.g. a new instance of a struct);
-							if (Value == null)
-							{
-								try
-								{
-									Value = Activator.CreateInstance(MemberType);
-								}
-								catch
-								{
-									throw new AutomationException("Add a default constructor to the class {0}!", MemberType);
-								}
+						AutoParam Opt = Attrib as AutoParam;
+						
+						// If the attribute had names provided use them, else use the name of the variable
+						string[] ParamNames = (Opt.OptionNames != null && Opt.OptionNames.Length > 0) ? Opt.OptionNames : [Member.Name];
 
-								// Set the new object as the refernce
-								if (Member is PropertyInfo)
-								{
-									((PropertyInfo)Member).SetValue(Obj, Value);
-								}
-								else if (Member is FieldInfo)
-								{
-									((FieldInfo)Member).SetValue(Obj, Value);
-								}
-							}
+						// Save the default
+						object DefaultValue = Opt.Default;
+						object NewValue = null;
 
-							// Recurse into this struct
-							ApplyParamsAndDefaultsInternal(Value, Args, ApplyDefaults);
-						}
-						else
+						// Go through all params used to refer to this member
+						foreach (string Name in ParamNames)
 						{
-							AutoParam Opt = Attrib as AutoParam;
-
-							// if the attribute had names provided use them, else use the name of the variable
-							string[] ParamNames = (Opt.OptionNames != null && Opt.OptionNames.Length > 0) ? Opt.OptionNames : new string[] { Member.Name };
-
-							// save the default
-							object DefaultValue = Opt.Default;
-							object NewValue = null;
-
-							if (DefaultValue != null && DefaultValue.GetType() != MemberType)
-							{
-								Log.Warning("AutoParam default value for member {0} is type {1}, not {2}", Member.Name, DefaultValue.GetType(), MemberType);
-							}
-
-							// Go through all params used to refer to this member
-							foreach (string Name in ParamNames)
+							if (IsSimple(MemberType))
 							{
 								if (MemberType == typeof(bool))
 								{
@@ -284,47 +268,39 @@ namespace Gauntlet
 								else
 								{
 									// for all other types try to parse out the value
-									NewValue = ParaseAndCoerceParam(Name, MemberType, Args);
-								}
-
-								// stop as soon as we find something
-								if (NewValue != null)
-								{
-									break;
+									NewValue = ParseAndCoerceParam(Name, MemberType, Args);
 								}
 							}
-
-							// If no value was found, use the default
-							if (NewValue != null || (ApplyDefaults && DefaultValue != null))
+							else if (IsEnumerableType(MemberType, out Type ElementType))
 							{
-								if (NewValue == null)
+								string StringValue = ParseAndCoerceParam(Name, typeof(string), Args) as string;
+								if (!string.IsNullOrEmpty(StringValue))
 								{
-									NewValue = DefaultValue;
-								}
-
-								if (MemberType.IsEnum)
-								{
-									if (NewValue.GetType() != MemberType)
-									{
-										Log.Warning("Default for member {0} is an enum of an incorrect type!", Member.Name);
-									}
-								}
-								else
-								{
-									// Force a conversion - need to do this even for default values incase someone specified a double for a float
-									NewValue = Convert.ChangeType(NewValue, MemberType);
-								}
-
-								// Set the value on the member
-								if (Member is PropertyInfo)
-								{
-									((PropertyInfo)Member).SetValue(Obj, NewValue);
-								}
-								else if (Member is FieldInfo)
-								{
-									((FieldInfo)Member).SetValue(Obj, NewValue);
+									NewValue = ConstructEnumerableFromString(Member, MemberType, ElementType, StringValue, EnumerableSeparator);
 								}
 							}
+							else
+							{
+								NewValue = ConstructComplexObject(Args, Member, MemberType, Obj, ApplyDefaults, EnumerableSeparator);
+							}
+							
+							// Stop as soon as we find something
+							// TODO support concatenation of enumerables?
+							if (NewValue != null)
+							{
+								break;
+							}
+						}
+		
+						if (NewValue == null && ApplyDefaults)
+						{
+							// No commandline arguments were found that matched the name(s) of this AutoParam. Apply the defaults
+							ApplyDefaultValue(Opt, Member, MemberType, Obj, EnumerableSeparator);
+						}
+						else if (NewValue != null)
+						{
+							// Value was overriden by commandline argument. Set the new member value
+							SetMemberValue(Member, Obj, NewValue);
 						}
 					}
 				}
@@ -338,19 +314,151 @@ namespace Gauntlet
 			}
 		}
 
-		public static void ApplyDefaults(object Obj)
+		public static void ApplyDefaults(object Obj, char EnumerableSeparator = ',')
+		{
+			ApplyParamsAndDefaultsInternal(Obj, new string[0], true, EnumerableSeparator);
+		}
+
+		public static void ApplyParams(object Obj, string[] Args, char EnumerableSeparator = ',')
+		{
+			ApplyParamsAndDefaultsInternal(Obj, Args, false, EnumerableSeparator);
+		}
+
+		public static void ApplyParamsAndDefaults(object Obj, string[] Args, char EnumerableSeparator = ',')
+		{
+			ApplyParamsAndDefaultsInternal(Obj, Args, true, EnumerableSeparator);
+		}
+
+		private static void SetMemberValue(MemberInfo Member, object Obj, object Value)
+		{
+			if (Member is PropertyInfo)
+			{
+				((PropertyInfo)Member).SetValue(Obj, Value);
+			}
+			else if (Member is FieldInfo)
+			{
+				((FieldInfo)Member).SetValue(Obj, Value);
+			}
+		}
+
+		private static void ApplyDefaultValue(AutoParam Param, MemberInfo Member, Type MemberType, object Obj, char EnumerableSeparator = ',')
 		{
-			ApplyParamsAndDefaultsInternal(Obj, new string[0], true);
+			object DefaultValue = Param.Default;
+
+			if (IsEnumerableType(MemberType, out Type ElementType) && DefaultValue is string DefaultString)
+			{
+				DefaultValue = ConstructEnumerableFromString(Member, MemberType, ElementType, DefaultString, EnumerableSeparator);
+			}
+
+			if (DefaultValue != null && DefaultValue.GetType() != MemberType)
+			{
+				throw new AutomationException("AutoParam Default Value for member {0} is type {1}, not {2}", Member.Name, DefaultValue.GetType(), MemberType);
+			}
+
+			SetMemberValue(Member, Obj, DefaultValue);
 		}
 
-		public static void ApplyParams(object Obj, string[] Args)
+		private static object ConstructComplexObject(string[] Args, MemberInfo Member, Type MemberType, object Obj, bool bApplyDefaults, char EnumerableSeparator = ',')
 		{
-			ApplyParamsAndDefaultsInternal(Obj, Args, false);
+			object ComplexObject = null;
+
+			// Check if the object has already been constructed
+			if (Member is PropertyInfo)
+			{
+				ComplexObject = ((PropertyInfo)Member).GetValue(Obj);
+			}
+			else if (Member is FieldInfo)
+			{
+				ComplexObject = ((FieldInfo)Member).GetValue(Obj);
+			}
+
+			// If null create a new one (e.g. a new instance of a struct);
+			if (ComplexObject == null)
+			{
+				try
+				{
+					ComplexObject = Activator.CreateInstance(MemberType);
+				}
+				catch
+				{
+					throw new AutomationException("AutoParam was unable to construct an instance of {0}. Add a default constructor to the class {0}!", MemberType);
+				}
+			}
+
+			// Recurse into the object applying autoparam values
+			ApplyParamsAndDefaultsInternal(ComplexObject, Args, bApplyDefaults, EnumerableSeparator);
+
+			return ComplexObject;
 		}
 
-		public static void ApplyParamsAndDefaults(object Obj, string[] Args)
+		private static object ConstructEnumerableFromString(MemberInfo Member, Type MemberType, Type ElementType, string Value, char EnumerableSeparator = ',')
 		{
-			ApplyParamsAndDefaultsInternal(Obj, Args, true);
+			if (string.IsNullOrEmpty(Value))
+			{
+				return null;
+			}
+			else
+			{
+				// Split and trim entries
+				string[] SplitValues = Value.Split(EnumerableSeparator, StringSplitOptions.RemoveEmptyEntries | StringSplitOptions.TrimEntries);
+
+				// Construct an array of the element type which will be converted to the underlying enumerable type
+				Array ConvertedArray = Array.CreateInstance(ElementType, SplitValues.Length);
+				for (int Index = 0; Index < SplitValues.Length; Index++)
+				{
+					if (IsSimple(ElementType))
+					{
+						if (ElementType.IsEnum)
+						{
+							ConvertedArray.SetValue(Enum.Parse(ElementType, SplitValues[Index], true), Index);
+						}
+						else
+						{
+							ConvertedArray.SetValue(Convert.ChangeType(SplitValues[Index], ElementType), Index);
+						}
+					}
+					else
+					{
+						throw new AutomationException("Member {0} of type {1} uses an invalid element type {2}. AutoParam only supports enumerables of simple types (int, string, etc)", Member.Name, MemberType, ElementType);
+					}
+				}
+
+				object ConstructedEnumerable = null;
+
+				// Convert the array to the appropriate collection type
+				if (MemberType.IsArray)
+				{
+					ConstructedEnumerable = ConvertedArray;
+				}
+				else if (MemberType.IsGenericType)
+				{
+					Type GenericTypeDef = MemberType.GetGenericTypeDefinition();
+
+					// Handle common list/collection interfaces and types
+					if (GenericTypeDef == typeof(List<>)
+						|| MemberType == typeof(IList<>).MakeGenericType(ElementType)
+						|| MemberType == typeof(ICollection<>).MakeGenericType(ElementType)
+						|| MemberType == typeof(IEnumerable<>).MakeGenericType(ElementType))
+					{
+						Type EnumerableType = typeof(List<>).MakeGenericType(ElementType);
+						ConstructedEnumerable = Activator.CreateInstance(EnumerableType, [ConvertedArray]);
+					}
+					else
+					{
+						// Try to create an instance using a constructor that accepts IEnumerable<T>
+						try
+						{
+							ConstructedEnumerable = Activator.CreateInstance(MemberType, [ConvertedArray]);
+						}
+						catch
+						{
+							throw new AutomationException("AutoParam is unable to create an instance of {0} from an array. Ensure it has a constructor accepting IEnumerable<{1}> if you wish to use AutoParam.", MemberType, ElementType);
+						}
+					}
+				}
+
+				return ConstructedEnumerable;
+			}
 		}
 
 		/*public static string[] GetParams(object Obj)
