Fixing Appium alert dismissing, adjusting the thread checking for alerts to check more often, take screenshot prior to dismiss.

Also adjusting logic around not uninstalling webdriver app if it's already installed.

#Synced-CL 50214590

diff --git a/Engine/Source/Programs/AutomationTool/Gauntlet/Platform/IOS/Gauntlet.IOSAppium.cs b/Engine/Source/Programs/AutomationTool/Gauntlet/Platform/IOS/Gauntlet.IOSAppium.cs
--- a/Engine/Source/Programs/AutomationTool/Gauntlet/Platform/IOS/Gauntlet.IOSAppium.cs
+++ b/Engine/Source/Programs/AutomationTool/Gauntlet/Platform/IOS/Gauntlet.IOSAppium.cs
@@ -99,7 +99,8 @@ namespace Gauntlet
 				Capabilities.AutomationName = AutomationName;
 				Capabilities.AddAdditionalAppiumOption("udid", UUID);
 				Capabilities.AddAdditionalAppiumOption("bundleId", PackageName);
-				Capabilities.AddAdditionalAppiumOption("autoAcceptAlerts", true);
+				// Use autoDismissAlerts instead of autoAcceptAlerts to click Cancel/Dismiss buttons
+				Capabilities.AddAdditionalAppiumOption("autoDismissAlerts", true);
 				Capabilities.AddAdditionalAppiumOption("xcodeOrgId", OrgId);
 				Capabilities.AddAdditionalAppiumOption("xcodeSigningId", SigningId);
 				Capabilities.AddAdditionalAppiumOption("autoLaunch", false);
@@ -233,32 +234,43 @@ namespace Gauntlet
 
 			this.PackageName = PackageName;
 
-			// Install WebdriverAgent
-			try
+			// Check if WebDriverAgent is already installed
+			bool isWdaInstalled = CheckIfWdaInstalled(UUID, AppiumConfig.WdaBundleId);
+			if (isWdaInstalled)
+			{
+				Log.Info("WebDriverAgent ({0}) already installed on device {1}, skipping install", AppiumConfig.WdaBundleId, UUID);
+			}
+			else
 			{
-				IProcessResult InstallResult;
-				using (ScopedSuspendECErrorParsing ErrorSuspension = new())
+				// Install WebdriverAgent
+				try
 				{
-					if (TargetDeviceIOS.UseDeviceCtl)
+					Log.Info("Installing WebDriverAgent ({0}) on device {1}...", AppiumConfig.WdaBundleId, UUID);
+					IProcessResult InstallResult;
+					using (ScopedSuspendECErrorParsing ErrorSuspension = new())
 					{
-						InstallResult = TargetDeviceIOS.ExecuteDevicectlCommand(
-							string.Format("device install app \"{0}\" -v", AppiumConfig.WdaAppPath), UUID, 60, AdditionalOptions: ERunOptions.NoStdOutRedirect);
+						if (TargetDeviceIOS.UseDeviceCtl)
+						{
+							InstallResult = TargetDeviceIOS.ExecuteDevicectlCommand(
+								string.Format("device install app \"{0}\" -v", AppiumConfig.WdaAppPath), UUID, 60, AdditionalOptions: ERunOptions.NoStdOutRedirect);
+						}
+						else
+						{
+							InstallResult = TargetDeviceIOS.ExecuteIOSDeployCommand(
+								string.Format("-b \"{0}\"", AppiumConfig.WdaAppPath), UUID, 60, AdditionalOptions: ERunOptions.NoStdOutRedirect);
+						}
 					}
-					else
+					if (InstallResult.ExitCode != 0)
 					{
-						InstallResult = TargetDeviceIOS.ExecuteIOSDeployCommand(
-							string.Format("-b \"{0}\"", AppiumConfig.WdaAppPath), UUID, 60, AdditionalOptions: ERunOptions.NoStdOutRedirect);
+						throw new AutomationException("Failed to install WDA app ({0}): {1}", InstallResult.ExitCode, InstallResult.Output);
 					}
+					Log.Info("WebDriverAgent installed successfully on device {0}", UUID);
 				}
-				if (InstallResult.ExitCode != 0)
+				catch (Exception Ex)
 				{
-					throw new AutomationException("Failed to install WDA app ({0}): {1}", InstallResult.ExitCode, InstallResult.Output);
+					throw new AutomationException($"Exception during WDA Install {Ex.Message}");
 				}
 			}
-			catch (Exception Ex)
-			{
-				throw new AutomationException($"Exception during WDA Install {Ex.Message}");
-			}
 
 			lock (Mutex)
 			{
@@ -321,6 +333,49 @@ namespace Gauntlet
 			// TODO: Terminate WDA app - parse devicectl for pid, then kill?
 		}
 
+		/// <summary>
+		/// Checks if WebDriverAgent is already installed on the device
+		/// </summary>
+		/// <param name="deviceUUID">UUID of the device to check</param>
+		/// <param name="bundleId">Bundle ID of WebDriverAgent (e.g., com.epicgames.WebDriverAgent)</param>
+		/// <returns>True if WDA is installed, false otherwise</returns>
+		private bool CheckIfWdaInstalled(string deviceUUID, string bundleId)
+		{
+			try
+			{
+				IProcessResult checkResult;
+
+				if (TargetDeviceIOS.UseDeviceCtl)
+				{
+					// Use devicectl to list installed apps
+					checkResult = TargetDeviceIOS.ExecuteDevicectlCommand("device info apps", deviceUUID, 30);
+				}
+				else
+				{
+					// Use ios-deploy to list bundle IDs
+					checkResult = TargetDeviceIOS.ExecuteIOSDeployCommand("--list_bundle_id", deviceUUID, 30);
+				}
+
+				if (checkResult.ExitCode == 0)
+				{
+					// Check for both the base bundle ID and the .xctrunner variant
+					// (actual installed app often has .xctrunner suffix)
+					bool isInstalled = checkResult.Output.Contains(bundleId) ||
+					                   checkResult.Output.Contains(bundleId + ".xctrunner");
+					Log.Verbose("Checked for {0} (or .xctrunner variant) on device {1}: {2}", bundleId, deviceUUID, isInstalled ? "found" : "not found");
+					return isInstalled;
+				}
+
+				Log.Verbose("Failed to check installed apps on device {0}, assuming WDA not installed", deviceUUID);
+				return false;
+			}
+			catch (Exception ex)
+			{
+				Log.Verbose("Exception checking for installed WDA: {0}. Assuming not installed.", ex.Message);
+				return false; // If check fails, assume not installed and proceed with install
+			}
+		}
+
 		private void ConfigureDrivers()
 		{
 			// Query installed drivers
@@ -414,9 +469,10 @@ namespace Gauntlet
 
 		/// <summary>
 		/// Checks for and dismisses Apple Sign-On alerts. Used by the monitor loop.
+		/// Takes a screenshot if an alert is found before dismissing it.
 		/// Only performs 1 retry attempt.
 		/// </summary>
-		private void CheckAndDismissAlert()
+		private void CheckAndDismissAlert(string ScreenshotDirectory)
 		{
 			var WebDriverWait = new WebDriverWait(new SystemClock(), Driver, TimeSpan.FromMilliseconds(5000), TimeSpan.FromMilliseconds(250));
 
@@ -443,6 +499,28 @@ namespace Gauntlet
 				{
 					Log.Info("Found Alert, trying to dismiss");
 
+					// Take screenshot before dismissing
+					try
+					{
+						if (!string.IsNullOrEmpty(ScreenshotDirectory))
+						{
+							if (!DirectoryExists(ScreenshotDirectory))
+							{
+								CreateDirectory(ScreenshotDirectory);
+							}
+
+							string FileName = $"alert-{DateTime.Now.ToString("yy-MM-dd-HH-mm-ss-ff")}.png";
+							string FilePath = Path.Combine(ScreenshotDirectory, FileName);
+							(Driver as ITakesScreenshot).GetScreenshot().SaveAsFile(FilePath);
+							Log.Info("Saved alert screenshot to: {0}", FilePath);
+						}
+					}
+					catch (Exception Ex)
+					{
+						Log.Verbose("Failed to capture alert screenshot: {0}", Ex.Message);
+					}
+
+					// Now dismiss the alert
 					// Retry find+click operation to handle stale elements
 					bool clicked = false;
 					for (int attempt = 0; attempt < 3 && !clicked; attempt++)
@@ -467,10 +545,14 @@ namespace Gauntlet
 						}
 					}
 
-					if (clicked && !string.IsNullOrEmpty(PackageName))
+					if (clicked)
 					{
-						Thread.Sleep(300);
-						Driver.ActivateApp(PackageName);
+						Log.Info("Monitor thread dismissed an alert.");
+						if (!string.IsNullOrEmpty(PackageName))
+						{
+							Thread.Sleep(300);
+							Driver.ActivateApp(PackageName);
+						}
 					}
 				}
 			}
@@ -488,6 +570,78 @@ namespace Gauntlet
 			}
 		}
 
+		/// <summary>
+		/// Checks for view-based popups that aren't detected as system alerts.
+		/// Looks for common dismiss buttons like Cancel, Try Again, etc.
+		/// </summary>
+		private bool CheckAndDismissViewPopup(string ScreenshotDirectory)
+		{
+			try
+			{
+				// Look for common popup button combinations
+				// This catches popups that aren't system alerts but are view elements
+				// NOTE: Intentionally excludes 'OK' buttons to avoid accidentally confirming dialogs
+				var popupButtons = Driver.FindElements(By.XPath(
+					"//XCUIElementTypeButton[@name='Cancel' or @name='Try Again' or @name='Close' or @name='Dismiss']"
+				));
+
+				if (popupButtons.Count >= 2) // Likely a popup if multiple buttons present
+				{
+					Log.Info("Found view-based popup with {0} buttons", popupButtons.Count);
+
+					// Take screenshot before dismissing
+					try
+					{
+						if (!string.IsNullOrEmpty(ScreenshotDirectory) && DirectoryExists(ScreenshotDirectory))
+						{
+							string FileName = $"popup-{DateTime.Now.ToString("yy-MM-dd-HH-mm-ss-ff")}.png";
+							string FilePath = Path.Combine(ScreenshotDirectory, FileName);
+							(Driver as ITakesScreenshot).GetScreenshot().SaveAsFile(FilePath);
+							Log.Info("Saved popup screenshot to: {0}", FilePath);
+						}
+					}
+					catch (Exception Ex)
+					{
+						Log.Verbose("Failed to capture popup screenshot: {0}", Ex.Message);
+					}
+
+					// Click Cancel/Dismiss button if present
+					foreach (var button in popupButtons)
+					{
+						string buttonName = button.GetAttribute("name");
+						// Use case-insensitive comparison and trim whitespace for robustness
+						if (!string.IsNullOrEmpty(buttonName))
+						{
+							string trimmedName = buttonName.Trim();
+							if (trimmedName.Equals("Cancel", StringComparison.OrdinalIgnoreCase) ||
+							    trimmedName.Equals("Dismiss", StringComparison.OrdinalIgnoreCase) ||
+							    trimmedName.Equals("Close", StringComparison.OrdinalIgnoreCase))
+							{
+								Log.Info("Clicking dismiss button on popup: {0}", buttonName);
+								button.Click();
+								if (!string.IsNullOrEmpty(PackageName))
+								{
+									Thread.Sleep(300);
+									Driver.ActivateApp(PackageName);
+								}
+								return true;
+							}
+						}
+					}
+
+					// No safe dismiss button found - log and don't click anything
+					Log.Warning("Found popup with {0} buttons but no Cancel/Dismiss/Close button. Not clicking to avoid confirming dialog.", popupButtons.Count);
+					return false;
+				}
+			}
+			catch (Exception ex)
+			{
+				Log.Verbose("View popup check failed: {0}", ex.Message);
+			}
+
+			return false;
+		}
+
 		/// <summary>
 		/// Uses appium WebDriver to take a screenshot. 
 		/// Filename is saved as png in the format of: yy-MM-dd-HH-mm-ss-ff
@@ -520,9 +674,16 @@ namespace Gauntlet
 
 		private void StartMonitorLoop()
 		{
-			TimeSpan Period = TimeSpan.FromSeconds(AppiumConfig.ScreenshotPeriod);
-			MonitorTrigger = new AutoResetEvent(false);
+			// Alert check runs frequently (every 2 seconds)
+			TimeSpan AlertCheckPeriod = TimeSpan.FromSeconds(2);
 
+			// Calculate how many alert checks happen per screenshot
+			// e.g., if ScreenshotPeriod = 30 seconds, and AlertCheckPeriod = 2 seconds
+			// then screenshotCheckCount = 30 / 2 = 15 iterations
+			int screenshotCheckCount = (int)Math.Max(1, AppiumConfig.ScreenshotPeriod / AlertCheckPeriod.TotalSeconds);
+			int iterationCounter = 0;
+
+			MonitorTrigger = new AutoResetEvent(false);
 			string ScreenshotDirectory = GetScreenshotDirectory();
 
 			if (DirectoryExists(ScreenshotDirectory))
@@ -530,20 +691,37 @@ namespace Gauntlet
 				DeleteDirectory_NoExceptions(ScreenshotDirectory);
 			}
 
-			while (!MonitorTrigger.WaitOne(Period))
+			Log.Verbose("Monitor loop: Alert checks every {0}s, Screenshots every {1}s ({2} iterations)",
+				AlertCheckPeriod.TotalSeconds, AppiumConfig.ScreenshotPeriod, screenshotCheckCount);
+
+			while (!MonitorTrigger.WaitOne(AlertCheckPeriod))
 			{
-				// Take screenshot
-				TakeScreenshot(ScreenshotDirectory);
+				// Check for system alerts first
+				try
+				{
+					CheckAndDismissAlert(ScreenshotDirectory);
+				}
+				catch (Exception Ex)
+				{
+					Log.Verbose("Alert check failed: {0}", Ex.Message);
+				}
 
-				// Check for alerts first (more critical)
-				Log.Info("Trying to find Alert");
+				// Fallback: Check for view-based popups
 				try
 				{
-					CheckAndDismissAlert();
+					CheckAndDismissViewPopup(ScreenshotDirectory);
 				}
 				catch (Exception Ex)
 				{
-					Log.Info("Alert check failed: {0}", Ex.Message);
+					Log.Verbose("Popup check failed: {0}", Ex.Message);
+				}
+
+				// Only take screenshot every N iterations
+				iterationCounter++;
+				if (iterationCounter >= screenshotCheckCount)
+				{
+					TakeScreenshot(ScreenshotDirectory);
+					iterationCounter = 0;
 				}
 			}
 		}
diff --git a/Engine/Source/Programs/AutomationTool/Gauntlet/Platform/IOS/Gauntlet.TargetDeviceIOS.cs b/Engine/Source/Programs/AutomationTool/Gauntlet/Platform/IOS/Gauntlet.TargetDeviceIOS.cs
--- a/Engine/Source/Programs/AutomationTool/Gauntlet/Platform/IOS/Gauntlet.TargetDeviceIOS.cs
+++ b/Engine/Source/Programs/AutomationTool/Gauntlet/Platform/IOS/Gauntlet.TargetDeviceIOS.cs
@@ -444,11 +444,12 @@ namespace Gauntlet
 
 		public void FullClean()
 		{
-			List<string> AppAllowList = new List<string>
+			// Patterns for apps that should NOT be uninstalled (supports wildcard matching via EndsWith)
+			List<string> AppAllowListPatterns = new List<string>
 			{
 				"WebDriverAgent.xctrunner",
-				"com.mobilelabsinc.deviceControl.xctrunner",
-				"com.mobilelabsinc.deviceControl"
+				"deviceControl.xctrunner",
+				"deviceControl"
 			};
 			string AppOutput;
 			if (UseDeviceCtl)
@@ -472,12 +473,24 @@ namespace Gauntlet
 				AppOutput = AppListResult.Output;
 			}
 
-			Regex AppRegex = new Regex("com[^\\s]*"); // Ex. com.epicgames.enginetest
-			IEnumerable<string> AppList = AppRegex.Matches(AppOutput)
+			// Match bundle IDs - handle both raw format and "bundleID: com.xxx" format
+			Regex AppRegex = new Regex(@"(?:bundleID:\s*)?(com\.[^\s,\)]+)", RegexOptions.IgnoreCase);
+			var allApps = AppRegex.Matches(AppOutput)
 				.Where(Match => Match.Success)
-				.Select(Match => Match.Value)
+				.Select(Match => Match.Groups[1].Value) // Get the captured group (the bundle ID)
 				.Where(Match => !Match.Contains("com.apple.", StringComparison.OrdinalIgnoreCase))
-				.Where(Match => !AppAllowList.Contains(Match));
+				.ToList();
+
+			IEnumerable<string> AppList = allApps
+				.Where(app =>
+				{
+					bool isProtected = AppAllowListPatterns.Any(pattern => app.EndsWith(pattern, StringComparison.OrdinalIgnoreCase));
+					if (isProtected)
+					{
+						Log.Verbose("Skipping protected app: {0}", app);
+					}
+					return !isProtected;
+				});
 
 			foreach (string App in AppList)
 			{
