Trivial change to IOS clean artifacts ios-deploy change to warning as it's unreliable.
#Synced-CL 50214590

diff --git a/Engine/Source/Programs/AutomationTool/Gauntlet/Platform/IOS/Gauntlet.TargetDeviceIOS.cs b/Engine/Source/Programs/AutomationTool/Gauntlet/Platform/IOS/Gauntlet.TargetDeviceIOS.cs
--- a/Engine/Source/Programs/AutomationTool/Gauntlet/Platform/IOS/Gauntlet.TargetDeviceIOS.cs
+++ b/Engine/Source/Programs/AutomationTool/Gauntlet/Platform/IOS/Gauntlet.TargetDeviceIOS.cs
@@ -551,7 +551,7 @@ namespace Gauntlet
 			}
 			else if (Result.ExitCode != 0)
 			{
-				throw new DeviceException("Listing Engine artifacts for {0} exited with code {1}:\n{2}", BundleID, Result.ExitCode, Result.Output);
+				Log.Warning(String.Format("Listing Engine artifacts for {0} exited with code {1}:\n{2}", BundleID, Result.ExitCode, Result.Output));
 			}
 
 			if (Result.Output.Contains(EngineSavedDirectory))
@@ -560,7 +560,7 @@ namespace Gauntlet
 				Result.WaitForExit();
 				if (Result.ExitCode != 0)
 				{
-					throw new DeviceException("Cleaning Engine artifacts for {0} exited with code {1}:\n{2}", BundleID, Result.ExitCode, Result.Output);
+					Log.Warning(String.Format("Cleaning Engine artifacts for {0} exited with code {1}:\n{2}", BundleID, Result.ExitCode, Result.Output));
 				}
 			}
 
@@ -571,7 +571,7 @@ namespace Gauntlet
 			Result.WaitForExit();
 			if (Result.ExitCode != 0)
 			{
-				throw new DeviceException("Listing Project artifacts for {0} exited with code {1}:\n{2}", BundleID, Result.ExitCode, Result.Output);
+				Log.Warning(String.Format("Listing Project artifacts for {0} exited with code {1}:\n{2}", BundleID, Result.ExitCode, Result.Output));
 			}
 
 			// The same applies here for the --list
@@ -581,8 +581,8 @@ namespace Gauntlet
 				Result.WaitForExit();
 				if (Result.ExitCode != 0)
 				{
-					throw new DeviceException("Cleaning Project artifacts for {0} exited with code {1}:\n{2}", BundleID, Result.ExitCode,
-						Result.Output);
+					Log.Warning(String.Format("Cleaning Project artifacts for {0} exited with code {1}:\n{2}", BundleID, Result.ExitCode, Result.Output));
+
 				}
 			}
 		}
