Added a placeholder remote output capture for Windows remote devices
- gauntlet and uat output now functional

#rnx

#Synced-CL 50857758

diff --git a/Engine/Source/Programs/AutomationTool/Gauntlet/Platform/Windows/Gauntlet.TargetDeviceRemoteWin.cs b/Engine/Source/Programs/AutomationTool/Gauntlet/Platform/Windows/Gauntlet.TargetDeviceRemoteWin.cs
--- a/Engine/Source/Programs/AutomationTool/Gauntlet/Platform/Windows/Gauntlet.TargetDeviceRemoteWin.cs
+++ b/Engine/Source/Programs/AutomationTool/Gauntlet/Platform/Windows/Gauntlet.TargetDeviceRemoteWin.cs
@@ -1,7 +1,6 @@
 // Copyright Epic Games, Inc. All Rights Reserved.
 
 using AutomationTool;
-using System;
 using System.Collections.Generic;
 using System.IO;
 using UnrealBuildTool;
@@ -119,12 +118,19 @@ namespace Gauntlet
 
 			Log.Info("Launching {0} on {1}", WinApp.Name, ToString());
 
-			if (!RemoteWinSupport.Run(Name, WinApp.ExecutablePath, WinApp.CommandLine))
+			string ExtraCmdLine = "";
+			RemoteWinSupport.Placeholder_RemoteWinProcessMonitor ProcessMonitor = new(Name, EpicGames.Core.Log.Logger);
+			if (ProcessMonitor.IsAvailable)
+			{
+				ExtraCmdLine += $" -RemoteConsoleHost={ProcessMonitor.ServerAddress}";
+			}
+
+			if (!RemoteWinSupport.Run(Name, WinApp.ExecutablePath, WinApp.CommandLine + ExtraCmdLine))
 			{
 				throw new AutomationException("Failed to launch {0}", WinApp.ExecutablePath);
 			}
 
-			return new RemoteWinAppInstance(this, WinApp.CommandLine);
+			return new RemoteWinAppInstance(this, WinApp, ProcessMonitor);
 		}
 		#endregion
 
@@ -179,12 +185,35 @@ namespace Gauntlet
 	{
 		private TargetDeviceRemoteWin TargetDevice { get; set; }
 
-		public RemoteWinAppInstance(TargetDeviceRemoteWin InTargetDevice, string InCommandLine)
+		private readonly RemoteWinSupport.Placeholder_RemoteWinProcessMonitor ProcessMonitor;
+
+		private readonly object LoggerSyncObject = null;
+		private readonly FileLogger Logger = null;
+
+
+		public RemoteWinAppInstance(TargetDeviceRemoteWin InTargetDevice, RemoteWinAppInstall WinApp, RemoteWinSupport.Placeholder_RemoteWinProcessMonitor InProcessMonitor)
 		{
-			CommandLine = InCommandLine;
+			CommandLine = WinApp.CommandLine;
 			TargetDevice = InTargetDevice;
+			ProcessMonitor = InProcessMonitor;
+			LoggerSyncObject = new object();
+			
+			if (ProcessMonitor.IsAvailable)
+			{
+				string LogFilename = ProcessUtils.GetLogFilePath(WinApp.Name, ProcessUtils.LocalLogsPath);
+				if (!string.IsNullOrEmpty(LogFilename))
+				{
+					Logger = new FileLogger(LogFilename, CommandLine);
+				}
+
+				ProcessMonitor.OutputReceived = OnOutputReceived;
+				ProcessMonitor.Start();
+			}
+		}
 
-			HasExited = true; // @todo: assume it has exited already because we can't track it yet
+		private void OnOutputReceived(string Line)
+		{
+			Logger?.AppendLine(Line);
 		}
 
 		public void Kill(bool GenerateDumpOnKill = false)
@@ -192,24 +221,80 @@ namespace Gauntlet
 			if (!HasExited)
 			{
 				RemoteWinSupport.Kill(TargetDevice.Name);
+				ProcessMonitor.Stop();
 				WasKilled = true;
-				HasExited = true;
 			}
 		}
+
 		public ITargetDevice Device => TargetDevice;
-		public bool HasExited { get; private set; }
+		public bool HasExited
+		{
+			get
+			{
+				bool bHasExited = !ProcessMonitor.IsAvailable;
+
+				if (bHasExited)
+				{
+					lock(LoggerSyncObject)
+					{
+						Logger?.CloseStream();
+					}
+				}
+
+				return bHasExited;
+			}
+		}
 		public bool WasKilled { get; private set; }
 		public string CommandLine { get; protected set; }
 
 
-		// @todo: these are stubs. we can't track process lifetime or read output yet
-		public string StdOut => "";
-		public int ExitCode => 0;
+		public string StdOut
+		{
+			get
+			{
+				lock (LoggerSyncObject)
+				{
+					return Logger?.GetReader().GetContent();
+				}
+			}
+		}
+
+		public int ExitCode
+		{
+			get
+			{
+				if (WasKilled)
+				{
+					return 1; // process was terminated
+				}
+				else if (ProcessMonitor.DidLaunch)
+				{
+					return 0; // process connected
+				}
+				else
+				{
+					return -1; // connection failure etc.
+				}
+			}
+		}
 		public string ArtifactPath => TargetDevice.LocalCachePath;
-		public ILogStreamReader GetLogBufferReader() => new DynamicStringReader(() => "");
-		public ILogStreamReader GetLogReader() => new DynamicStringReader(() => "");
-		public bool WriteOutputToFile(string FilePath) => false;
-		public int WaitForExit() => ExitCode;
+		public ILogStreamReader GetLogBufferReader() => Logger?.GetReader();
+		public ILogStreamReader GetLogReader() => Logger?.GetReader();
+		public bool WriteOutputToFile(string FilePath)
+		{ 
+			Logger?.CopyFile(FilePath);
+			return (Logger != null);
+		}
+		public int WaitForExit()
+		{
+			if (ProcessMonitor.IsAvailable)
+			{
+				ProcessMonitor.WaitForExit();
+				ProcessMonitor.Stop();
+				Logger?.CloseStream();
+			}
+			return ExitCode;
+		}
 	};
 
 	/*
diff --git a/Engine/Source/Programs/AutomationTool/Win/RemoteWinSupport.cs b/Engine/Source/Programs/AutomationTool/Win/RemoteWinSupport.cs
--- a/Engine/Source/Programs/AutomationTool/Win/RemoteWinSupport.cs
+++ b/Engine/Source/Programs/AutomationTool/Win/RemoteWinSupport.cs
@@ -1,20 +1,22 @@
 // Copyright Epic Games, Inc. All Rights Reserved.
 using AutomationTool;
 using EpicGames.Core;
-using Google.Protobuf.WellKnownTypes;
 using Microsoft.Extensions.Logging;
 using System;
 using System.Collections.Generic;
 using System.Diagnostics;
 using System.IO;
 using System.Linq;
+using System.Net;
+using System.Net.Sockets;
+using System.Threading;
 using UnrealBuildTool;
 
 using static AutomationTool.CommandUtils;
 
 
-/* Remote Windows device support currently requires the Microsoft.Gaming.RemoteIterationClient_Microsoft WinGet package. 
-* (The remote devices need the Microsoft.Gaming.RemoteIterationEndpoint_Microsoft WinGetpackage)
+/* Remote Windows device support currently requires the Microsoft.Gaming.RemoteIterationClient WinGet package. 
+* (The remote devices need the Microsoft.Gaming.RemoteIterationEndpoint WinGet package)
 * 
 * It is recommended to use the Xbox PC Toolbox installer to configure host and remote devices : https://aka.ms/ToolboxInstaller
 * See https://aka.ms/GameRemoteDevtools for more details
@@ -119,7 +121,7 @@ public static class RemoteWinSupport
 		return true;
 	}
 
-	public static IProcessResult RunClient(UnrealTargetPlatform Platform, CommandUtils.ERunOptions ClientRunFlags, string ClientApp, string ClientCmdLine, ProjectParams Params)
+	public static IProcessResult RunClient(UnrealTargetPlatform Platform, ERunOptions ClientRunFlags, string ClientApp, string ClientCmdLine, ProjectParams Params)
 	{
 		if (Params.Devices.Count == 0)
 		{
@@ -129,6 +131,7 @@ public static class RemoteWinSupport
 		ClientCmdLine = SanitizeCommandLine(ClientCmdLine);
 
 		RemoteWinProcessResult Result = null;
+
 		foreach (string Device in Params.Devices)
 		{
 			if (Platform == GetTargetDevicePlatform(Device, Platform))
@@ -136,17 +139,25 @@ public static class RemoteWinSupport
 				string DeviceName = GetTargetDeviceName(Device);
 				string ExecutablePath = GetRemoteExecutablePath(Platform, ClientApp, Params);
 
-				if (!Run(DeviceName, ExecutablePath, ClientCmdLine))
+				string ExtraCmdLine = "";
+				Placeholder_RemoteWinProcessMonitor ProcessMonitor = new(DeviceName, Logger);
+				if (ProcessMonitor.IsAvailable)
+				{
+					ExtraCmdLine += $" -RemoteConsoleHost={ProcessMonitor.ServerAddress}";
+				}
+
+				if (!Run(DeviceName, ExecutablePath, ClientCmdLine + ExtraCmdLine))
 				{
 					throw new AutomationException("Remote launch on {0} failed", Device);
 				}
 
 				Logger.LogInformation("Running Package@Device:{ProcessName}@{DeviceName}", Path.GetFileName(ClientApp), DeviceName);
-				Result = new RemoteWinProcessResult(DeviceName, Path.GetFileName(ClientApp));
+				Result = new(DeviceName, Path.GetFileName(ClientApp), ProcessMonitor, ClientRunFlags);
 			}
 		}
 
-		if (Result != null && !ClientRunFlags.HasFlag(CommandUtils.ERunOptions.NoWaitForExit))
+
+		if (Result != null && !ClientRunFlags.HasFlag(ERunOptions.NoWaitForExit))
 		{
 			Result.WaitForExit();
 		}
@@ -297,58 +308,203 @@ public static class RemoteWinSupport
 	}
 
 
-
-	private class RemoteWinProcessResult : IProcessResult
+	/*
+	 * 
+	 * Helper class for monitoring the remote process's output and lifetime via a simple tcp server. 
+	 * 
+	 * Likely placeholder until the new wdremote tools support process lifetime management & output redirection
+	 * Paired with FPlaceholder_WindowsRemoteConsoleOutputDevice in the engine runtime
+	 * 
+	 * Not enabled in shipping builds by default (see WITH_REMOTEWIN_CONSOLE)
+	 * 
+	 */
+	public class Placeholder_RemoteWinProcessMonitor
 	{
-		private readonly string DeviceName;
-		private readonly string ProcessName;
-		private int ProcessExitCode = 0;
+		public Action<string> OutputReceived;
 
-		public RemoteWinProcessResult(string InDeviceName, string InProcessName)
+
+		public Placeholder_RemoteWinProcessMonitor( string InDeviceName, ILogger InLogger = null )
 		{
 			DeviceName = InDeviceName;
-			ProcessName = InProcessName;
+			Logger = InLogger;
+
+			try
+			{
+				Listener = new(IPAddress.Any, 0);
+				Listener.Start(1);
+
+				IEnumerable<string> ServerIPAddresses = Dns.GetHostEntry(Dns.GetHostName(), AddressFamily.InterNetwork).AddressList.Select( X => X.ToString() );
+				int Port = ((IPEndPoint)Listener.LocalEndpoint).Port;
+				ServerAddress = string.Join( '+', ServerIPAddresses ) + $":{Port}";
+
+				ListenThread = new Thread( new ThreadStart(RunServerThread) );
+
+				IsAvailable = true;
+			}
+			catch( Exception e )
+			{
+				Logger?.LogWarning("Could not set up windows remote TTY socket for {device} : {message}", DeviceName, e.Message);
+			}
 		}
 
-		public void StopProcess(bool KillDescendants = true)
+		public void Start()
 		{
-			Kill(DeviceName);
+			if (IsAvailable)
+			{
+				LaunchCTS = bNoTimeout ? new() : new(TimeSpan.FromSeconds(LaunchTimeoutSeconds));
+				ListenThread.Start();
+			}
+		}
+		public void Stop()
+		{
+			if (IsAvailable)
+			{
+				LaunchCTS.Cancel();
+				Listener.Stop();
+			}
 		}
 
-		public bool HasExited => true; // @todo: cannot monitor the remote process yet. assume it has terminated so launching doesn't get stuck waiting indefinitely
+		public void WaitForExit()
+		{
+			while (IsAvailable)
+			{
+				Thread.Sleep(0);
+			}
+		}
 
-		public string GetProcessName()
+		private async void RunServerThread()
 		{
-			return ProcessName;
+			if (!LaunchCTS.IsCancellationRequested)
+			{
+				try
+				{
+					TcpClient Client = await Listener.AcceptTcpClientAsync(LaunchCTS.Token);
+					if (Client != null)
+					{
+						bHadRemoteConnection = true;
+						StreamReader ClientStream = new(Client.GetStream());
+
+						string Line;
+						CancellationTokenSource ReadCTS = bNoTimeout ? new() : new(TimeSpan.FromSeconds(HeartbeatTimeutSeconds));
+
+						while ((Line = await ClientStream.ReadLineAsync(ReadCTS.Token)) != null)
+						{
+							ReadCTS.TryReset();
+							if (Line != "remotewin_heartbeat")
+							{
+								OutputReceived?.Invoke(Line);
+							}
+						}
+
+						Logger?.LogTrace("Remote windows process finished on {device}", DeviceName);
+					}
+					else
+					{
+						Logger?.LogTrace("Remote windows process - did not get a connection from {device}", DeviceName);
+					}
+				}
+				catch (OperationCanceledException)
+				{
+					Logger?.LogWarning("Remote windows process did not launch quickly enough on {device}", DeviceName);
+				}
+				catch (Exception e)
+				{
+					Logger?.LogTrace("Remote windows process likely terminated on {device} ({message})", DeviceName, e.Message);
+				}
+			}
+
+			IsAvailable = false;
 		}
 
-		public void OnProcessExited() { }
+		public bool DidLaunch => bHadRemoteConnection;
+		public bool IsAvailable {get; private set; } = false;
+		public string ServerAddress {get; private set; }
 
-		public void DisposeProcess() { }
+		private const int LaunchTimeoutSeconds = 20;
+		private const int HeartbeatTimeutSeconds = 10;
+		private const bool bNoTimeout = false; // set to true for debugging purposes only
+
+		private readonly string DeviceName;
+		private readonly ILogger Logger;
+		private readonly TcpListener Listener;
+		private readonly Thread ListenThread;
+		CancellationTokenSource LaunchCTS;
+		private bool bHadRemoteConnection = false;
+
+
+	}
 
-		public void StdOut(object sender, DataReceivedEventArgs e) { }
 
-		public void StdErr(object sender, DataReceivedEventArgs e) { }
 
-		public int ExitCode
+
+
+	private class RemoteWinProcessResult : IProcessResult
+	{
+		private readonly string DeviceName;
+		private readonly string ProcessName;
+		private readonly Placeholder_RemoteWinProcessMonitor ProcessMonitor;
+		private readonly bool bLogToConsole;
+		private readonly bool bLogToBuffer;
+
+		public RemoteWinProcessResult(string InDeviceName, string InProcessName, Placeholder_RemoteWinProcessMonitor InProcessMonitor, ERunOptions RunOptions)
 		{
-			get { return ProcessExitCode; }
-			set { ProcessExitCode = value; }
-		}
+			DeviceName = InDeviceName;
+			ProcessName = InProcessName;
+			ProcessMonitor = InProcessMonitor;
+			bLogToConsole = RunOptions.HasFlag(ERunOptions.AllowSpew);
+			bLogToBuffer = !RunOptions.HasFlag(ERunOptions.NoStdOutCapture);
+			if (!bLogToBuffer)
+			{
+				Output = null;
+			}
 
-		public bool bExitCodeSuccess => ExitCode == 0;
+			ExitCode = -1;
 
-		public string Output => "";
+			if (ProcessMonitor.IsAvailable)
+			{
+				ProcessMonitor.OutputReceived = OnOutputReceived;
+				ProcessMonitor.Start();
+			}
+		}
 
-		public Process ProcessObject => null;
+		private void OnOutputReceived( string Line )
+		{
+			if (bLogToBuffer)
+			{
+				Output += Line;
+			}
+			if (bLogToConsole)
+			{
+				Console.WriteLine(Line);
+			}
+		}
 
-		public new string ToString()
+		public void StopProcess(bool KillDescendants = true)
 		{
-			return $"RemoteWin RemoteWinProcessResult: {GetProcessName()} (ExitCode: {ExitCode})";
+			ProcessMonitor.Stop();
+			ExitCode = 1;
+			Kill(DeviceName);
 		}
 
+		public bool HasExited => !ProcessMonitor.IsAvailable;
+		public string GetProcessName() => ProcessName;
+		public void OnProcessExited() { }
+		public void DisposeProcess() { }
+		public void StdOut(object sender, DataReceivedEventArgs e) { }
+		public void StdErr(object sender, DataReceivedEventArgs e) { }
+		public int ExitCode {get; set; }
+		public bool bExitCodeSuccess => ExitCode == 0;
+		public string Output {get; private set; } = "";
+		public Process ProcessObject => null;
+		public new string ToString() => $"RemoteWin RemoteWinProcessResult: {GetProcessName()} (ExitCode: {ExitCode})";
+
 		public void WaitForExit()
 		{
+			if (ProcessMonitor.IsAvailable)
+			{
+				ProcessMonitor.WaitForExit();
+				ExitCode = 0;
+			}
 		}
 
 		public FileReference WriteOutputToFile(string FileName)
diff --git a/Engine/Source/Runtime/ApplicationCore/Private/Windows/WindowsPlatformApplicationMisc.cpp b/Engine/Source/Runtime/ApplicationCore/Private/Windows/WindowsPlatformApplicationMisc.cpp
--- a/Engine/Source/Runtime/ApplicationCore/Private/Windows/WindowsPlatformApplicationMisc.cpp
+++ b/Engine/Source/Runtime/ApplicationCore/Private/Windows/WindowsPlatformApplicationMisc.cpp
@@ -5,6 +5,7 @@
 #include "Windows/WindowsApplicationErrorOutputDevice.h"
 #include "Windows/WindowsConsoleOutputDevice.h"
 #include "Windows/WindowsConsoleOutputDevice2.h"
+#include "Windows/WindowsRemoteConsoleOutputDevice.h"
 #include "Windows/WindowsFeedbackContext.h"
 #include "HAL/FeedbackContextAnsi.h"
 #include "Misc/App.h"
@@ -139,6 +140,14 @@ void FWindowsPlatformApplicationMisc::LoadStartupModules()
 class FOutputDeviceConsole* FWindowsPlatformApplicationMisc::CreateConsoleOutputDevice()
 {
 	// this is a slightly different kind of singleton that gives ownership to the caller and should not be called more than once
+
+#if WITH_REMOTEWIN_CONSOLE
+	FString RemoteConsoleHost;
+	if (FParse::Value(FCommandLine::Get(), TEXT("RemoteConsoleHost="), RemoteConsoleHost))
+		return new FPlaceholder_WindowsRemoteConsoleOutputDevice(RemoteConsoleHost);
+	else
+#endif
+
 	if (FParse::Param(FCommandLine::Get(), TEXT("NewConsole")))
 		return new FWindowsConsoleOutputDevice2();
 	else
diff --git a/Engine/Source/Runtime/ApplicationCore/Private/Windows/WindowsRemoteConsoleOutputDevice.cpp b/Engine/Source/Runtime/ApplicationCore/Private/Windows/WindowsRemoteConsoleOutputDevice.cpp
new file mode 100644
--- /dev/null
+++ b/Engine/Source/Runtime/ApplicationCore/Private/Windows/WindowsRemoteConsoleOutputDevice.cpp
@@ -0,0 +1,226 @@
+// Copyright Epic Games, Inc. All Rights Reserved.
+
+#include "Windows/WindowsRemoteConsoleOutputDevice.h"
+
+#if WITH_REMOTEWIN_CONSOLE
+
+#include "Misc/CommandLine.h"
+#include "Misc/Parse.h"
+#include "Misc/OutputDeviceHelper.h"
+#include "Containers/SpscQueue.h"
+#include "Containers/Utf8String.h"
+#include "HAL/Runnable.h"
+#include "HAL/RunnableThread.h"
+#include "HAL/PlatformProcess.h"
+#include "HAL/Event.h"
+
+#include "Microsoft/AllowMicrosoftPlatformTypes.h"
+#include <WinSock2.h>
+#include <ws2tcpip.h>
+#include "Microsoft/HideMicrosoftPlatformTypes.h"
+
+
+class FRemoteConsoleRunnable : public FRunnable
+{
+public:
+	FRemoteConsoleRunnable( const TArray<FString>& RemoteAddresses, uint32 RemotePort)
+	{
+		WakeEvent = FPlatformProcess::GetSynchEventFromPool(false);
+		Socket = INVALID_SOCKET;
+
+		WSADATA WSAData;
+		WSAStartup(0x0101, &WSAData);
+
+		// convert the given addresses to winsock ipv4 addresses
+		for (const FString& RemoteAddress : RemoteAddresses)
+		{
+			sockaddr_storage& SockAddr = Addresses.AddZeroed_GetRef();
+			sockaddr_in* SockAddrIP4 = (sockaddr_in*)&SockAddr;
+			SockAddrIP4->sin_family = AF_INET;
+			SockAddrIP4->sin_port = htons((u_short)RemotePort);
+			InetPtonW(AF_INET, *RemoteAddress, &SockAddrIP4->sin_addr);
+		}
+	}
+
+
+	virtual ~FRemoteConsoleRunnable()
+	{
+		WSACleanup();
+
+		// note: this class is stored in GScopedLogConsole so there is no need to return WakeEvent to the pool as the pool is destroyed before us
+	}
+
+
+	virtual uint32 Run() override
+	{
+		// connect to the host PC
+		if (!Connect())
+		{
+			UE_LOG( LogWindows, Error, TEXT("Failed to connect to a WinRemoteConsoleDevice socket: %d"), WSAGetLastError());
+			return 1;
+		}
+
+		// keep pumping log messages
+		while (!IsEngineExitRequested() && Socket != INVALID_SOCKET)
+		{
+			// wait for any new serialized log messages, waking periodically to send a heartbeat
+			if (WakeEvent->Wait(FTimespan::FromSeconds(5)))
+			{		
+				FUtf8String Line;
+				while (Lines.Dequeue(Line))
+				{
+					SendString(Line);
+				}
+			}
+			else
+			{
+				// send a heartbeat to let UAT know that we are still running
+				SendString(FUtf8String(TEXT("remotewin_heartbeat\n")));
+			}
+		}
+
+		// clean up
+		closesocket(Socket);
+		Socket = INVALID_SOCKET;
+		return 0;
+	}
+
+
+	virtual void Stop() override
+	{
+		WakeEvent->Trigger();
+		closesocket(Socket);
+		Socket = INVALID_SOCKET;
+	}
+
+
+	void Serialize( const TCHAR* Line )
+	{
+		if (Socket != INVALID_SOCKET)
+		{
+			Lines.Enqueue(FUtf8String(Line));
+			WakeEvent->Trigger();
+		}
+	}
+
+private:
+	bool Connect()
+	{
+		static const int32 ConnectionTimeoutSeconds = 5;
+
+		// find an address to connect to
+		bool bIsConnected = false;
+		for (sockaddr_storage& SockAddr : Addresses)
+		{
+			Socket = socket(AF_INET, SOCK_STREAM, IPPROTO_TCP);
+			if (Socket == INVALID_SOCKET)
+			{
+				continue;
+			}
+
+			unsigned long NonBlockingMode = 1;
+			ioctlsocket(Socket, FIONBIO, &NonBlockingMode);
+
+			if (connect(Socket, (sockaddr*)&SockAddr, sizeof(sockaddr_in)) == 0)
+			{
+				bIsConnected = true;
+				break;
+			}
+
+			if (WSAGetLastError() == WSAEWOULDBLOCK)
+			{
+				fd_set FdSet = { 1, { Socket }, };
+				TIMEVAL TimeVal { ConnectionTimeoutSeconds, 0};
+				int Res = select(0, nullptr, &FdSet, nullptr, &TimeVal);
+				if (Res > 0)
+				{
+					bIsConnected = true;
+					break;
+				}
+			}
+
+			closesocket(Socket);
+			Socket = INVALID_SOCKET;
+		}
+
+		// don't need the addresses any more
+		Addresses.Reset();
+	
+		return bIsConnected;
+	}
+
+
+	void SendString( const FUtf8String& Line )
+	{
+		uint32 BytesToSend = Line.Len();
+		const char* DataPtr = (const char*)*Line;
+
+		while (BytesToSend > 0)
+		{
+			int32 BytesSent = send(Socket, DataPtr, BytesToSend, 0);
+			if (BytesSent == SOCKET_ERROR)
+			{
+				return;
+			}
+			BytesToSend -= BytesSent;
+			DataPtr += BytesSent;
+		}
+	}
+
+	TArray<sockaddr_storage> Addresses;
+	TSpscQueue<FUtf8String> Lines;
+	FEvent* WakeEvent;
+	SOCKET Socket;
+};
+
+
+
+
+FPlaceholder_WindowsRemoteConsoleOutputDevice::FPlaceholder_WindowsRemoteConsoleOutputDevice( const FString& RemoteConsoleHost )
+	: Runnable(nullptr)
+	, Thread(nullptr)
+{	
+	FString AddressesString;
+	FString PortString;
+	if (RemoteConsoleHost.Split(TEXT(":"), &AddressesString, &PortString))
+	{
+		TArray<FString> RemoteAddresses;
+		AddressesString.ParseIntoArray(RemoteAddresses, TEXT("+"));
+
+		if (RemoteAddresses.Num() > 0)
+		{
+			uint32 RemotePort = FCString::Atoi(*PortString);
+
+			Runnable = new FRemoteConsoleRunnable(RemoteAddresses, RemotePort);
+			Thread = FRunnableThread::Create( Runnable, TEXT("WinRemoteConsoleDevice"), 0, EThreadPriority::TPri_BelowNormal );
+		}
+	}
+
+	UE_CLOG( Thread == nullptr, LogWindows, Error, TEXT("Failed to create WinRemoteConsoleDevice"));
+}
+
+
+FPlaceholder_WindowsRemoteConsoleOutputDevice::~FPlaceholder_WindowsRemoteConsoleOutputDevice()
+{
+	if( Thread )
+	{
+		Thread->Kill();
+		delete Thread;
+
+		delete Runnable;
+		Runnable = nullptr;
+	}
+}
+
+
+void FPlaceholder_WindowsRemoteConsoleOutputDevice::Serialize( const TCHAR* Msg, ELogVerbosity::Type Verbosity, const class FName& Category )
+{
+	if (Runnable != nullptr)
+	{
+		FString LogMessage = FOutputDeviceHelper::FormatLogLine(Verbosity, Category, Msg, GPrintLogTimes) + LINE_TERMINATOR;
+		Runnable->Serialize(*LogMessage);
+	}
+}
+
+
+#endif // WITH_REMOTEWIN_CONSOLE
diff --git a/Engine/Source/Runtime/ApplicationCore/Public/Windows/WindowsRemoteConsoleOutputDevice.h b/Engine/Source/Runtime/ApplicationCore/Public/Windows/WindowsRemoteConsoleOutputDevice.h
new file mode 100644
--- /dev/null
+++ b/Engine/Source/Runtime/ApplicationCore/Public/Windows/WindowsRemoteConsoleOutputDevice.h
@@ -0,0 +1,45 @@
+// Copyright Epic Games, Inc. All Rights Reserved.
+
+#pragma once
+
+#include "CoreTypes.h"
+#include "Misc/OutputDeviceConsole.h"
+
+#if !defined(WITH_REMOTEWIN_CONSOLE)
+#define WITH_REMOTEWIN_CONSOLE (!UE_BUILD_SHIPPING) && (!WITH_EDITOR)
+#endif
+
+#if WITH_REMOTEWIN_CONSOLE
+
+/**
+ * 
+ * Implementation of console logging for a remote windows deployment
+ * 
+ * Likely placeholder until the wdremote tools support process lifetime management & output redirection.
+ * Paired with Placeholder_RemoteWinProcessMonitor in UAT
+ * 
+ * Not enabled in shipping builds by default (see WITH_REMOTEWIN_CONSOLE)
+ * 
+ */
+class FPlaceholder_WindowsRemoteConsoleOutputDevice : public FOutputDeviceConsole
+{
+public:
+
+	/** Default constructor. */
+	APPLICATIONCORE_API FPlaceholder_WindowsRemoteConsoleOutputDevice( const FString& RemoteConsoleHost );
+
+	/** Destructor. */
+	APPLICATIONCORE_API ~FPlaceholder_WindowsRemoteConsoleOutputDevice();
+
+	/** Sends the string to the remote logging host */
+	APPLICATIONCORE_API void Serialize( const TCHAR* Msg, ELogVerbosity::Type Verbosity, const class FName& Category );
+
+private:
+	virtual void Show(bool ShowWindow) {}
+	virtual bool IsShown() { return true; }
+
+	class FRemoteConsoleRunnable* Runnable;
+	class FRunnableThread* Thread;
+};
+
+#endif // WITH_REMOTEWIN_CONSOLE
