Remove unecessary include

#rnx

#Synced-CL 50905358

diff --git a/Engine/Source/Programs/AutomationTool/Gauntlet/Framework/Gauntlet.Utils.cs b/Engine/Source/Programs/AutomationTool/Gauntlet/Framework/Gauntlet.Utils.cs
--- a/Engine/Source/Programs/AutomationTool/Gauntlet/Framework/Gauntlet.Utils.cs
+++ b/Engine/Source/Programs/AutomationTool/Gauntlet/Framework/Gauntlet.Utils.cs
@@ -3,7 +3,6 @@
 using AutomationTool;
 using EpicGames.Core;
 using Gauntlet.TestTracking;
-using Google.Protobuf;
 using ImageMagick;
 using Microsoft.Extensions.Logging;
 using System;
