Fix chunking on Linux (redoing JoshA's change).

#Synced-CL 50879295

diff --git a/Engine/Source/Programs/AutomationTool/Scripts/CopyBuildToStagingDirectory.Automation.cs b/Engine/Source/Programs/AutomationTool/Scripts/CopyBuildToStagingDirectory.Automation.cs
--- a/Engine/Source/Programs/AutomationTool/Scripts/CopyBuildToStagingDirectory.Automation.cs
+++ b/Engine/Source/Programs/AutomationTool/Scripts/CopyBuildToStagingDirectory.Automation.cs
@@ -5345,11 +5345,13 @@ namespace AutomationScripts
 				{
 					foreach (string FileName in Chunk.Manifest)
 					{
+						string FixedName = FileName.Replace('/', Path.DirectorySeparatorChar).Replace('\\', Path.DirectorySeparatorChar);
+
 						List<ChunkDefinition> Chunks;
-						if (!FileNameToChunks.TryGetValue(FileName, out Chunks))
+						if (!FileNameToChunks.TryGetValue(FixedName, out Chunks))
 						{
 							Chunks = new List<ChunkDefinition>();
-							FileNameToChunks.Add(FileName, Chunks);
+							FileNameToChunks.Add(FixedName, Chunks);
 						}
 						Chunks.Add(Chunk);
 					}
