No tests executed message Gauntlet side

#rnx
#Synced-CL 0

diff --git a/Engine/Source/Programs/AutomationTool/Gauntlet/Unreal/Automation/UE.Automation.cs b/Engine/Source/Programs/AutomationTool/Gauntlet/Unreal/Automation/UE.Automation.cs
--- a/Engine/Source/Programs/AutomationTool/Gauntlet/Unreal/Automation/UE.Automation.cs
+++ b/Engine/Source/Programs/AutomationTool/Gauntlet/Unreal/Automation/UE.Automation.cs
@@ -1218,7 +1218,7 @@ namespace UE
 					{
 						if (bAllowZeroTestResults)
 						{
-							Log.Warning(KnownLogEvents.Gauntlet_UnrealEngineTestEvent, " * No tests were executed, check test filter.");
+							Log.Warning(KnownLogEvents.Gauntlet_UnrealEngineTestEvent, " * No tests were executed - confirm with test filter.");
 							Gauntlet.Horde.UpdateStepMessage(Gauntlet.Horde.GauntletHordeStepMessage.ZeroTestsRun);
 						}
 						else
