Automation: Replace usage of Ionic.Zip.Reduced.dll with System.IO.Compression in the Turnkey project.

#Synced-CL 50931528

diff --git a/Engine/Source/Programs/AutomationTool/Turnkey/CopyProvider.cs b/Engine/Source/Programs/AutomationTool/Turnkey/CopyProvider.cs
--- a/Engine/Source/Programs/AutomationTool/Turnkey/CopyProvider.cs
+++ b/Engine/Source/Programs/AutomationTool/Turnkey/CopyProvider.cs
@@ -8,9 +8,9 @@ using System.Text;
 using System.Threading.Tasks;
 using System.IO;
 using AutomationTool;
-using Ionic.Zip;
 using UnrealBuildTool;
 using EpicGames.Core;
+using System.IO.Compression;
 
 namespace Turnkey
 {
@@ -157,13 +157,8 @@ namespace Turnkey
 						{
 							try
 							{
-								using (ZipFile ZipFile = ZipFile.Read(OutputPath))
-								{
-									TurnkeyUtils.Log("Unzipping {0} to {1}...", OutputPath, DecompressLocation);
-
-									// extract the zip to it
-									ZipFile.ExtractAll(DecompressLocation);
-								}
+								TurnkeyUtils.Log("Unzipping {0} to {1}...", OutputPath, DecompressLocation);
+								ZipFile.ExtractToDirectory(OutputPath, DecompressLocation, overwriteFiles: true);
 							}
 							catch (Exception Ex)
 							{
diff --git a/Engine/Source/Programs/AutomationTool/Turnkey/Turnkey.Automation.csproj b/Engine/Source/Programs/AutomationTool/Turnkey/Turnkey.Automation.csproj
--- a/Engine/Source/Programs/AutomationTool/Turnkey/Turnkey.Automation.csproj
+++ b/Engine/Source/Programs/AutomationTool/Turnkey/Turnkey.Automation.csproj
@@ -15,10 +15,6 @@
 </PropertyGroup>
 
   <ItemGroup>
-    <Reference Include="Ionic.Zip.Reduced">
-      <HintPath>..\..\..\..\Binaries\DotNET\Ionic.Zip.Reduced.dll</HintPath>
-      <Private>False</Private>
-    </Reference>
     <Reference Include="UnrealWindowsForms">
       <HintPath>..\..\..\..\Binaries\DotNET\UnrealWindowsForms.dll</HintPath>
     </Reference>
