Automation: Replace usage of Ionic.Zip.Reduced.dll with System.IO.Compression in the Android project.

#Synced-CL 50931528

diff --git a/Engine/Source/Programs/AutomationTool/Android/Android.Automation.csproj b/Engine/Source/Programs/AutomationTool/Android/Android.Automation.csproj
--- a/Engine/Source/Programs/AutomationTool/Android/Android.Automation.csproj
+++ b/Engine/Source/Programs/AutomationTool/Android/Android.Automation.csproj
@@ -11,12 +11,6 @@
     <OutputPath>..\..\..\..\Binaries\DotNET\AutomationTool\AutomationScripts\Platforms\Android</OutputPath>
     <TreatWarningsAsErrors>true</TreatWarningsAsErrors>
   </PropertyGroup>
-  <ItemGroup>
-    <Reference Include="Ionic.Zip.Reduced">
-      <HintPath>..\..\..\..\Binaries\DotNET\Ionic.Zip.Reduced.dll</HintPath>
-      <Private>False</Private>
-    </Reference>
-  </ItemGroup>
   <ItemGroup>
     <Compile Include="..\..\Shared\MetaData.cs">
       <Link>Properties\MetaData.cs</Link>
diff --git a/Engine/Source/Programs/AutomationTool/Android/AndroidFileClient.cs b/Engine/Source/Programs/AutomationTool/Android/AndroidFileClient.cs
--- a/Engine/Source/Programs/AutomationTool/Android/AndroidFileClient.cs
+++ b/Engine/Source/Programs/AutomationTool/Android/AndroidFileClient.cs
@@ -2,26 +2,16 @@
 
 using System;
 using System.Collections.Generic;
-using System.Linq;
-using System.Text;
-using System.IO;
 using System.Diagnostics;
-using System.Net.NetworkInformation;
-using System.Threading;
-using AutomationTool;
-using UnrealBuildTool;
-using Ionic.Zip;
-using EpicGames.Core;
-using UnrealBuildBase;
-using System.Text.RegularExpressions;
-using AutomationScripts;
+using System.IO;
+using System.IO.Compression;
 using System.Net;
 using System.Net.Sockets;
-using System.IO.Compression;
+using System.Text;
+using System.Threading;
+using EpicGames.Core;
 using Microsoft.Extensions.Logging;
 
-using static AutomationTool.CommandUtils;
-
 
 namespace AutomationTool
 {
diff --git a/Engine/Source/Programs/AutomationTool/Android/AndroidPlatform.Automation.cs b/Engine/Source/Programs/AutomationTool/Android/AndroidPlatform.Automation.cs
--- a/Engine/Source/Programs/AutomationTool/Android/AndroidPlatform.Automation.cs
+++ b/Engine/Source/Programs/AutomationTool/Android/AndroidPlatform.Automation.cs
@@ -1,19 +1,19 @@
 // Copyright Epic Games, Inc. All Rights Reserved.
+using AutomationScripts;
+using AutomationTool;
+using EpicGames.Core;
+using Microsoft.Extensions.Logging;
 using System;
 using System.Collections.Generic;
-using System.Linq;
-using System.IO;
 using System.Diagnostics;
+using System.IO;
+using System.IO.Compression;
+using System.Linq;
 using System.Text.Json;
+using System.Text.RegularExpressions;
 using System.Threading;
-using AutomationTool;
-using UnrealBuildTool;
-using Ionic.Zip;
-using EpicGames.Core;
 using UnrealBuildBase;
-using System.Text.RegularExpressions;
-using AutomationScripts;
-using Microsoft.Extensions.Logging;
+using UnrealBuildTool;
 
 public class AndroidPlatform : Platform
 {
@@ -1084,42 +1084,24 @@ public class AndroidPlatform : Platform
 	private bool CreateOBBFile(DeploymentContext SC, string OutputFilename, List<FileReference> FilesForObb)
 	{
 		Logger.LogInformation("Creating {OutputFilename} from {Arg1}", OutputFilename, SC.StageDirectory);
-		using (ZipFile ObbFile = new ZipFile(OutputFilename))
+		try
 		{
-			ObbFile.CompressionMethod = CompressionMethod.None;
-			ObbFile.CompressionLevel = Ionic.Zlib.CompressionLevel.None;
-			ObbFile.UseZip64WhenSaving = Ionic.Zip.Zip64Option.Never;
-			ObbFile.Comment = String.Format("{0,10}", "1");
-
+			using ZipArchive ObbFile = ZipFile.Open(OutputFilename, ZipArchiveMode.Create);
+			ObbFile.Comment = $"{"1",10}";
 			int ObbFileCount = 0;
-			ObbFile.AddProgress +=
-				delegate (object sender, AddProgressEventArgs e)
-				{
-					if (e.EventType == ZipProgressEventType.Adding_AfterAddEntry)
-					{
-						ObbFileCount += 1;
-						Logger.LogInformation("[{Count}/{Total}] Adding {File} to OBB",
-							ObbFileCount, e.EntriesTotal,
-							e.CurrentEntry.FileName);
-					}
-				};
 
 			foreach (FileReference FileRef in FilesForObb)
 			{
-				string DestinationDirectoryPath = Path.GetRelativePath(SC.StageDirectory.FullName, Path.GetDirectoryName(FileRef.FullName));
-				ObbFile.AddFile(FileRef.FullName, DestinationDirectoryPath);
-			}
-
-			// ObbFile.AddDirectory(SC.StageDirectory+"/"+SC.ShortProjectName, SC.ShortProjectName);
-			try
-			{
-				ObbFile.Save();
-			}
-			catch (Exception)
-			{
-				return false;
+				string RelativePath = FileRef.MakeRelativeTo(SC.StageDirectory);
+				ObbFile.CreateEntryFromFile_CrossPlatform(FileRef.FullName, RelativePath, CompressionLevel.NoCompression);
+				ObbFileCount += 1;
+				Logger.LogInformation("[{Count}/{Total}] Adding {File} to OBB", ObbFileCount, FilesForObb.Count, RelativePath);
 			}
 		}
+		catch (Exception)
+		{
+			return false;
+		}
 		return true;
 	}
 
@@ -1127,18 +1109,15 @@ public class AndroidPlatform : Platform
 	{
 		string Version = Path.GetFileNameWithoutExtension(Filename).Split('.')[1];
 
-		using (ZipFile ObbFile = ZipFile.Read(Filename))
+		try
 		{
+			using ZipArchive ObbFile = ZipFile.Open(Filename, ZipArchiveMode.Update);
 			// Add the store version from the filename as a comment
-			ObbFile.Comment = String.Format("{0,10}", Version);
-			try
-			{
-				ObbFile.Save();
-			}
-			catch (Exception)
-			{
-				return false;
-			}
+			ObbFile.Comment = $"{Version,10}";
+		}
+		catch (Exception)
+		{
+			return false;
 		}
 		return true;
 	}
