Automation: Replace usage of Ionic.Zip.Reduced.dll with System.IO.Compression in the AutomationUtils project.

#Synced-CL 50931528

diff --git a/Engine/Source/Programs/AutomationTool/AutomationUtils/AutomationUtils.Automation.csproj b/Engine/Source/Programs/AutomationTool/AutomationUtils/AutomationUtils.Automation.csproj
--- a/Engine/Source/Programs/AutomationTool/AutomationUtils/AutomationUtils.Automation.csproj
+++ b/Engine/Source/Programs/AutomationTool/AutomationUtils/AutomationUtils.Automation.csproj
@@ -18,10 +18,6 @@
       <HintPath>..\..\..\..\Binaries\ThirdParty\fastJSON\netstandard2.0\fastJSON.dll</HintPath>
       <Private>False</Private>
     </Reference>
-    <Reference Include="Ionic.Zip.Reduced">
-      <HintPath>..\..\..\..\Binaries\DotNET\Ionic.Zip.Reduced.dll</HintPath>
-      <Private>False</Private>
-    </Reference>
   </ItemGroup>
 
   <!-- The following will glob all platform extensions' AutomationUtils source
diff --git a/Engine/Source/Programs/AutomationTool/AutomationUtils/CommandUtils.cs b/Engine/Source/Programs/AutomationTool/AutomationUtils/CommandUtils.cs
--- a/Engine/Source/Programs/AutomationTool/AutomationUtils/CommandUtils.cs
+++ b/Engine/Source/Programs/AutomationTool/AutomationUtils/CommandUtils.cs
@@ -2527,24 +2527,20 @@ namespace AutomationTool
 			}
 			else
 			{
-				// manually extract the files. There was a problem with the Ionic.Zip library that required this on non-PC at one point,
-				// but that problem is now fixed. Leaving this code as is as we need to return the list of created files anyway.
-				using (Ionic.Zip.ZipFile Zip = new Ionic.Zip.ZipFile(ZipFileName))
+				// Manually extract the files. Leaving this code as is in order to preserve behaviour around file permissions.
+				using ZipArchive Zip = ZipFile.OpenRead(ZipFileName);
+				FileMode OutputFileMode = (OverwriteFiles) ? FileMode.Create : FileMode.CreateNew;
+				foreach (ZipArchiveEntry Entry in Zip.Entries.Where(e => !string.IsNullOrEmpty(e.Name)))
 				{
-					FileMode OutputFileMode = (OverwriteFiles) ? FileMode.Create : FileMode.CreateNew;
-					foreach (Ionic.Zip.ZipEntry Entry in Zip.Entries.Where(x => !x.IsDirectory))
+					string OutputFileName = Path.Combine(BaseDirectory, Entry.FullName);
+					if (OverwriteFiles || File.Exists(OutputFileName) == false)
 					{
-						string OutputFileName = Path.Combine(BaseDirectory, Entry.FileName);
-						if (OverwriteFiles || File.Exists(OutputFileName) == false)
-						{
-							Directory.CreateDirectory(Path.GetDirectoryName(OutputFileName));
-							using (FileStream OutputStream = new FileStream(OutputFileName, OutputFileMode, FileAccess.Write))
-							{
-								Entry.Extract(OutputStream);
-							}
-						}
-						OutputFileNames.Add(OutputFileName);
+						Directory.CreateDirectory(Path.GetDirectoryName(OutputFileName));
+						using FileStream OutputStream = new(OutputFileName, OutputFileMode, FileAccess.Write);
+						using Stream EntryStream = Entry.Open();
+						EntryStream.CopyTo(OutputStream);
 					}
+					OutputFileNames.Add(OutputFileName);
 				}
 			}
 			return OutputFileNames;
