don't remove Win64 from RestrictedFolderNames in MSGameStoreCustomDeploymentHandler

#rnx
#Synced-CL 50857758

diff --git a/Engine/Source/Programs/AutomationTool/Win/MSGaming/MSGameStoreCustomDeploymentHandler.cs b/Engine/Source/Programs/AutomationTool/Win/MSGaming/MSGameStoreCustomDeploymentHandler.cs
--- a/Engine/Source/Programs/AutomationTool/Win/MSGaming/MSGameStoreCustomDeploymentHandler.cs
+++ b/Engine/Source/Programs/AutomationTool/Win/MSGaming/MSGameStoreCustomDeploymentHandler.cs
@@ -269,7 +269,6 @@ public class MSGameStoreCustomDeploymentHandler : CustomDeploymentHandler
 	public override bool PreGetFilesToDeployOrStage(ProjectParams Params, DeploymentContext SC)
 	{
 		SC.RestrictedFolderNames.RemoveWhere(x => x == "GDK");
-		SC.RestrictedFolderNames.RemoveWhere(x => x == "Win64");
 
 		// copy all custom install actions to the correct folder if we are not using package precaching
 		MSGameStorePackageGenerator PackageGenerator = new(SC.RawProjectPath, SC.StageTargetPlatform.PlatformType, IniSection_TargetSettings, SC.CustomConfig, SC.StageTargetPlatform.CrashReportPlatform);
