ReportItemizedExecutableCode - Improve module matching and report category totals.

Updates Query UBT Mode to also include the short name, module directory and intermediate directory for each module.
Doing this has the advantage of being able to do much more accurate module match as we can match based on intermediate path, as well as being able to get the real source path so we can tell if it's an engine or game module.

Update ReportItemizedExecutableCode to use module info from UBT to better match compile units. Also adds a 'totals' report which includes game/engine/plugins breakdowns.

#rnx

#Synced-CL 50645574

diff --git a/Engine/Source/Programs/AutomationTool/Scripts/ReportItemizedExecutableCode.cs b/Engine/Source/Programs/AutomationTool/Scripts/ReportItemizedExecutableCode.cs
--- a/Engine/Source/Programs/AutomationTool/Scripts/ReportItemizedExecutableCode.cs
+++ b/Engine/Source/Programs/AutomationTool/Scripts/ReportItemizedExecutableCode.cs
@@ -2,15 +2,50 @@
 
 using AutomationTool;
 using EpicGames.Core;
+using System;
 using System.Collections.Generic;
 using System.IO;
-using System;
-using UnrealBuildBase;
 using System.Linq;
 using System.Text;
+using System.Text.Json;
+using System.Text.Json.Serialization;
+using UnrealBuildBase;
+using UnrealBuildTool;
 
 namespace ReportItemizedExecutableCode.Automation
 {
+	struct ModuleInfo
+	{
+		/** Module Name */
+		public string Name;
+		
+		/** Module ShortName, null if there isn't one. */
+		[JsonPropertyName("shortName")]
+		public string ShortName { get; set; }
+
+		/** Source directory for the module. Relative to project root. */
+		[JsonPropertyName("moduleDirectory")]
+		public string ModuleDirectory { get; set; }
+
+		/** Intermediate directory for the module (ie. where the unity modules are put). Relative to project root. */
+		[JsonPropertyName("intermediateDirectory")]
+		public string IntermediateDirectory { get; set; }
+
+		public bool IsEngineModule => ModuleDirectory.StartsWith("Engine");
+		public bool IsEnginePlugin => IsEngineModule && ModuleDirectory.Contains("\\Plugins\\");
+		public bool IsGameModule => !IsEngineModule;
+		public bool IsGamePlugin => IsGameModule && ModuleDirectory.Contains("\\Plugins\\");
+		public bool IsGameFeaturePlugin => IsGamePlugin && ModuleDirectory.Contains("\\GameFeatures\\");
+	}
+
+	struct ModuleResult
+	{
+		/** Resolved module name. */
+		public string Name;
+		/** Resolved module info if found. Will be null if we had to fallback to parsing the module name from the path. */
+		public ModuleInfo? ModuleInfo;
+	}
+
 	struct ModuleInfos
 	{
 		/** Holds the actual itemized data (module to its bytes) */
@@ -22,15 +57,38 @@ namespace ReportItemizedExecutableCode.Automation
 		/** Which sources we couldn't classify */
 		public List<string> SourcesWithoutModule;
 
+		/** Total size of all compileunits. */
+		public uint Total = 0;
+
 		/** Size of sources that we couldn't classify */
-		public uint TotalSizeForSourcesWithoutModule;
+		public uint TotalUnclassified = 0;
+
+		/** Size of engine sources */
+		public uint TotalEngineSize = 0;
+
+		/** Size of game sources */
+		public uint TotalGameSize = 0;
+
+		/** Size of all engine plugins. This total is captured in TotalEngineSize. */
+		public uint TotalEnginePluginSize = 0;
+
+		/** Size of all game plugins. This total is captured in TotalGameSize. */
+		public uint TotalGamePluginSize = 0;
+
+		/** Size of all game feature plugins. This total is captured in TotalGamePluginSize. */
+		public uint TotalGameFeaturesSize = 0;
+
+		/** Size coming from external libraries. */
+		public uint TotalExternalLibs = 0;
+
+		/** These are bytes that couldn't be attributed to a specific compile unit. */
+		public uint TotalUnattributed = 0;
 
 		public ModuleInfos()
 		{
 			Modules = new Dictionary<string, uint>();
 			ModuleNameToSources = new Dictionary<string, List<string>>();
 			SourcesWithoutModule = new List<string>();
-			TotalSizeForSourcesWithoutModule = 0;
 		}
 	};
 
@@ -83,13 +141,18 @@ namespace ReportItemizedExecutableCode.Automation
 	/// </example>
 
 	[Help("Reports itemized binary size in a format consumed by PerfReportServer")]
-	[ParamHelp("ExeToItemize", "Absolute path to the binary to itemize", ParamType = typeof(FileReference))]
+	[ParamHelp("ExeToItemize", "Absolute path to the binary to itemize", ParamType = typeof(FileReference), Required = true)]
 	[ParamHelp("SymbolsExe", "Absolute path to the symbolicated binary", ParamType = typeof(FileReference))]
 	[ParamHelp("ReportPath", "Absolute path to the directory to write the reports to.", ParamType = typeof(DirectoryReference))]
+	[ParamHelp("Platform", "Target platform for module scan: Auto|Switch|Win64|PS4|PS5|Android|iOS|Xbox (default: Auto)", ParamType = typeof(string), Required = true)]
+	[ParamHelp("Project", "The Project name or path.", ParamType = typeof(string), Required = true)]
+	[ParamHelp("Target", "The UBT Target.", ParamType = typeof(string), Required = true)]
+
+	[ParamHelp("PrintResults", "Print itemization results to stdout. Defaults to true", ParamType = typeof(bool))]
+	[ParamHelp("Verbose", "Print verbose output. Defaults to false", ParamType = typeof(bool))]
 
 	// Merged Modules (optional flags)
 	[ParamHelp("DllRoot", "Root folder to scan for merged modules (staged outputs). If omitted, merged-modules logic is skipped.", ParamType = typeof(DirectoryReference))]
-	[ParamHelp("Platform", "Target platform for module scan: Auto|Switch|Win64|PS4|PS5|Android|iOS|Xbox (default: Auto)", ParamType = typeof(string))]
 	[ParamHelp("ModuleExts", "Override module extensions (comma-separated, e.g. .nro,.dll). Default chosen by Platform.", ParamType = typeof(string))]
 	[ParamHelp("EagerPrefixes", "Override 'eager' name prefixes (semicolon-separated, e.g. EOS;KITT).", ParamType = typeof(string))]
 	[ParamHelp("PackageMarkerSuffixes", "Override package marker suffixes (semicolon-separated, e.g. .nspd).", ParamType = typeof(string))]
@@ -102,6 +165,19 @@ namespace ReportItemizedExecutableCode.Automation
 	[ParamHelp("SkipExeItemization", "If true, do not run bloaty on the main executable; only run merged-modules logic (if DllRoot is set).", ParamType = typeof(bool))]
 	class ReportItemizedExecutableCode : BuildCommand
 	{
+		// Verbose output
+		private bool bVerbose = false;
+
+		private List<ModuleInfo> Modules = new();
+
+		private void VerboseLog(string Message)
+		{
+			if (bVerbose)
+			{
+				Console.WriteLine(Message);
+			}
+		}
+
 		private string[] GetReportDestDirectories()
 		{
 			List<string> DestPaths = new List<string>()
@@ -118,7 +194,7 @@ namespace ReportItemizedExecutableCode.Automation
 		}
 
 		/** Attempts to assign a reasonable module name from a source file name */
-		string FindModuleName(string SourcePath)
+		string ParseModuleName(string SourcePath)
 		{
 			// For context, here are the examples we might be dealing with:
 
@@ -225,13 +301,80 @@ namespace ReportItemizedExecutableCode.Automation
 			return SourcePath;
 		}
 
+		ModuleResult? FindModule(string SourcePath)
+		{
+			string NormalizedPath = SourcePath.Replace("/", "\\");
+			string DirectoryPath = Path.GetDirectoryName(NormalizedPath);
+			if (string.IsNullOrEmpty(DirectoryPath))
+			{
+				return null;
+			}
+
+			if (Modules != null)
+			{
+				// Look for the longest match to avoid "MyModuleCore" being matched to "MyModule".
+				ModuleInfo? BestMatch = null;
+				int LongestMatch = 0;
+				foreach (ModuleInfo Module in Modules)
+				{
+					int MatchLength = 0;
+
+					// The source path for unity files will by the intermediate dir.
+					if (DirectoryPath.Contains(Module.IntermediateDirectory))
+					{
+						MatchLength = Module.IntermediateDirectory.Length;
+					}
+
+					// Whereas separately compiled libs will be the actual source location.
+					if (DirectoryPath.Contains(Module.ModuleDirectory))
+					{
+						MatchLength = Module.ModuleDirectory.Length;
+					}
+
+					if (MatchLength > LongestMatch)
+					{
+						LongestMatch = MatchLength;
+						BestMatch = Module;
+					}
+				}
+
+				if (BestMatch != null)
+				{
+					// Use the short name if available for backwards compat to prevent stat names changing.
+					return new ModuleResult()
+					{
+						Name = BestMatch?.ShortName ?? BestMatch?.Name,
+						ModuleInfo = BestMatch
+					};
+				}
+			}
+
+			// Fallback to brute force method.
+			// This is necessary for some third party libs where there's just no way to map it to a project module.
+			string ParsedModuleName = ParseModuleName(SourcePath);
+			VerboseLog($"Could not resolve module for {SourcePath}. Parsed as {ParsedModuleName}");
+
+			if (IsValidModuleName(ParsedModuleName))
+			{
+				return new ModuleResult()
+				{
+					Name = ParsedModuleName
+				};
+			}
+			return null;
+		}
+
 		/** Module names can be used as perf metrics, so we need to remove spaces and non-alphanumeric characters except underscores and hyphens */
-		string SanitizeModuleName(string ModuleName)
+		private static string SanitizeModuleName(string ModuleName)
 		{
+			if (ModuleName == null)
+			{
+				return ModuleName;
+			}
 			return ModuleName.Replace(' ', '_').Replace('.', '-').Replace('#', '-').Replace("[", "").Replace("]", "");
 		}
 
-		bool ValidateModuleName(string ModuleName, string SourcePath)
+		private static bool IsValidModuleName(string ModuleName)
 		{
 			if (string.IsNullOrEmpty(ModuleName) || ModuleName == ".." || char.IsDigit(ModuleName[0]))
 			{
@@ -241,6 +384,15 @@ namespace ReportItemizedExecutableCode.Automation
 			return true;
 		}
 
+		private static bool IsPathLike(string InString)
+		{
+			if (!string.IsNullOrWhiteSpace(InString))
+			{
+				return InString.IndexOfAny(Path.GetInvalidPathChars()) < 0;
+			}
+			return false;
+		}
+
 		/** Attempts to run bloaty on the executable and returns the results */
 		protected ModuleInfos ItemizeExecutableWithBloaty(FileReference Binary, FileReference SymbolsBinary)
 		{
@@ -248,16 +400,20 @@ namespace ReportItemizedExecutableCode.Automation
 
 			if (FileReference.Exists(Binary))
 			{
+				// If a csv was provided as the binary, treat it as bloaty output so we can skip running it.
+				string ResultsOverride = null;
+				if (Binary.FullName.EndsWith(".csv"))
+				{
+					ResultsOverride = File.ReadAllText(Binary.FullName);
+				}
+
 				string Args = string.Format("-d compileunits -n 0 --csv {0}", Binary.FullName);
 				if (SymbolsBinary != null && FileReference.Exists(SymbolsBinary))
 				{
 					Args += string.Format(" --debug-file={0}", SymbolsBinary.FullName);
 				}
 
-				// Hack!
-				//string[] LineArray = RunBloaty(Args, HardCodedResults);
-
-				string[] LineArray = RunBloaty(Args);
+				string[] LineArray = RunBloaty(Args, ResultsOverride);
 
 				// since we told bloaty to output csv, the result will be already machine readable. Just make sure that the first line is as we expect
 				const string ExpectedCSVOutput = "compileunits,vmsize,filesize";
@@ -266,8 +422,6 @@ namespace ReportItemizedExecutableCode.Automation
 					throw new AutomationException("CSV output ({0}) does not match expectation ({1})", LineArray[0], ExpectedCSVOutput);
 				}
 
-				uint Total = 0;
-
 				for (int IdxLine = 1; IdxLine < LineArray.Length; ++IdxLine)
 				{
 					string[] LineElems = LineArray[IdxLine].Split(',');
@@ -276,12 +430,14 @@ namespace ReportItemizedExecutableCode.Automation
 						throw new AutomationException("CSV line ({0}) does not match expectation ({1}) - comma in source file?", LineArray[IdxLine], ExpectedCSVOutput);
 					}
 
-					string ModuleName = SanitizeModuleName(FindModuleName(LineElems[0]));
+					string SourcePath = LineElems[0];
+					ModuleResult? Module = FindModule(SourcePath);
 					uint Size = uint.Parse(LineElems[1]);
-					Total += Size;
+					ModInfos.Total += Size;
 
-					if (ValidateModuleName(ModuleName, LineArray[IdxLine]))
+					if (Module is ModuleResult Result)
 					{
+						string ModuleName = SanitizeModuleName(Result.Name);
 						if (ModInfos.Modules.ContainsKey(ModuleName))
 						{
 							ModInfos.Modules[ModuleName] += Size;
@@ -293,24 +449,82 @@ namespace ReportItemizedExecutableCode.Automation
 
 						if (ModInfos.ModuleNameToSources.ContainsKey(ModuleName))
 						{
-							ModInfos.ModuleNameToSources[ModuleName].Add(LineElems[0]);
+							ModInfos.ModuleNameToSources[ModuleName].Add(SourcePath);
+						}
+						else
+						{
+							ModInfos.ModuleNameToSources.Add(ModuleName, new List<string>() { SourcePath });
+						}
+
+						if (Result.ModuleInfo is ModuleInfo Info)
+						{
+							if (Info.IsEngineModule)
+							{
+								ModInfos.TotalEngineSize += Size;
+							}
+							else
+							{
+								ModInfos.TotalGameSize += Size;
+							}
+
+							if (Info.IsGamePlugin)
+							{
+								ModInfos.TotalGamePluginSize += Size;
+								if (Info.IsGameFeaturePlugin)
+								{
+									ModInfos.TotalGameFeaturesSize += Size;
+								}
+							}
+							else if (Info.IsEnginePlugin)
+							{
+								ModInfos.TotalEnginePluginSize += Size;
+							}
 						}
 						else
 						{
-							List<string> Sources = new List<string>();
-							Sources.Add(LineElems[0]);
-							ModInfos.ModuleNameToSources.Add(ModuleName, Sources);
+							// If it looks like a path then it's likely an externally compiled lib.
+							if (IsPathLike(SourcePath))
+							{
+								ModInfos.TotalExternalLibs += Size;
+								VerboseLog($"Classifying {SourcePath} as a external lib");
+							}
+							else
+							{
+								ModInfos.TotalUnclassified += Size;
+								VerboseLog($"Could not classify {SourcePath}");
+							}
 						}
 					}
 					else
 					{
+						VerboseLog($"Unable to resolve module name for {SourcePath}");
 						ModInfos.SourcesWithoutModule.Add(LineArray[IdxLine]);
-						ModInfos.TotalSizeForSourcesWithoutModule += Size;
+
+						// If bloaty cannot associate some bytes with a specific compile unit then it will group it under the section header instead.
+						// Other sections are rows like [section .rodata]
+						if (SourcePath.StartsWith("["))
+						{
+							ModInfos.TotalUnattributed += Size;
+							VerboseLog($"Classifying {SourcePath} as unattributed");
+						}
+						else if (IsPathLike(SourcePath))
+						{
+							// If it looks like a path then it's likely an externally compiled lib.
+							ModInfos.TotalExternalLibs += Size;
+							VerboseLog($"Classifying {SourcePath} as a external lib");
+						}
+						else
+						{
+							ModInfos.TotalUnclassified += Size;
+							VerboseLog($"Could not classify {SourcePath}");
+						}
 					}
 				}
 
 				// add a _Total metric
-				ModInfos.Modules.Add("_Total", Total);
+				ModInfos.Modules.Add("_Total", ModInfos.Total);
+				ModInfos.Modules.Add("_EngineTotal", ModInfos.TotalEngineSize);
+				ModInfos.Modules.Add("_GameTotal", ModInfos.TotalGameSize);
 				return ModInfos;
 			}
 			else
@@ -330,32 +544,25 @@ namespace ReportItemizedExecutableCode.Automation
 			FileReference Exe = ParseRequiredFileReferenceParam("ExeToItemize");
 			string[] DestDirs = GetReportDestDirectories();
 
-			bool SkipExe = ParseParamBool("SkipExeItemization");
+			bVerbose = ParseParamBool("Verbose", false);
+
+			Modules = BuildModuleList();
+			if (Modules == null)
+			{
+				Console.Error.WriteLine("Failed to create module list. Falling back on source path parsing.");
+			}
+			else
+			{
+				Console.WriteLine($"Found {Modules.Count} modules.");
+			}
 
 			// --- SECTION 1: Original functionality for itemizing a single executable ---
 			// Can be skipped with -SkipExeItemization=true if only Merged Module analysis is needed.
 
+			bool SkipExe = ParseParamBool("SkipExeItemization");
 			if (!SkipExe)
 			{
-				try
-				{
-					string ItemizedName = Exe.GetFileNameWithoutExtension() + "_ItemizedModules.csv";
-					WriteItemizedReport(Exe, DestDirs.Select(d => Path.Combine(d, ItemizedName)).ToArray());
-				}
-				catch (Exception ex)
-				{
-					Console.WriteLine($"Skipping ItemizedModules (EXE): {ex.Message}");
-				}
-
-				try
-				{
-					string SectionsName = Exe.GetFileNameWithoutExtension() + "_Sections.csv";
-					WriteSectionsReport(Exe, DestDirs.Select(d => Path.Combine(d, SectionsName)).ToArray());
-				}
-				catch (Exception ex)
-				{
-					Console.WriteLine($"Skipping Sections (EXE): {ex.Message}");
-				}
+				ExecuteItemization(Exe, DestDirs);
 			}
 
 			// --- SECTION 2: New functionality for Merged Modules analysis ---
@@ -367,60 +574,183 @@ namespace ReportItemizedExecutableCode.Automation
 			}
 		}
 
-
-		private void WriteItemizedReport(FileReference ExecutableFileToItemize, string[] DestPaths)
+		private List<ModuleInfo> BuildModuleList()
 		{
-			FileReference ExecutableFileWithSymbols = ParseOptionalFileReferenceParam("SymbolsExe");
+			string ProjectNameOrPath = ParseRequiredStringParam("Project");
+			FileReference ProjectPath = ParseProjectString(ProjectNameOrPath);
+			if (ProjectPath == null)
+			{
+				Console.Error.WriteLine("Failed to parse project path.");
+				return null;
+			}
+
+			string Target = ParseRequiredStringParam("Target");
+
+			string ModuleInfoPath = Path.GetTempFileName();
+			UnrealTargetPlatform Platform = UnrealTargetPlatform.Parse(ParseRequiredStringParam("Platform"));
+
+			// Run UBT in query mode to dump out all the modules with their paths to json.
+			List<string> Args = new List<string>()
+			{
+				$"-project=\"{ProjectPath}\"",
+				$"-Target={Target}",
+				$"-Platform={Platform}",
+				$"-Configuration={UnrealTargetConfiguration.Test}",
+				$"-mode=Query",
+				$"-Query=TargetModuleReferences",
+				$"-NoIntellisense",
+				$"-OutputPath=\"{ModuleInfoPath}\""
+			};
+
+			try
+			{
+				string UBTArgs = string.Join(" ", Args);
+				UnrealBuildUtils.RunUBT(CommandUtils.CmdEnv, UnrealBuild.UnrealBuildToolDll, UBTArgs);
+			}
+			catch (AutomationException Ex)
+			{
+				Console.Error.WriteLine($"Error dumping modules from UBT: {Ex}");
+				return null;
+			}
+			
+			if (!Path.Exists(ModuleInfoPath))
+			{
+				Console.Error.WriteLine($"No module json file found at {ModuleInfoPath}");
+				return null;
+			}
+
+			List<ModuleInfo> Modules = [];
+
+			try
+			{
+				string JsonText = File.ReadAllText(ModuleInfoPath);
+				var ModuleMap = JsonSerializer.Deserialize<Dictionary<string, Dictionary<string, ModuleInfo>>>(JsonText);
+
+				DirectoryReference RootDirectory = UnrealBuildBase.Unreal.RootDirectory;
 
-			// Hack for quick iteration on formatting (run bloaty manually and just point this at the csv file)
-			//HardCodedResults = FileReference.ReadAllText(ExecutableFileToItemize);
+				Modules = ModuleMap["modules"]
+					.Select(Kvp => new ModuleInfo()
+					{
+						Name = Kvp.Key,
+						ShortName = Kvp.Value.ShortName,
+						ModuleDirectory = new DirectoryReference(Kvp.Value.ModuleDirectory).MakeRelativeTo(RootDirectory),
+						IntermediateDirectory = new DirectoryReference(Kvp.Value.IntermediateDirectory).MakeRelativeTo(RootDirectory)
+					})
+					.ToList();
+			}
+			catch (Exception Ex)
+			{
+				Console.Error.WriteLine($"Error parsing module json: {Ex}");
+			}
+			finally
+			{
+				File.Delete(ModuleInfoPath);
+			}
+			return Modules;
+		}
 
-			System.Console.WriteLine("\n\n\n\nExecuting itemization of executable: {0}", ExecutableFileToItemize);
+		private void ExecuteItemization(FileReference ExecutableFileToItemize, string[] DestDirs)
+		{
+			FileReference ExecutableFileWithSymbols = ParseOptionalFileReferenceParam("SymbolsExe");
+			System.Console.WriteLine("Executing itemization of executable: {0}", ExecutableFileToItemize);
 			if (ExecutableFileWithSymbols != null)
 			{
-				System.Console.WriteLine("\n\n\n\nUsing executable for symbolication: {0}", ExecutableFileWithSymbols);
+				System.Console.WriteLine("Using executable for symbolication: {0}", ExecutableFileWithSymbols);
 			}
 
 			ModuleInfos Info = PlatformItemizeExecutable(ExecutableFileToItemize, ExecutableFileWithSymbols);
 			if (Info.Modules.Count == 0)
 			{
-				throw new AutomationException("Unable to itemize binary '{0}'", ExecutableFileToItemize.ToString());
+				Console.Error.WriteLine("Unable to itemize binary '{0}'", ExecutableFileToItemize.ToString());
+				return;
 			}
 
-			System.Console.WriteLine("\n\n\n\nExec itemization, here's how we arrived at module names:");
-			foreach (KeyValuePair<string, List<string>> ModuleMap in Info.ModuleNameToSources)
+			try
 			{
-				uint SizeForThisModule = Info.Modules[ModuleMap.Key];
+				string[] DestPaths = DestDirs.Select(Dir => Path.Combine(Dir, ExecutableFileToItemize.GetFileNameWithoutExtension() + "_ItemizedModules.csv")).ToArray();
+				WriteItemizedReport(Info, DestPaths);
+			}
+			catch (Exception ex)
+			{
+				Console.WriteLine($"Failed to generate module itemization report: {ex.Message}");
+			}
 
-				System.Console.WriteLine("{0} (takes {1:F2} MB):", ModuleMap.Key, SizeForThisModule / (1024.0 * 1024.0));
-				foreach (string Source in ModuleMap.Value)
-				{
-					System.Console.WriteLine("\t{0}", Source);
-				}
+			try
+			{
+				string[] DestPaths = DestDirs.Select(Dir => Path.Combine(Dir, ExecutableFileToItemize.GetFileNameWithoutExtension() + "_Totals.csv")).ToArray();
+				WriteTotalsReport(Info, DestPaths);
+			}
+			catch (Exception ex)
+			{
+				Console.WriteLine($"Failed to generate totals report: {ex.Message}");
 			}
 
-			System.Console.WriteLine("\n\n\n\nAdditionally, {0} source files of {1} bytes ({2:F2} MB in total) failed to be categorized:",
-				Info.SourcesWithoutModule.Count, Info.TotalSizeForSourcesWithoutModule, Info.TotalSizeForSourcesWithoutModule / (1024.0 * 1024.0));
-			if (Info.SourcesWithoutModule.Count > 0)
+			try
 			{
-				foreach (string Source in Info.SourcesWithoutModule)
-				{
-					System.Console.WriteLine("\t{0}", Source);
-				}
+				string[] DestPaths = DestDirs.Select(Dir => Path.Combine(Dir, ExecutableFileToItemize.GetFileNameWithoutExtension() + "_Sections.csv")).ToArray();
+				WriteSectionsReport(ExecutableFileToItemize, DestPaths);
 			}
-			else
+			catch (Exception ex)
 			{
-				System.Console.WriteLine("None!");
+				Console.WriteLine($"Failed to generate sections report: {ex.Message}");
+			}
+		}
+
+		private void WriteItemizedReport(FileReference ExecutableFileToItemize, string[] DestPaths)
+		{
+			FileReference ExecutableFileWithSymbols = ParseOptionalFileReferenceParam("SymbolsExe");
+
+			ModuleInfos Info = PlatformItemizeExecutable(ExecutableFileToItemize, ExecutableFileWithSymbols);
+			if (Info.Modules.Count == 0)
+			{
+				throw new AutomationException("Unable to itemize binary '{0}'", ExecutableFileToItemize.ToString());
 			}
 
+			WriteItemizedReport(Info, DestPaths);
+		}
+
+		private void WriteItemizedReport(ModuleInfos Info, string[] DestPaths)
+		{
+			bool bPrintResults = ParseParamBool("PrintResults", true);
+			if (bPrintResults)
+			{
+				System.Console.WriteLine("\n\n\n\nExec itemization, here's how we arrived at module names:");
+				foreach (KeyValuePair<string, List<string>> ModuleMap in Info.ModuleNameToSources)
+				{
+					uint SizeForThisModule = Info.Modules[ModuleMap.Key];
+
+					System.Console.WriteLine("{0} (takes {1:F2} MB):", ModuleMap.Key, SizeForThisModule / (1024.0 * 1024.0));
+					foreach (string Source in ModuleMap.Value)
+					{
+						System.Console.WriteLine("\t{0}", Source);
+					}
+				}
+
+				System.Console.WriteLine("\n\n\n\nAdditionally, {0} source files of {1} bytes ({2:F2} MB in total) failed to be categorized:",
+					Info.SourcesWithoutModule.Count, Info.TotalUnclassified, Info.TotalUnclassified / (1024.0 * 1024.0));
+				if (Info.SourcesWithoutModule.Count > 0)
+				{
+					foreach (string Source in Info.SourcesWithoutModule)
+					{
+						System.Console.WriteLine("\t{0}", Source);
+					}
+				}
+				else
+				{
+					System.Console.WriteLine("None!");
+				}
+			}
 
 			// sort modules alphabetically and print out
 			List<KeyValuePair<string, uint>> ModInfos = Info.Modules.ToList();
 
-			System.Console.WriteLine("\n\n\n\n-------------------------------------------------------------------------------------------------------------------------------------------------------------------");
-			System.Console.WriteLine("Exec itemization, modules sorted alphabetically by name");
-			System.Console.WriteLine("-------------------------------------------------------------------------------------------------------------------------------------------------------------------");
-			System.Console.WriteLine("Module, Megabytes (float), Bytes (uint)");
+			if (bPrintResults)
+			{
+				System.Console.WriteLine("\n\n\n\n-------------------------------------------------------------------------------------------------------------------------------------------------------------------");
+				System.Console.WriteLine("Exec itemization, modules sorted alphabetically by name");
+				System.Console.WriteLine("-------------------------------------------------------------------------------------------------------------------------------------------------------------------");
+				System.Console.WriteLine("Module, Megabytes (float), Bytes (uint)");
+			}
 			ModInfos.SortBy(ModInfo => ModInfo.Key);
 
 			// Create a new string builder for CSV output
@@ -428,36 +758,73 @@ namespace ReportItemizedExecutableCode.Automation
 
 			foreach (KeyValuePair<string, uint> ModInfo in ModInfos)
 			{
-				System.Console.WriteLine("{0}, {1:F2}, {2}", ModInfo.Key, ModInfo.Value / (1024.0 * 1024.0), ModInfo.Value);
+				if (bPrintResults)
+				{
+					System.Console.WriteLine("{0}, {1:F2}, {2}", ModInfo.Key, ModInfo.Value / (1024.0 * 1024.0), ModInfo.Value);
+				}
 				Builder.AppendLine(string.Format("{0},{1:F2},{2}", ModInfo.Key, ModInfo.Value / (1024.0 * 1024.0), ModInfo.Value));
 			}
 
-			System.Console.WriteLine("\n\n\n\n-------------------------------------------------------------------------------------------------------------------------------------------------------------------");
-			System.Console.WriteLine("Exec itemization, modules sorted by their size");
-			System.Console.WriteLine("-------------------------------------------------------------------------------------------------------------------------------------------------------------------");
-			System.Console.WriteLine("Module, Megabytes (float), Bytes (uint)");
 			ModInfos.SortBy(ModInfo => ModInfo.Value);
 			ModInfos.Reverse();
-			foreach (KeyValuePair<string, uint> ModInfo in ModInfos)
+
+			if (bPrintResults)
 			{
-				System.Console.WriteLine("{0}, {1:F2}, {2}", ModInfo.Key, ModInfo.Value / (1024.0 * 1024.0), ModInfo.Value);
+				System.Console.WriteLine("\n\n\n\n-------------------------------------------------------------------------------------------------------------------------------------------------------------------");
+				System.Console.WriteLine("Exec itemization, modules sorted by their size");
+				System.Console.WriteLine("-------------------------------------------------------------------------------------------------------------------------------------------------------------------");
+				System.Console.WriteLine("Module, Megabytes (float), Bytes (uint)");
+				foreach (KeyValuePair<string, uint> ModInfo in ModInfos)
+				{
+					System.Console.WriteLine("{0}, {1:F2}, {2}", ModInfo.Key, ModInfo.Value / (1024.0 * 1024.0), ModInfo.Value);
+				}
 			}
 
 			foreach (string DestPath in DestPaths)
 			{
-				System.Console.WriteLine($"\nWriting itemized module information to: {DestPath}");
+				System.Console.WriteLine($"Writing itemized module information to: {DestPath}");
 
 				string DirPath = Path.GetDirectoryName(DestPath);
 				Directory.CreateDirectory(DirPath);
 
 				File.WriteAllText(DestPath, Builder.ToString());
-
-				System.Console.WriteLine("\n\n");
 			}
 		}
 
-		// Hack for a quick iteration on formatting - stores the results from bloaty.
-		//string HardCodedResults;
+		private void WriteTotalsReport(ModuleInfos Info, string[] DestPaths)
+		{
+			StringBuilder Builder = new StringBuilder("Name,Megabytes,Bytes\n");
+			// Grand total
+			Builder.AppendLine($"Total,{Info.Total / (1024.0 * 1024.0):F2},{Info.Total}");
+			// Sections making up total
+			Builder.AppendLine($"Engine,{Info.TotalEngineSize / (1024.0 * 1024.0):F2},{Info.TotalEngineSize}");
+			Builder.AppendLine($"Game,{Info.TotalGameSize / (1024.0 * 1024.0):F2},{Info.TotalGameSize}");
+			Builder.AppendLine($"Unattributed,{Info.TotalUnattributed / (1024.0 * 1024.0):F2},{Info.TotalUnattributed}");
+			Builder.AppendLine($"ExternalLibs,{Info.TotalExternalLibs / (1024.0 * 1024.0):F2},{Info.TotalExternalLibs}");
+			Builder.AppendLine($"Unclassified,{Info.TotalUnclassified / (1024.0 * 1024.0):F2},{Info.TotalUnclassified}");
+
+			// Engine breakdown
+			uint EngineBase = Info.TotalEngineSize - Info.TotalEnginePluginSize;
+			Builder.AppendLine($"Engine/Base,{EngineBase / (1024.0 * 1024.0):F2},{EngineBase}");
+			Builder.AppendLine($"Engine/Plugins,{Info.TotalEnginePluginSize / (1024.0 * 1024.0):F2},{Info.TotalEnginePluginSize}");
+
+			// Game breakdown
+			uint GameBase = Info.TotalGameSize - Info.TotalGamePluginSize;
+			uint GamePlugins = Info.TotalGamePluginSize - Info.TotalGameFeaturesSize;
+			Builder.AppendLine($"Game/Base,{GameBase / (1024.0 * 1024.0):F2},{GameBase}");
+			Builder.AppendLine($"Game/Plugins,{GamePlugins / (1024.0 * 1024.0):F2},{GamePlugins}");
+			Builder.AppendLine($"Game/GameFeatures,{Info.TotalGameFeaturesSize / (1024.0 * 1024.0):F2},{Info.TotalGameFeaturesSize}");
+
+			foreach (string DestPath in DestPaths)
+			{
+				System.Console.WriteLine($"Writing totals information to: {DestPath}");
+
+				string DirPath = Path.GetDirectoryName(DestPath);
+				Directory.CreateDirectory(DirPath);
+
+				File.WriteAllText(DestPath, Builder.ToString());
+			}
+		}
 
 		private void WriteSectionsReport(FileReference ExecutableFile, string[] DestPaths)
 		{
@@ -487,14 +854,12 @@ namespace ReportItemizedExecutableCode.Automation
 
 			foreach (string DestPath in DestPaths)
 			{
-				System.Console.WriteLine($"\nWriting sections report to: {DestPath}");
+				System.Console.WriteLine($"Writing sections report to: {DestPath}");
 
 				string DirPath = Path.GetDirectoryName(DestPath);
 				Directory.CreateDirectory(DirPath);
 
 				File.WriteAllText(DestPath, Builder.ToString());
-
-				System.Console.WriteLine("\n\n");
 			}
 		}
 
diff --git a/Engine/Source/Programs/UnrealBuildTool/Modes/QueryMode.cs b/Engine/Source/Programs/UnrealBuildTool/Modes/QueryMode.cs
--- a/Engine/Source/Programs/UnrealBuildTool/Modes/QueryMode.cs
+++ b/Engine/Source/Programs/UnrealBuildTool/Modes/QueryMode.cs
@@ -595,6 +595,10 @@ namespace UnrealBuildTool
 
 						var Info = new
 						{
+							ShortName = Module.Rules.ShortName,
+							ModuleDirectory = Module.ModuleDirectory.FullName,
+							// Strip the query suffix from the target name so we output the real path.
+							IntermediateDirectory = Module.IntermediateDirectory.FullName.Replace($"{CurrentTarget.TargetName}QRY", CurrentTarget.TargetName, StringComparison.Ordinal),
 							Public = new SortedSet<string>(Module.PublicDependencyModules?.Where(x => x is UEBuildModuleCPP).Select(x => x.Name) ?? []),
 							Private = new SortedSet<string>(Module.PrivateDependencyModules?.Where(x => x is UEBuildModuleCPP).Select(x => x.Name) ?? []),
 							PublicInclude = new SortedSet<string>(Module.PublicIncludePathModules?.Where(x => x is UEBuildModuleCPP).Select(x => x.Name) ?? []),
