SubmitTool RPC test

#rnx

#Synced-CL 0

diff --git a/Engine/Plugins/RPCBase/RPCBase.uplugin b/Engine/Plugins/RPCBase/RPCBase.uplugin
new file mode 100644
--- /dev/null
+++ b/Engine/Plugins/RPCBase/RPCBase.uplugin
@@ -0,0 +1,24 @@
+{
+	"FileVersion": 3,
+	"Version": 8,
+	"VersionName": "1.0",
+	"FriendlyName": "Remote Procedure Calls Base",
+	"Description": "RPC base framework that client applications can use to implement RPC libraries.",
+	"Category": "Runtime",
+	"CreatedBy": "Epic Games, Inc.",
+	"CreatedByURL": "https://epicgames.com",
+	"DocsURL": "",
+	"MarketplaceURL": "",
+	"SupportURL": "",
+	"EnabledByDefault": false,
+	"CanContainContent": false,
+	"IsBetaVersion": false,
+	"Installed": false,
+	"Modules": [
+		{
+			"Name": "RPCBase",
+			"Type": "Runtime",
+			"LoadingPhase": "Default"
+		}
+	]
+}
diff --git a/Engine/Plugins/RPCBase/Source/RPCBase/Private/RPCLibraryBase.cpp b/Engine/Plugins/RPCBase/Source/RPCBase/Private/RPCLibraryBase.cpp
new file mode 100644
--- /dev/null
+++ b/Engine/Plugins/RPCBase/Source/RPCBase/Private/RPCLibraryBase.cpp
@@ -0,0 +1,120 @@
+// Copyright Epic Games, Inc. All Rights Reserved.
+
+#include "RPCLibraryBase.h"
+
+#include "Misc/CommandLine.h"
+
+DEFINE_LOG_CATEGORY_STATIC(LogRPCLibraryBase, Log, All);
+
+URPCLibraryBase::URPCLibraryBase()
+{
+
+}
+
+URPCLibraryBase::~URPCLibraryBase()
+{
+	if (Registrator != nullptr)
+	{
+		Registrator->DeregisterHttpCallbacks();
+		Registrator->RemoveFromRoot();
+	}
+}
+
+void URPCLibraryBase::Initialize()
+{
+#if WITH_RPC_REGISTRY
+	UE_LOG(LogRPCLibraryBase, Display, TEXT("RPC library initialization"));
+
+	check(UExternalRpcRegistry::GetInstance());
+
+	Registrator = NewObject<UExternalRpcRegistrationComponent>();
+	FParse::Value(FCommandLine::Get(), TEXT("externalrpclistenaddress="), Registrator->ListenerAddress);
+	FParse::Value(FCommandLine::Get(), TEXT("rpcsenderid="), Registrator->SenderID);
+
+	Registrator->AddToRoot();
+
+	// Default RPC "RegistrationFinished" that sets RegistrationFinished to true to be used as a "ready" flag for issuing RPCs
+	RegisterRPC(
+		TEXT("RegistrationFinished"),
+		FHttpPath(TEXT("/rpclibraryybase/registrationfinished")),
+		EHttpServerRequestVerbs::VERB_POST,
+		FHttpRequestHandler::CreateUObject(this, &URPCLibraryBase::HandleRegistrationFinishedRoute));
+#endif // WITH_RPC_REGISTRY
+}
+
+void URPCLibraryBase::StartListening()
+{
+	static bool bListenersStarted = false;
+	if (!bListenersStarted)
+	{
+		UE_LOG(LogRPCLibraryBase, Display, TEXT("FHttpServerModule: calling StartAllListeners"));
+		FHttpServerModule::Get().StartAllListeners();
+		bListenersStarted = true;
+	}
+}
+
+#if WITH_RPC_REGISTRY
+void URPCLibraryBase::RegisterRPC(FName RouteName, const FHttpPath& HttpPath, const EHttpServerRequestVerbs& RequestVerbs, const FHttpRequestHandler& Handler, bool bOverrideIfBound, FString OptionalCategory, FString OptionalContentType, TArray<FExternalRpcArgumentDesc> OptionalInArguments)
+{
+	if (!Registrator->RegisteredRoutes.Contains(RouteName))
+	{
+		UE_LOG(LogRPCLibraryBase, Display, TEXT("Registering new route %s at path %s"), *RouteName.ToString(), *HttpPath.GetPath());
+
+		FExternalRouteInfo NewRouteInfo;
+		NewRouteInfo.RouteName = RouteName;
+		NewRouteInfo.RoutePath = HttpPath;
+		NewRouteInfo.RequestVerbs = RequestVerbs;
+		NewRouteInfo.InputContentType = MoveTemp(OptionalContentType);
+		NewRouteInfo.ExpectedArguments = MoveTemp(OptionalInArguments);
+		NewRouteInfo.RpcCategory = MoveTemp(OptionalCategory);
+		Registrator->RegisterHttpCallback(MoveTemp(NewRouteInfo), Registrator->CreateRouteHandle(Handler), bOverrideIfBound);
+	}
+}
+
+bool URPCLibraryBase::HandleRegistrationFinishedRoute(const FHttpServerRequest& Request, const FHttpResultCallback& OnComplete)
+{
+	FString ResponseStr;
+	TSharedRef<TJsonWriter<>> JsonWriter = TJsonWriterFactory<>::Create(&ResponseStr);
+	JsonWriter->WriteObjectStart();
+	JsonWriter->WriteValue(TEXT("RegistrationFinished"), RegistrationFinished);
+	JsonWriter->WriteObjectEnd();
+	JsonWriter->Close();
+
+	auto Response = Registrator->CreateSimpleResponse(true, ResponseStr, false);
+	OnComplete(MoveTemp(Response));
+
+	return true;
+}
+
+TUniquePtr<FHttpServerResponse> URPCLibraryBase::CreateSimpleResponse(bool bInWasSuccessful, FString InValue, bool bInFatal)
+{
+	return Registrator->CreateSimpleResponse(bInWasSuccessful, InValue, bInFatal);
+}
+
+bool URPCLibraryBase::DeserializeRequest(TSharedPtr<FJsonObject>& RootObject, const TSharedRef<TJsonReader<>>& JsonReader, const TSharedRef<TJsonWriter<>>& JsonWriter)
+{
+	if (!FJsonSerializer::Deserialize(JsonReader, RootObject))
+	{
+		FString ErrorMessage = TEXT("Cannot deserialize RootObject from request body.");
+		UE_LOG(LogRPCLibraryBase, Warning, TEXT("%s"), *ErrorMessage);
+		JsonWriter->WriteValue(TEXT("Error"), ErrorMessage);
+		JsonWriter->WriteObjectEnd();
+		JsonWriter->Close();
+		return false;
+	}
+	return true;
+}
+
+bool URPCLibraryBase::GetStringField(TSharedPtr<FJsonObject>& RootObject, FString FieldName, FString& OutFieldValue, const TSharedRef<TJsonWriter<>>& JsonWriter)
+{
+	if (!RootObject->TryGetStringField(FieldName, OutFieldValue))
+	{
+		FString ErrorMessage = FString::Printf(TEXT("Cannot get field %s from request body."), *FieldName);
+		JsonWriter->WriteValue(TEXT("Error"), ErrorMessage);
+		JsonWriter->WriteObjectEnd();
+		JsonWriter->Close();
+		return false;
+	}
+	return true;
+}
+#endif // WITH_RPC_REGISTRY
diff --git a/Engine/Plugins/RPCBase/Source/RPCBase/Public/RPCLibraryBase.h b/Engine/Plugins/RPCBase/Source/RPCBase/Public/RPCLibraryBase.h
new file mode 100644
--- /dev/null
+++ b/Engine/Plugins/RPCBase/Source/RPCBase/Public/RPCLibraryBase.h
@@ -0,0 +1,46 @@
+// Copyright Epic Games, Inc. All Rights Reserved.
+
+#pragma once
+
+#include "Serialization/JsonReader.h"
+#include "Serialization/JsonWriter.h"
+#include "Serialization/JsonSerializer.h"
+
+#if WITH_RPC_REGISTRY
+#include "ExternalRpcRegistrationComponent.h"
+#include "ExternalRpcRegistry.h"
+#include "HttpRequestHandler.h"
+#include "HttpServerModule.h"
+#endif
+
+#include "RPCLibraryBase.generated.h"
+
+#if WITH_RPC_REGISTRY
+enum class EHttpServerRequestVerbs : uint16;
+struct FHttpPath;
+struct FHttpServerResponse;
+struct FExternalRouteInfo;
+struct FExternalRpcArgumentDesc;
+#endif
+
+UCLASS(MinimalAPI)
+class URPCLibraryBase : public UObject
+{
+	GENERATED_BODY()
+private:
+	UExternalRpcRegistrationComponent* Registrator;
+public:	
+	URPCLibraryBase();
+	~URPCLibraryBase();
+protected:
+	void Initialize();
+	void StartListening();
+	bool RegistrationFinished = false;
+#if WITH_RPC_REGISTRY
+	void RegisterRPC(FName RouteName, const FHttpPath& HttpPath, const EHttpServerRequestVerbs& RequestVerb, const FHttpRequestHandler& Handler, bool bOverrideIfBound = false, FString OptionalCategory = {}, FString OptionalContentType = {}, TArray<FExternalRpcArgumentDesc> OptionalInArguments = TArray<FExternalRpcArgumentDesc>());
+	bool HandleRegistrationFinishedRoute(const FHttpServerRequest& Request, const FHttpResultCallback& OnComplete);
+	TUniquePtr<FHttpServerResponse> CreateSimpleResponse(bool bInWasSuccessful, FString InValue = "", bool bInFatal = false);
+	bool DeserializeRequest(TSharedPtr<FJsonObject>& RootObject, const TSharedRef<TJsonReader<>>& JsonReader, const TSharedRef<TJsonWriter<>>& JsonWriter);
+	bool GetStringField(TSharedPtr<FJsonObject>&  RootObject, FString FiledName, FString& OutFieldValue, const TSharedRef<TJsonWriter<>>& JsonWriter);
+#endif
+};
\ No newline at end of file
diff --git a/Engine/Plugins/RPCBase/Source/RPCBase/RPCBase.Build.cs b/Engine/Plugins/RPCBase/Source/RPCBase/RPCBase.Build.cs
new file mode 100644
--- /dev/null
+++ b/Engine/Plugins/RPCBase/Source/RPCBase/RPCBase.Build.cs
@@ -0,0 +1,27 @@
+// Copyright Epic Games, Inc. All Rights Reserved.
+
+using UnrealBuildTool;
+
+public class RPCBase : ModuleRules
+{
+	public RPCBase(ReadOnlyTargetRules target) : base(target)
+	{
+		PrivateDependencyModuleNames.AddRange(
+		[
+			"Core",
+			"CoreUObject",
+			"Json"
+		]);
+
+		if (Target.Configuration < UnrealTargetConfiguration.Shipping)
+		{
+			PublicDependencyModuleNames.Add("ExternalRpcRegistry");
+			PublicDependencyModuleNames.Add("HTTPServer");
+			PublicDefinitions.Add("WITH_RPC_REGISTRY=1");
+		}
+		else
+		{
+			PublicDefinitions.Add("WITH_RPC_REGISTRY=0");
+		}
+	}
+}
\ No newline at end of file
diff --git a/Engine/Source/Developer/SubmitToolCore/Public/CommandLine/CmdLineParameters.h b/Engine/Source/Developer/SubmitToolCore/Public/CommandLine/CmdLineParameters.h
--- a/Engine/Source/Developer/SubmitToolCore/Public/CommandLine/CmdLineParameters.h
+++ b/Engine/Source/Developer/SubmitToolCore/Public/CommandLine/CmdLineParameters.h
@@ -15,6 +15,8 @@ namespace FSubmitToolCmdLine
 	const FString RootDir = TEXT("root-dir");
 	const FString EditorFlag = TEXT("from-editor");
 	const FString EnableLocks = TEXT("enablelocks");
+	const FString TestMode = TEXT("test-mode");
+	const FString WaitForDebugger = TEXT("wait-for-debugger");
 
 	const TArray<TSharedPtr<FCmdLineParameter>> SubmitToolCmdLineArgs =
 	{
@@ -60,6 +62,8 @@ namespace FSubmitToolCmdLine
 		}),
 		MakeShared<FCmdLineParameter>(EditorFlag, false, TEXT("Flag indicating if submit tool is being invoked from unreal editor."), true),
 		MakeShared<FCmdLineParameter>(EnableLocks, false, TEXT("Enables lockdown files previously defined by the config by their id")),
+		MakeShared<FCmdLineParameter>(TestMode, false, TEXT("Submit tool launched to execute RPC tests."), true),
+		MakeShared<FCmdLineParameter>(WaitForDebugger, false, TEXT("Submit tool waits for a debugger to attach."), true)
 	};
 }
 
diff --git a/Engine/Source/Programs/AutomationTool/Gauntlet/Framework/RpcFramework/Gauntlet.RpcExecutor.cs b/Engine/Source/Programs/AutomationTool/Gauntlet/Framework/RpcFramework/Gauntlet.RpcExecutor.cs
--- a/Engine/Source/Programs/AutomationTool/Gauntlet/Framework/RpcFramework/Gauntlet.RpcExecutor.cs
+++ b/Engine/Source/Programs/AutomationTool/Gauntlet/Framework/RpcFramework/Gauntlet.RpcExecutor.cs
@@ -13,6 +13,7 @@ using System.Net.Http;
 using System.Net.Sockets;
 using System.Runtime.CompilerServices;
 using System.Text;
+using System.Text.Json.Nodes;
 using System.Text.RegularExpressions;
 using System.Threading;
 using ThreadState = System.Threading.ThreadState;
@@ -884,10 +885,10 @@ namespace Gauntlet
 			string RpcUrl = InTarget.CreateRpcUrl(RpcName);
 			RpcDefinition RpcEntry = InTarget.GetRpc(RpcName);
 
-			Log.VeryVerbose("Calling RPC" + RpcName + " on target " + InTarget.TargetName + " at " + RpcUrl);
+			Log.VeryVerbose($"Calling RPC {RpcName} on target {InTarget.TargetName} at {RpcUrl}");
 			if (string.IsNullOrEmpty(RpcUrl) || RpcEntry == null)
 			{
-				Log.Verbose("Failed Calling RPC" + RpcName);
+				Log.Verbose($"Failed Calling RPC {RpcName}");
 				return null;
 			}
 			if (InArgs == null)
@@ -1002,12 +1003,12 @@ namespace Gauntlet
 		
 		public static HttpResponseMessage CallRPCWithUrlParams(RpcTarget InTarget, string RpcName, Dictionary<string, object> InArgs, int Timeout = 5000,  string InContentType = "application/json")
         {
-        	Log.VeryVerbose("Calling RPC" + RpcName);
+        	Log.VeryVerbose($"Calling RPC {RpcName}");
         	string RpcUrl = InTarget.CreateRpcUrl(RpcName);
         	RpcDefinition RpcEntry = InTarget.GetRpc(RpcName);
         	if (string.IsNullOrEmpty(RpcUrl) || RpcEntry == null)
         	{
-        		Log.Verbose("Failed Calling RPC" + RpcName);
+        		Log.Verbose($"Failed Calling RPC {RpcName}");
         		return null;
         	}
             
diff --git a/Engine/Source/Programs/AutomationTool/Gauntlet/Framework/RpcFramework/Gauntlet.RpcLibrary.cs b/Engine/Source/Programs/AutomationTool/Gauntlet/Framework/RpcFramework/Gauntlet.RpcLibrary.cs
--- a/Engine/Source/Programs/AutomationTool/Gauntlet/Framework/RpcFramework/Gauntlet.RpcLibrary.cs
+++ b/Engine/Source/Programs/AutomationTool/Gauntlet/Framework/RpcFramework/Gauntlet.RpcLibrary.cs
@@ -1,5 +1,7 @@
 // Copyright Epic Games, Inc. All Rights Reserved.
 
+using AutomationTool;
+using EpicGames.Core;
 using Newtonsoft.Json;
 using System;
 using System.Collections.Generic;
@@ -7,8 +9,10 @@ using System.Diagnostics;
 using System.Diagnostics.CodeAnalysis;
 using System.IO;
 using System.Linq;
+using System.Net;
 using System.Net.Http;
 using System.Threading;
+using static Gauntlet.RpcLibrary;
 
 namespace Gauntlet
 {
@@ -750,4 +754,112 @@ namespace Gauntlet
 		public bool WasSuccessful = false;
 		public string Message = string.Empty;
 	}
+
+	public class BaseRpcResponse : SimpleResponse
+	{
+		public HttpResponseMessage HttpResponse { get; set; } = null;
+		public string RawContent { get; set; } = null;
+	}
+
+
+	/// <summary>
+	/// Simple base class for building RPC libraries.
+	/// </summary>
+	public class BaseRpcLibrary
+	{
+		const string REGISTRATION_FINISHED = "RegistrationFinished";
+
+		/// <summary>
+		/// Runs a RPC for a given target.
+		/// </summary>
+		public BaseRpcResponse RunRpc(RpcTarget InTarget, string InRpcName, Dictionary<string, object> InRequestArgs, int InTimeout = 5000, bool InSilenceErrors = false)
+		{
+			ValidateRpcTarget(InTarget);
+
+			string MessageContent = string.Empty;
+
+			HttpResponseMessage MessageResponse = RpcExecutor.CallRpc(InTarget, InRpcName, InTimeout, InRequestArgs);
+			if (MessageResponse == null)
+			{
+				if (!InSilenceErrors)
+				{
+					throw new AutomationException($"An error occurred calling {InRpcName}. Check the logs.");
+				}
+			}
+			else
+			{
+				var RunValidatorStatus = MessageResponse.StatusCode;
+				MessageContent = MessageResponse.Content.ReadAsStringAsync().Result;
+				Log.Info($"=============>RPC test attempt is: {RunValidatorStatus} | {MessageContent}");
+			}
+			SimpleResponse SimpleResponse = GetSimpleResponse(MessageResponse);
+			return new BaseRpcResponse
+			{
+				Succeeded = SimpleResponse.Succeeded,
+				Value = SimpleResponse.Value,
+				Fatal = SimpleResponse.Fatal,
+				HttpResponse = MessageResponse,
+				RawContent = MessageContent
+			};
+		}
+
+		private void ValidateRpcTarget(RpcTarget Target)
+		{
+			if (null == Target)
+			{
+				throw new AutomationException("The RpcTarget is null.");
+			}
+
+			if ("localhost" != Target.IpAddress.ToLower())
+			{
+				IPAddress IpAddress = null;
+				if (!IPAddress.TryParse(Target.IpAddress, out IpAddress))
+				{
+					throw new AutomationException($"The rpcTarget has an invalid IPAddress: {Target.IpAddress}");
+				}
+			}
+		}
+
+		/// <summary>
+		/// Run the default "RegistrationFinished" RPC that checks if registration has finished.
+		/// </summary>
+		public BaseRpcResponse RegistrationFinished(RpcTarget InTarget)
+		{
+			return RunRpc(InTarget, REGISTRATION_FINISHED, null, 0, true);
+		}
+
+		public JsonObject GetResponseValueAsJson(SimpleResponse InResponse)
+		{
+			if (InResponse == null)
+			{
+				throw new AutomationException("[GetResponseValueAsJson] SimpleResponse cannot be null.");
+			}
+
+			if (!InResponse.Succeeded)
+			{
+				Log.Warning("[GetResponseValueAsJson] SimpleResponse \"succeeded\" field is false.");
+			}
+
+			if (InResponse.Succeeded)
+			{
+				return JsonObject.Parse(InResponse.Value);
+			}
+			else if (InResponse.Fatal)
+			{
+				throw new AutomationException("[GetResponseValueAsJson] SimpleResponse \"fatal\" field is true, its value field is: {0}.", InResponse.Value);
+			}
+			return new JsonObject();
+		}
+
+		public bool GetResponseBoolFieldValue(SimpleResponse InResponse, string InBoolFieldName)
+		{
+			JsonObject ResponseJson = GetResponseValueAsJson(InResponse);
+			bool ResponseBoolFieldValue = false;
+			if(!ResponseJson.TryGetBoolField(InBoolFieldName, out ResponseBoolFieldValue))
+			{
+				throw new AutomationException("[GetResponseBoolFieldValue] Failed to get bool field {0} from response json.", InBoolFieldName);
+			}
+			return ResponseBoolFieldValue;
+		}
+	}
 }
\ No newline at end of file
diff --git a/Engine/Source/Programs/AutomationTool/SubmitTool/RunSubmitToolTests.cs b/Engine/Source/Programs/AutomationTool/SubmitTool/RunSubmitToolTests.cs
new file mode 100644
--- /dev/null
+++ b/Engine/Source/Programs/AutomationTool/SubmitTool/RunSubmitToolTests.cs
@@ -0,0 +1,73 @@
+// Copyright Epic Games, Inc. All Rights Reserved.
+
+using AutomationTool;
+using Gauntlet;
+using System.Collections.Generic;
+using UnrealBuildBase;
+
+namespace SubmitToolTests
+{
+	public class RunSubmitToolTests : RunApplicationTests
+	{
+		private P4Connection CurrentStream;
+		private int TestCL = 0;
+
+		public override ExitCode Execute()
+		{
+			return base.Execute();
+		}
+
+		public override ApplicationTestExecutorOptions ConstructApplicationTestExecutorOptions()
+		{
+			return new SubmitToolTestExecutorOptions();
+		}
+
+		public override ApplicationTestsBuildSource ConstructApplicationTestsBuildSource(ApplicationTestExecutorOptions InTestExecutorOptions)
+		{
+			return new SubmitToolBuildSource((SubmitToolTestExecutorOptions)InTestExecutorOptions);
+		}
+
+		public override ApplicationTestContext ConstructApplicationTestContext(ApplicationTestsBuildSource BuildSource, ApplicationTestRoleContext RoleContext, ApplicationTestExecutorOptions InOptions)
+		{
+			return new SubmitToolTestContext(BuildSource, RoleContext, (SubmitToolTestExecutorOptions)InOptions);
+		}
+
+		public override IEnumerable<string> GetTestNamespaces()
+		{
+			return ["SubmitApplicationTests"];
+		}
+
+		public override string GetTestName()
+		{
+			return "SubmitToolRPCTests";
+		}
+
+		public override void ValidateTestNode(ITestNode NewTest)
+		{
+			if (!(NewTest is SubmitToolRPCTests))
+			{
+				throw new AutomationException("Expected ITestNode type of SubmitToolRPCTests");
+			}
+		}
+
+		protected override void SetupGlobal(ApplicationTestExecutorOptions Context)
+		{
+			P4ClientInfo Client = P4.GetClientInfo(P4Env.Client);
+			CurrentStream = new P4Connection(P4Env.User, Client.Name);
+			TestCL = CurrentStream.CreateChange();
+			((SubmitToolTestExecutorOptions)Context).TestCL = TestCL;
+
+			Log.Info($"New test changelist: {TestCL}");
+		}
+
+		protected override void TeardownGlobal(ApplicationTestExecutorOptions Context)
+		{
+			if (CurrentStream != null && TestCL != 0)
+			{
+				Log.Info($"Reverting test changelist: {TestCL}");
+				CurrentStream.RevertAll(TestCL);
+				CurrentStream.DeleteChange(TestCL);
+			}
+		}
+	}
+}
diff --git a/Engine/Source/Programs/AutomationTool/SubmitTool/SubmitTool.Automation.csproj b/Engine/Source/Programs/AutomationTool/SubmitTool/SubmitTool.Automation.csproj
new file mode 100644
--- /dev/null
+++ b/Engine/Source/Programs/AutomationTool/SubmitTool/SubmitTool.Automation.csproj
@@ -0,0 +1,47 @@
+<Project Sdk="Microsoft.NET.Sdk">
+  <Import Project="..\..\Shared\UnrealEngine.csproj.props" />
+  
+  <PropertyGroup>
+    <TargetFramework>net8.0</TargetFramework>
+	<Configuration Condition=" '$(Configuration)' == '' ">Development</Configuration>
+    <OutputType>Library</OutputType>
+	<RootNamespace>SubmitTool.Automation</RootNamespace>
+    <AssemblyName>SubmitTool.Automation</AssemblyName>
+    <GenerateAssemblyInfo>false</GenerateAssemblyInfo>
+    <GenerateTargetFrameworkAttribute>false</GenerateTargetFrameworkAttribute>
+    <OutputPath>..\..\..\..\Binaries\DotNET\AutomationTool\AutomationScripts\SubmitTool</OutputPath>
+    <Configurations>Debug;Release;Development</Configurations>
+    <TreatWarningsAsErrors>true</TreatWarningsAsErrors>
+    <AppendTargetFrameworkToOutputPath>false</AppendTargetFrameworkToOutputPath>
+    <DebugType>pdbonly</DebugType>
+	<SatelliteResourceLanguages>en</SatelliteResourceLanguages>    <!-- remove non english resource languages -->
+  </PropertyGroup>
+  <PropertyGroup Condition=" '$(Configuration)|$(Platform)' == 'Development|AnyCPU' ">
+    <DefineConstants>$(DefineConstants);TRACE</DefineConstants>
+    <Optimize>true</Optimize>
+  </PropertyGroup>
+  <PropertyGroup Condition=" '$(Configuration)|$(Platform)' == 'Release|AnyCPU' ">
+    <DefineConstants>$(DefineConstants);TRACE</DefineConstants>
+    <Optimize>true</Optimize>
+  </PropertyGroup>
+  <PropertyGroup Condition=" '$(Configuration)|$(Platform)' == 'Debug|AnyCPU' ">
+    <DefineConstants>$(DefineConstants);DEBUG;TRACE</DefineConstants>
+    <DebugSymbols>true</DebugSymbols>
+    <DebugType>full</DebugType>
+  </PropertyGroup>
+  <ItemGroup>
+    <ProjectReference Include="..\..\Shared\EpicGames.Core\EpicGames.Core.csproj" PrivateAssets="All">
+      <Private>false</Private>
+    </ProjectReference>
+    <ProjectReference Include="..\..\UnrealBuildTool\UnrealBuildTool.csproj" PrivateAssets="All">
+      <Private>false</Private>
+    </ProjectReference>
+    <ProjectReference Include="..\AutomationUtils\AutomationUtils.Automation.csproj" PrivateAssets="All">
+      <Private>false</Private>
+    </ProjectReference>
+    <ProjectReference Include="..\Scripts\AutomationScripts.Automation.csproj" PrivateAssets="All">
+      <Private>false</Private>
+    </ProjectReference>
+    <ProjectReference Include="..\Gauntlet\Gauntlet.Automation.csproj"/>
+  </ItemGroup>
+</Project>
diff --git a/Engine/Source/Programs/AutomationTool/SubmitTool/Tests/Gauntlet.SubmitToolTestNode.cs b/Engine/Source/Programs/AutomationTool/SubmitTool/Tests/Gauntlet.SubmitToolTestNode.cs
new file mode 100644
--- /dev/null
+++ b/Engine/Source/Programs/AutomationTool/SubmitTool/Tests/Gauntlet.SubmitToolTestNode.cs
@@ -0,0 +1,460 @@
+// Copyright Epic Games, Inc. All Rights Reserved.
+
+using EpicGames.Core;
+using Gauntlet;
+using Gauntlet.Report;
+using Gauntlet.TestTracking;
+using SubmitTool.Automation;
+using System;
+using System.Collections.Generic;
+using System.Linq;
+using System.Net;
+using System.Net.NetworkInformation;
+using System.Net.Sockets;
+using System.Threading;
+using System.Threading.Tasks;
+using static Gauntlet.RpcLibrary;
+using Log = Gauntlet.Log;
+
+namespace SubmitToolTests
+{
+	public class SubmitToolRPCTests : ApplicationTest
+	{
+		public new SubmitToolTestContext Context { get { return (SubmitToolTestContext)base.Context; } }
+
+		private readonly List<string> ValidationSuccessStates = new List<string>() { "Valid", "Skipped" };
+		private readonly List<string> ValidationFailedStates = new List<string>() { "Failed", "Timeout" };
+
+		private readonly string SubmitToolRpcClient= "SubmitToolRpcClient";
+
+		private RpcTargetRegistry TargetRegistry = new RpcTargetRegistry();
+
+		private bool SubmitToolRPCsReady = false;
+		private bool DoneQueuingCases = false;
+
+		private TestCaseQueue<SubmitToolRPCTestCase, TestCaseState> TestCaseQueue = new TestCaseQueue<SubmitToolRPCTestCase, TestCaseState>();
+
+		public SubmitToolRPCTests(ApplicationTestContext InContext) : base(InContext)
+		{
+			StopTestOnInactivity = true;
+			MaxDuration = 30.0f * 60.0f; // 30 minutes
+		}
+
+		public override string Name { get { return "SubmitToolRPCTest"; } }
+
+		public override float MaxDuration { protected set; get; }
+
+		public override string GetTestSummary()
+		{
+			return "SubmitTool Test";
+		}
+
+		public override bool StartTest(int Pass, int NumPasses)
+		{
+			bool Result = base.StartTest(Pass, NumPasses);
+			return Result;
+		}
+
+		public override bool QueueTestCases()
+		{
+			if (TargetRegistry.RpcTargets.IsEmpty)
+			{
+				string ClientIP = "localhost";
+				TargetRegistry.AddNewTarget(SubmitToolRpcClient, ClientIP, this.Context.BuildInfo.BasePort, UnrealTargetRole.Client);
+
+				RpcExecutor.Instance.StartListenThread();
+			}
+			else if (!DoneQueuingCases)
+			{
+				RpcLibrary.WaitForAllTargetsToGetEndpoints(TargetRegistry.GetAllTargetsOfType(UnrealTargetRole.Client), 60, 10);
+				TargetRegistry.UpdateAvailableRpcForAllTargets();
+
+				if (SubmitToolRPCsReady)
+				{
+					// Add more cases here
+					TestCaseQueue.Enqueue(new SubmitToolRPCTestCase(this, TargetRegistry.GetTarget(SubmitToolRpcClient), Phases, SubmitToolRPCs.SetChangelistDescription, new Dictionary<string, object>() { ["Description"] = "Test description #jira nojira" }));
+					TestCaseQueue.Enqueue(new SubmitToolRPCTestCase(this, TargetRegistry.GetTarget(SubmitToolRpcClient), Phases, SubmitToolRPCs.RunValidator, new Dictionary<string, object>() { ["ValidatorName"] = "Tag" }));
+					TestCaseQueue.Enqueue(new SubmitToolRPCTestCase(this, TargetRegistry.GetTarget(SubmitToolRpcClient), Phases, SubmitToolRPCs.GetValidatorState, new Dictionary<string, object>() { ["ValidatorName"] = "Tag" },
+						(SimpleResponse ResponseMessage) =>
+						{
+							if (ResponseMessage != null && ResponseMessage.Succeeded)
+							{
+								JsonObject JsonResult = JsonObject.Parse(ResponseMessage.Value);
+								string ValidatorState;
+								if (JsonResult.TryGetStringField("ValidatorState", out ValidatorState))
+								{
+									if (ValidationSuccessStates.Contains(ValidatorState))
+									{
+										return TestCaseState.CompletedSuccess;
+									}
+									
+									if(ValidationFailedStates.Contains(ValidatorState))
+									{
+										return TestCaseState.CompletedFailure;
+									}
+								}
+							}
+							else
+							{
+								return TestCaseState.CompletedFailure;
+							}
+							// retry
+							return TestCaseState.RetryRpc;
+						}));
+					TestCaseQueue.Enqueue(new SubmitToolRPCTestCase(this, TargetRegistry.GetTarget(SubmitToolRpcClient), Phases, SubmitToolRPCs.CloseApplication));
+					
+					DoneQueuingCases = true;
+					return true;
+				}
+				else if (!TargetRegistry.RpcTargets.IsEmpty)
+				{
+					SimpleResponse RegistrationFinishedResponse = SubmitToolRpcLibrary.Instance.RegistrationFinished(TargetRegistry.GetTarget(SubmitToolRpcClient));
+					SubmitToolRPCsReady = SubmitToolRpcLibrary.Instance.GetResponseBoolFieldValue(RegistrationFinishedResponse, "RegistrationFinished");
+				}
+			}
+			return false;
+		}
+
+		public static string GetLocalHostIPAddress()
+		{
+			string LocalIP = UnrealHelpers.GetHostIpAddress();
+			if (LocalIP == IPAddress.Loopback.ToString())
+			{
+				IPAddress LocalNonLoopback =
+					UnrealHelpers.GetAllLocalIPv4(NetworkInterfaceType.Ethernet)
+					.ToList()
+					.Where(IP => IP.AddressFamily == AddressFamily.InterNetwork && IP.GetAddressBytes()[0] != 169 && IP != IPAddress.Loopback)
+					.First();
+				if (LocalNonLoopback != null)
+				{
+					LocalIP = LocalNonLoopback.ToString();
+				}
+			}
+			return LocalIP;
+		}
+		public override void StopTest(StopReason InReason)
+		{
+			try
+			{
+				RpcExecutor.Instance.ClearIncomingMessageQueue();
+				RpcExecutor.Instance.ShutdownListenThread();
+
+				base.StopTest(InReason);
+			}
+			catch
+			{
+				throw;
+			}
+		}
+
+		protected override bool HasFailedTestCases()
+		{
+			return TestCaseQueue.HasFailedTests;
+		}
+
+		protected override bool IsRunningTestCases()
+		{
+			return TestCaseQueue.IsRunningTestCases;
+		}
+
+		protected override bool TestCasesCompleted()
+		{
+			return TestCaseQueue.HasCompletedQueuedTests;
+		}
+
+		protected override void UpdateTestCases()
+		{
+			TestCaseQueue.Update();
+		}
+
+		public override void TickTestCases()
+		{
+			TestCaseQueue.TickTestCases();
+		}
+	}
+
+	public class SubmitToolBuildSource : ApplicationTestsBuildSource
+	{
+		public new SubmitToolTestExecutorOptions TestExecutorOptions
+		{
+			get
+			{
+				return (SubmitToolTestExecutorOptions)base.TestExecutorOptions;
+			}
+		}
+
+		public SubmitToolBuildSource(SubmitToolTestExecutorOptions InTestExecutorOptions)
+			: base(InTestExecutorOptions)
+		{
+			ApplicationAppName = "SubmitTool";
+		}
+
+		public override void UpdateConfiguration()
+		{
+			AppConfig.CommandLineParams.AddRawCommandline($"-cl={TestExecutorOptions.TestCL}");
+			AppConfig.CommandLineParams.AddRawCommandline("-test-mode");
+			AppConfig.CommandLineParams.AddRawCommandline($"-abslog=SubmitToolRPCTests.log");
+
+			if (TestExecutorOptions.Sleep > 0)
+			{
+				AppConfig.CommandLineParams.AddRawCommandline(string.Format("--sleep={0}", TestExecutorOptions.Sleep));
+			}
+			if (TestExecutorOptions.AttachToDebugger)
+			{
+				AppConfig.CommandLineParams.AddRawCommandline("--attach-to-debugger");
+			}
+
+			ApplyRpcCommandLineParams(AppConfig);
+		}
+
+		private void ApplyRpcCommandLineParams(UnrealAppConfig AppConfig)
+		{
+			AppConfig.CommandLineParams.Add("rpcport", BasePort);
+			AppConfig.CommandLineParams.Add("externalrpclistenaddress", RpcExecutor.GetListenAddressAndPort());
+			AppConfig.CommandLineParams.Add("rpcsenderid", "SubmitToolRpcClient");
+
+			AppConfig.CommandLineParams.Add("ini:Engine:[HTTPServer.Listeners]:DefaultBindAddress", "any");
+			AppConfig.CommandLineParams.Add("ini:Engine:[HTTPServer.Listeners]:DefaultReuseAddressAndPort", true);
+
+			AppConfig.CommandLineParams.Add("logcmds", "LogSockets VeryVerbose,LogHttpServerModule VeryVerbose");
+
+			AppConfig.CommandLineParams.Add("DisableSSLCertificatePinning");
+		}
+	}
+
+	public class SubmitToolTestContext : ApplicationTestContext
+	{
+		public new SubmitToolTestExecutorOptions Options
+		{
+			get
+			{
+				return (SubmitToolTestExecutorOptions)base.Options;
+			}
+		}
+
+		public SubmitToolTestContext(ApplicationTestsBuildSource InBuildInfo, ApplicationTestRoleContext InRoleContext, SubmitToolTestExecutorOptions InOptions, int InPerTestTimeout = 0)
+			: base(InBuildInfo, InRoleContext, InOptions, InPerTestTimeout)
+		{
+		}
+	}
+
+	public class SubmitToolTestExecutorOptions : ApplicationTestExecutorOptions
+	{
+		public int TestCL = 0;
+		public SubmitToolTestExecutorOptions() : base()
+		{
+		}
+	}
+	internal class TestCaseState
+	{
+		public const string Init = "Init";
+		public const string CallRpc = "CallRpc";
+		public const string RetryRpc = "RetryRpc";
+		public const string InProgress = "InProgress";
+		public const string CompletedSuccess = "CompletedSuccess";
+		public const string CompletedFailure = "CompletedFailure";
+		public const string Finished = "Finished";
+		public const string CompletedTimeout = "CompletedTimeout";
+	}
+
+	public enum SubmitToolRPCs
+	{
+		SetChangelistDescription,
+		RunValidator,
+		GetValidatorState,
+		CloseApplication
+	}
+
+	public class SubmitToolRPCTestCase : BaseTestCase, IPhaseDataHandler
+	{
+		private ApplicationTest ParentTest;
+		private RpcTarget RPCTarget;
+		private SubmitToolRPCs RPCToRun;
+		private Dictionary<string, object> RPCArgs;
+		private Func<SimpleResponse, string> ResponseCheckerFunc;
+		private bool RPCSucceeded = false;
+
+		private string CurrentState = TestCaseState.Init;
+
+		private DateTime TestCaseStarted;
+		private int MaxTimetoutMinutes = 10;
+
+		private int NumRetry = 0;
+		private int MaxRetry = 10;
+
+		PhaseQueue Phases;
+
+		public SubmitToolRPCTestCase(ApplicationTest InParentTest, RpcTarget InRPCTarget, PhaseQueue InPhases, SubmitToolRPCs InRPCToRun, Dictionary<string, object> InRPCArgs = null, Func<SimpleResponse, string> InResponseCheckerFunc = null)
+		{
+			RPCTarget = InRPCTarget;
+			RPCToRun = InRPCToRun;
+			RPCArgs = InRPCArgs;
+			ResponseCheckerFunc = InResponseCheckerFunc;
+		
+			ParentTest = InParentTest;
+			Phases = InPhases;
+
+			Phases.Start(Enum.GetName(typeof(SubmitToolRPCs), RPCToRun), this);
+
+			TestCaseStarted = DateTime.UtcNow;
+		}
+
+		public void OnAddToReport(HordeReport.AutomatedTestSessionData.TestPhase ReportedPhase)
+		{
+			ReportedPhase.SetTiming(TestCaseStarted, (float)(DateTime.UtcNow - TestCaseStarted).TotalSeconds);
+
+			ReportedPhase.AddTags([TestTag.RPC]);
+
+			switch (GetTestResult())
+			{
+				case TestResult.TimedOut:
+				case TestResult.Failed:
+					{
+						ReportedPhase.SetOutcome(HordeReport.AutomatedTestSessionData.TestPhaseOutcome.Failed);
+						break;
+					}
+				case TestResult.Passed:
+					{
+						ReportedPhase.SetOutcome(HordeReport.AutomatedTestSessionData.TestPhaseOutcome.Success);
+						break;
+					}
+				default:
+					{
+						ReportedPhase.SetOutcome(HordeReport.AutomatedTestSessionData.TestPhaseOutcome.Unknown);
+						break;
+					}
+			}
+		}
+
+		public override void TickTestCase()
+		{
+			ProcessParentTestEvents();
+
+			CheckForTimeout();
+
+			switch (CurrentState)
+			{
+				case TestCaseState.Init:
+					CurrentState = TestCaseState.CallRpc;
+					LogState();
+					break;
+				case TestCaseState.RetryRpc:
+				case TestCaseState.CallRpc:
+					if (CurrentState == TestCaseState.RetryRpc)
+					{
+						Thread.Sleep(5000);
+						NumRetry++;
+						if (NumRetry == MaxRetry)
+						{
+							Log.Error("RPC retry max limit reached.");
+							CurrentState = TestCaseState.CompletedTimeout;
+							break;
+						}
+					}
+
+					CurrentState = TestCaseState.InProgress;
+					LogState();
+
+					Task.Run(() =>
+					{
+						SimpleResponse Response = null;
+
+						switch (RPCToRun)
+						{
+							case SubmitToolRPCs.SetChangelistDescription:
+								Response = SubmitToolRpcLibrary.Instance.SetChangelistDescription(RPCTarget, RPCArgs);
+								break;
+							case SubmitToolRPCs.RunValidator:
+								Response = SubmitToolRpcLibrary.Instance.RunValidator(RPCTarget, RPCArgs);
+								break;
+							case SubmitToolRPCs.GetValidatorState:
+								Response = SubmitToolRpcLibrary.Instance.GetValidatorState(RPCTarget, RPCArgs);
+								break;
+							case SubmitToolRPCs.CloseApplication:
+								Response = SubmitToolRpcLibrary.Instance.CloseApplication(RPCTarget);
+								break;
+						}
+
+						RPCSucceeded = Response.Succeeded;
+
+						if (ResponseCheckerFunc != null)
+						{
+							CurrentState = ResponseCheckerFunc.Invoke(Response);
+							LogState();
+						}
+						else if (Response != null)
+						{
+							if (Response.Succeeded)
+							{
+								CurrentState = TestCaseState.CompletedSuccess;
+								LogState();
+							}
+							else
+							{
+								CurrentState = TestCaseState.CompletedFailure;
+								LogState();
+							}
+						}
+					});
+					break;
+				case TestCaseState.CompletedTimeout:
+				case TestCaseState.CompletedSuccess:
+				case TestCaseState.CompletedFailure:
+					CurrentState = TestCaseState.Finished;
+					LogState();
+					Phases.Get(Enum.GetName(typeof(SubmitToolRPCs), RPCToRun))?.End();
+					break;
+			}
+		}
+
+		private void CheckForTimeout()
+		{
+			if (DateTime.UtcNow - TestCaseStarted > TimeSpan.FromMinutes(MaxTimetoutMinutes))
+			{
+				CurrentState = TestCaseState.CompletedTimeout;
+				LogState();
+			}
+		}
+
+		private void ProcessParentTestEvents()
+		{
+			var CurrentEvents = ParentTest.GetLoggedEvents();
+			foreach (UnrealTestEvent Event in CurrentEvents)
+			{
+				Phases.Get(Enum.GetName(typeof(SubmitToolRPCs), RPCToRun))?.AddEvent(ConvertToLogLevel(Event.Severity), Event.Summary);
+			}
+		}
+
+		private Microsoft.Extensions.Logging.LogLevel ConvertToLogLevel(EventSeverity InSeverity)
+		{
+			switch (InSeverity)
+			{
+				case EventSeverity.Info:
+					return Microsoft.Extensions.Logging.LogLevel.Information;
+				case EventSeverity.Warning:
+					return Microsoft.Extensions.Logging.LogLevel.Warning;
+				case EventSeverity.Error:
+					return Microsoft.Extensions.Logging.LogLevel.Error;
+				case EventSeverity.Fatal:
+					return Microsoft.Extensions.Logging.LogLevel.Critical;
+			}
+
+			return Microsoft.Extensions.Logging.LogLevel.None;
+		}
+
+		public override TestResult GetTestResult()
+		{
+			return CurrentState == TestCaseState.CompletedTimeout ? TestResult.TimedOut : (RPCSucceeded ? TestResult.Passed : TestResult.Failed);
+		}
+
+		public override string GetTestStatus()
+		{
+			return CurrentState;
+		}
+
+		public override bool HasFinished()
+		{
+			return CurrentState == TestCaseState.Finished;
+		}
+	}
+}
diff --git a/Engine/Source/Programs/AutomationTool/SubmitTool/Tests/SubmitTool.RpcLibrary.cs b/Engine/Source/Programs/AutomationTool/SubmitTool/Tests/SubmitTool.RpcLibrary.cs
new file mode 100644
--- /dev/null
+++ b/Engine/Source/Programs/AutomationTool/SubmitTool/Tests/SubmitTool.RpcLibrary.cs
@@ -0,0 +1,69 @@
+// Copyright Epic Games, Inc. All Rights Reserved.
+
+using Gauntlet;
+using System;
+using System.Collections.Generic;
+
+namespace SubmitTool.Automation
+{
+	/// <summary>
+	/// SubmitTool RPC library (run validators, verify submit tool fields etc)
+	/// </summary>
+	public class SubmitToolRpcLibrary : BaseRpcLibrary
+	{
+		const string SET_CHANGELIST_DESCRIPTION_RPC = "SetChangelistDescription";
+		const string RUN_VALIDATOR_RPC = "RunValidator";
+		const string GET_VALIDATOR_STATE_RPC = "GetValidatorState";
+		const string CLOSE_APPLICATION_RPC = "CloseApplication";
+
+		private static readonly Lazy<SubmitToolRpcLibrary> LibraryInstance = new Lazy<SubmitToolRpcLibrary>(() => new SubmitToolRpcLibrary());
+
+		public static SubmitToolRpcLibrary Instance { get { return LibraryInstance.Value; } }
+
+		private SubmitToolRpcLibrary()
+		{
+		}
+
+		/// <summary>
+		/// Run a submit tool validator.
+		/// </summary>
+		/// <param name="InTarget">Rpc Target to run against</param>
+		/// <param name="ValidatorArgs">Arguments containing the name id of the SubmitTool validator as specified in config e.g. "Preflight", "Tag"</param>
+		/// <param name="InTimeoutMs">Timeout for RPC expected finish in milliseconds</param>
+		public BaseRpcResponse RunValidator(RpcTarget InTarget, Dictionary<string, object> ValidatorArgs, int InTimeoutMs = 0)
+		{
+			return RunRpc(InTarget, RUN_VALIDATOR_RPC, ValidatorArgs, InTimeoutMs);
+		}
+
+		/// <summary>
+		/// Query validator state
+		/// </summary>
+		/// <param name="ValidatorArgs">Arguments containing the name id of the SubmitTool validator as specified in config e.g. "Preflight", "Tag"</param>
+		/// <param name="InTimeoutMs">Timeout for RPC expected finish in milliseconds</param>
+		public BaseRpcResponse GetValidatorState(RpcTarget InTarget, Dictionary<string, object> ValidatorArgs, int InTimeoutMs = 0)
+		{
+			return RunRpc(InTarget, GET_VALIDATOR_STATE_RPC, ValidatorArgs, InTimeoutMs);
+		}
+
+		/// <summary>
+		/// Request close for SubmitTool application.
+		/// </summary>
+		/// <param name="InTarget">Rpc Target to run against</param>
+		/// <param name="InTimeoutMs">Timeout for RPC expected finish in milliseconds</param>
+		public BaseRpcResponse CloseApplication(RpcTarget InTarget, int InTimeoutMs = 0)
+		{
+			return RunRpc(InTarget, CLOSE_APPLICATION_RPC, null, InTimeoutMs);
+		}
+
+		/// <summary>
+		/// Set the description of the current changelist.
+		/// </summary>
+		/// <param name="InTarget">Rpc Target to run against</param>
+		/// <param name="ChangelistDescriptionArgs">Changelist description to be set for the current CL</param>
+		/// <param name="InTimeoutMs">Timeout for RPC expected finish in milliseconds</param>
+		public BaseRpcResponse SetChangelistDescription(RpcTarget InTarget, Dictionary<string, object> ChangelistDescriptionArgs, int InTimeoutMs = 0)
+		{
+			return RunRpc(InTarget, SET_CHANGELIST_DESCRIPTION_RPC, ChangelistDescriptionArgs, InTimeoutMs);
+		}
+	}
+}
\ No newline at end of file
diff --git a/Engine/Source/Programs/SubmitTool/Private/SubmitToolApp.cpp b/Engine/Source/Programs/SubmitTool/Private/SubmitToolApp.cpp
--- a/Engine/Source/Programs/SubmitTool/Private/SubmitToolApp.cpp
+++ b/Engine/Source/Programs/SubmitTool/Private/SubmitToolApp.cpp
@@ -31,6 +31,12 @@
 
 #include "Configuration/Configuration.h"
 
+#if !UE_BUILD_SHIPPING
+#include "Tests/RPCLibrary.h"
+
+#include "Logging/LogScopedVerbosityOverride.h"
+#endif // !UE_BUILD_SHIPPING
+
 IMPLEMENT_APPLICATION(SubmitTool, "SubmitTool");
 
 #define LOCTEXT_NAMESPACE "SubmitTool"
@@ -183,6 +189,16 @@ int RunSubmitTool(const TCHAR* CommandLine)
 	// start up the main loop
 	GEngineLoop.PreInit(CommandLine);
 
+#if !UE_BUILD_SHIPPING
+	if (FCmdLineParameters::Get().Contains(FSubmitToolCmdLine::TestMode))
+	{
+		// silence most logging in test mode except fatal errors
+		LogSubmitToolP4Debug.SetVerbosity(ELogVerbosity::Fatal);
+		LogSubmitToolDebug.SetVerbosity(ELogVerbosity::Fatal);
+		LogSubmitTool.SetVerbosity(ELogVerbosity::Fatal);
+	}
+#endif // !UE_BUILD_SHIPPING
+
 	// ensure that the backlog is enabled
 	if(GLog)
 	{
@@ -234,6 +250,31 @@ int RunSubmitTool(const TCHAR* CommandLine)
 		TUniquePtr<FSlateUI> SlateUI = MakeSlateUI(ModelInterface.Get());
 		SlateUI->Build(UserPrefs.Get());
 
+		// Test mode, attach to debugger, only non-shipping
+#if !UE_BUILD_SHIPPING
+		if (FCmdLineParameters::Get().Contains(FSubmitToolCmdLine::WaitForDebugger))
+		{
+			int MaxWaitSeconds = 300;
+			while (!FPlatformMisc::IsDebuggerPresent() && MaxWaitSeconds > 0)
+			{
+				MaxWaitSeconds--;
+				FPlatformProcess::Sleep(1.f);
+			}
+			UE_DEBUG_BREAK();
+		}
+
+		if (FCmdLineParameters::Get().Contains(FSubmitToolCmdLine::TestMode))
+		{
+			USubmitToolRPCLibrary* RPCLibrary = NewObject<USubmitToolRPCLibrary>();
+			RPCLibrary->SetModelInterface(ModelInterface.Get());
+			RPCLibrary->RegisterLibraryRPCs();
+
+			LogSubmitToolP4Debug.SetVerbosity(ELogVerbosity::All);
+			LogSubmitToolDebug.SetVerbosity(ELogVerbosity::All);
+			LogSubmitTool.SetVerbosity(ELogVerbosity::All);
+		}
+#endif
+
 		// record that the application has started
 		// do this after initializing UI to give ISTSourceControlService time to populate GetCurrentStreamName
 		FTelemetryService::Init(ModelInterface->GetParameters().Telemetry);
diff --git a/Engine/Source/Programs/SubmitTool/Private/Tests/RPCLibrary.cpp b/Engine/Source/Programs/SubmitTool/Private/Tests/RPCLibrary.cpp
new file mode 100644
--- /dev/null
+++ b/Engine/Source/Programs/SubmitTool/Private/Tests/RPCLibrary.cpp
@@ -0,0 +1,182 @@
+// Copyright Epic Games, Inc. All Rights Reserved.
+
+#if !UE_BUILD_SHIPPING
+
+#include "Tests/RPCLibrary.h"
+
+DEFINE_LOG_CATEGORY_STATIC(LogSubmitToolRPCLibrary, Log, All);
+
+USubmitToolRPCLibrary::USubmitToolRPCLibrary()
+{
+}
+
+void USubmitToolRPCLibrary::RegisterLibraryRPCs()
+{
+#if WITH_RPC_REGISTRY
+	StartListening();
+	Initialize();
+	RegisterRPC(
+		TEXT("RunValidator"),
+		FHttpPath(TEXT("/submittool/runvalidator")),
+		EHttpServerRequestVerbs::VERB_POST,
+		FHttpRequestHandler::CreateUObject(this, &USubmitToolRPCLibrary::HandleRunValidatorRoute));
+	RegisterRPC(
+		TEXT("GetValidatorState"),
+		FHttpPath(TEXT("/submittool/getvalidatorstate")),
+		EHttpServerRequestVerbs::VERB_POST, // GET doesn't work because of .NET not allowing payload for this verb
+		FHttpRequestHandler::CreateUObject(this, &USubmitToolRPCLibrary::HandleGetValidatorStateRoute));
+	RegisterRPC(
+		TEXT("SetChangelistDescription"),
+		FHttpPath(TEXT("/submittool/setchangelistdescription")),
+		EHttpServerRequestVerbs::VERB_POST,
+		FHttpRequestHandler::CreateUObject(this, &USubmitToolRPCLibrary::HandleSetChangelistDescriptionRoute));
+	RegisterRPC(
+		TEXT("CloseApplication"),
+		FHttpPath(TEXT("/submittool/closeapplication")),
+		EHttpServerRequestVerbs::VERB_POST,
+		FHttpRequestHandler::CreateUObject(this, &USubmitToolRPCLibrary::HandleCloseApplicationRoute));
+	RegistrationFinished = true;
+#endif //WITH_RPC_REGISTRY
+}
+
+#if WITH_RPC_REGISTRY
+bool USubmitToolRPCLibrary::HandleRunValidatorRoute(const FHttpServerRequest& Request, const FHttpResultCallback& OnComplete)
+{
+	check(FCmdLineParameters::Get().Contains(FSubmitToolCmdLine::TestMode));
+
+	FString ResponseStr;
+	TSharedRef<TJsonWriter<>> JsonWriter = TJsonWriterFactory<>::Create(&ResponseStr);
+	JsonWriter->WriteObjectStart();
+
+	TSharedPtr<FJsonObject> RootObject;
+	FUTF8ToTCHAR RequestBodyByteBuffer(reinterpret_cast<const ANSICHAR*>(Request.Body.GetData()), Request.Body.Num());
+	const FString RequestBody = FString::ConstructFromPtrSize(RequestBodyByteBuffer.Get(), RequestBodyByteBuffer.Length());
+	TSharedRef<TJsonReader<>> JsonReader = TJsonReaderFactory<>::Create(RequestBody);
+
+	bool status = DeserializeRequest(RootObject, JsonReader, JsonWriter);
+
+	if (status)
+	{
+		FString ValidatorName;
+		status = GetStringField(RootObject, TEXT("ValidatorName"), ValidatorName, JsonWriter);
+		if (status)
+		{
+			UE_LOG(LogSubmitToolRPCLibrary, Log, TEXT("Running validator: %s"), *ValidatorName);
+
+			OnSingleValidatorFinishedHandle = ModelInterface->AddSingleValidatorFinishedCallback(FOnSingleTaskFinished::FDelegate::CreateLambda([this](const FValidatorBase& InValidator) { OnSingleValidatorFinished(InValidator); }));
+			ModelInterface->ValidateSingle(FName(ValidatorName), true);
+			ValidatorStates.Add(ValidatorName, EValidationStates::Not_Applicable);
+
+			JsonWriter->WriteValue(TEXT("IsStarted"), true);
+			JsonWriter->WriteObjectEnd();
+			JsonWriter->Close();
+			status = true;
+		}
+	}
+
+	auto Response = CreateSimpleResponse(true, ResponseStr, false);
+	OnComplete(MoveTemp(Response));
+	return status;
+}
+
+void USubmitToolRPCLibrary::OnSingleValidatorFinished(const FValidatorBase& InValidator)
+{
+	UE_LOG(LogSubmitToolRPCLibrary, Display, TEXT("Validator %s finished."), *InValidator.GetValidatorNameId().ToString());
+	ValidatorStates.Emplace(InValidator.GetValidatorNameId().ToString(), InValidator.GetState());
+}
+
+bool USubmitToolRPCLibrary::HandleGetValidatorStateRoute(const FHttpServerRequest& Request, const FHttpResultCallback& OnComplete)
+{
+	FString ResponseStr;
+	TSharedRef<TJsonWriter<>> JsonWriter = TJsonWriterFactory<>::Create(&ResponseStr);
+	JsonWriter->WriteObjectStart();
+
+	TSharedPtr<FJsonObject> RootObject;
+	FUTF8ToTCHAR RequestBodyByteBuffer(reinterpret_cast<const ANSICHAR*>(Request.Body.GetData()), Request.Body.Num());
+	const FString RequestBody = FString::ConstructFromPtrSize(RequestBodyByteBuffer.Get(), RequestBodyByteBuffer.Length());
+	TSharedRef<TJsonReader<>> JsonReader = TJsonReaderFactory<>::Create(RequestBody);
+
+	bool status = DeserializeRequest(RootObject, JsonReader, JsonWriter);
+
+	if (status)
+	{
+		FString ValidatorName;
+		status = GetStringField(RootObject, TEXT("ValidatorName"), ValidatorName, JsonWriter);
+		if (status)
+		{
+			UE_LOG(LogSubmitToolRPCLibrary, Log, TEXT("Running validator: %s"), *ValidatorName);
+			TArray<TWeakPtr<const FValidatorBase>> ValidatorsArray = ModelInterface->GetValidators();
+
+			EValidationStates ValidatorState = EValidationStates::Not_Applicable;
+
+			for (const TWeakPtr<const FValidatorBase>& Validator : ValidatorsArray)
+			{
+				if (Validator.Pin()->GetValidatorNameId().ToString() == ValidatorName)
+				{
+					ValidatorState = Validator.Pin()->GetState();
+					break;
+				}
+			}
+
+			if (ValidatorState == EValidationStates::Not_Applicable && ValidatorStates.Contains(ValidatorName))
+			{
+				ValidatorState = *ValidatorStates.Find(ValidatorName);
+			}
+
+			JsonWriter->WriteValue(TEXT("ValidatorState"), *(StaticEnum<EValidationStates>()->GetNameStringByValue((int64)ValidatorState)));
+			JsonWriter->WriteObjectEnd();
+			JsonWriter->Close();
+			status = true;
+		}
+	}
+
+	auto Response = CreateSimpleResponse(true, ResponseStr, false);
+	OnComplete(MoveTemp(Response));
+
+	return status;
+}
+
+bool USubmitToolRPCLibrary::HandleSetChangelistDescriptionRoute(const FHttpServerRequest& Request, const FHttpResultCallback& OnComplete)
+{
+	FString ResponseStr;
+	TSharedRef<TJsonWriter<>> JsonWriter = TJsonWriterFactory<>::Create(&ResponseStr);
+	JsonWriter->WriteObjectStart();
+
+	TSharedPtr<FJsonObject> RootObject;
+	FUTF8ToTCHAR RequestBodyByteBuffer(reinterpret_cast<const ANSICHAR*>(Request.Body.GetData()), Request.Body.Num());
+	const FString RequestBody = FString::ConstructFromPtrSize(RequestBodyByteBuffer.Get(), RequestBodyByteBuffer.Length());
+	TSharedRef<TJsonReader<>> JsonReader = TJsonReaderFactory<>::Create(RequestBody);
+
+	bool status = DeserializeRequest(RootObject, JsonReader, JsonWriter);
+
+	if (status)
+	{
+		FString Description;
+		status = GetStringField(RootObject, TEXT("Description"), Description, JsonWriter);
+		if (status)
+		{
+			UE_LOG(LogSubmitToolRPCLibrary, Log, TEXT("Setting CL description to: %s"), *Description);
+			ModelInterface->SetCLDescription(FText::FromString(Description));
+		}
+	}
+
+	auto Response = CreateSimpleResponse(true, ResponseStr, false);
+	OnComplete(MoveTemp(Response));
+
+	return true;
+}
+
+bool USubmitToolRPCLibrary::HandleCloseApplicationRoute(const FHttpServerRequest& Request, const FHttpResultCallback& OnComplete)
+{
+	FString ResponseStr;
+
+	auto Response = CreateSimpleResponse(true, ResponseStr, false);
+	OnComplete(MoveTemp(Response));
+
+	FPlatformMisc::RequestExitWithStatus(true, 0);
+
+	return true;
+}
+
+#endif // WITH_RPC_REGISTRY
+#endif // !UE_BUILD_SHIPPING
\ No newline at end of file
diff --git a/Engine/Source/Programs/SubmitTool/Public/Tests/RPCLibrary.h b/Engine/Source/Programs/SubmitTool/Public/Tests/RPCLibrary.h
new file mode 100644
--- /dev/null
+++ b/Engine/Source/Programs/SubmitTool/Public/Tests/RPCLibrary.h
@@ -0,0 +1,47 @@
+// Copyright Epic Games, Inc. All Rights Reserved.
+
+#pragma once
+
+#include "RPCLibraryBase.h"
+
+#include "Serialization/JsonReader.h"
+#include "Serialization/JsonSerializer.h"
+#include "Serialization/JsonWriter.h"
+
+#include "CommandLine/CmdLineParameters.h"
+#include "Logic/Services/SourceControl/SubmitToolPerforce.h"
+#include "Models/ModelInterface.h"
+#include "Parameters/SubmitToolParametersBuilder.h"
+
+#include "RPCLibrary.generated.h"
+
+UCLASS(MinimalAPI)
+class USubmitToolRPCLibrary : public URPCLibraryBase
+{
+	GENERATED_BODY()
+public:
+	USubmitToolRPCLibrary();
+	virtual void BeginDestroy() override
+	{
+		if (ModelInterface && OnSingleValidatorFinishedHandle.IsValid())
+		{
+			ModelInterface->RemoveValidationFinishedCallback(OnSingleValidatorFinishedHandle);
+		}
+		Super::BeginDestroy();
+	}
+	void SetModelInterface(FModelInterface* InModelInterface) { ModelInterface = InModelInterface; }
+	void RegisterLibraryRPCs();
+private:
+	FModelInterface* ModelInterface;
+
+	TMap<FString, EValidationStates> ValidatorStates;
+
+	FDelegateHandle OnSingleValidatorFinishedHandle;
+	void OnSingleValidatorFinished(const FValidatorBase& InValidator);
+#if WITH_RPC_REGISTRY
+	bool HandleRunValidatorRoute(const FHttpServerRequest& Request, const FHttpResultCallback& OnComplete);
+	bool HandleGetValidatorStateRoute(const FHttpServerRequest& Request, const FHttpResultCallback& OnComplete);
+	bool HandleSetChangelistDescriptionRoute(const FHttpServerRequest& Request, const FHttpResultCallback& OnComplete);
+	bool HandleCloseApplicationRoute(const FHttpServerRequest& Request, const FHttpResultCallback& OnComplete);
+#endif // WITH_RPC_REGISTRY
+};
\ No newline at end of file
diff --git a/Engine/Source/Programs/SubmitTool/SubmitTool.Build.cs b/Engine/Source/Programs/SubmitTool/SubmitTool.Build.cs
--- a/Engine/Source/Programs/SubmitTool/SubmitTool.Build.cs
+++ b/Engine/Source/Programs/SubmitTool/SubmitTool.Build.cs
@@ -58,6 +58,9 @@ public class SubmitTool : ModuleRules
 				"SlateReflector",
 				}
 			);
+
+			// Test dependencies
+			PublicDependencyModuleNames.Add("RPCBase");
 		}
 
 		AddEngineThirdPartyPrivateStaticDependencies(Target, "Perforce");
