Bringing PSODeepTest to Android.

Handled low-memory devices by checking the memory pressure status, increasing frames between spawns, and decreasing the number of objects before a GC.

#rnx

#Synced-CL 50983689

diff --git a/Engine/Source/Programs/AutomationTool/Gauntlet/Platform/Android/Gauntlet.TargetDeviceAndroid.cs b/Engine/Source/Programs/AutomationTool/Gauntlet/Platform/Android/Gauntlet.TargetDeviceAndroid.cs
--- a/Engine/Source/Programs/AutomationTool/Gauntlet/Platform/Android/Gauntlet.TargetDeviceAndroid.cs
+++ b/Engine/Source/Programs/AutomationTool/Gauntlet/Platform/Android/Gauntlet.TargetDeviceAndroid.cs
@@ -553,6 +553,45 @@ namespace Gauntlet
 				}
 			}
 
+			// disable samsung game optimisation service, when -enableSamsungGOS not present.
+			// log the attempt but ignore failure if not found etc.
+			// we disable this to prevent GOS's 'touch protection' from down clocking the device.
+			if (Globals.Params.ParseParam("enableSamsungGOS"))
+			{
+				IProcessResult AdbResult = RunAdbDeviceCommand("shell pm enable --user 0 com.samsung.android.game.gos", true, true);
+				if (AdbResult.ExitCode != 0)
+				{
+					Log.Info("Failed to enable Samsung GOS");
+				}
+			}
+			else
+			{
+				IProcessResult AdbResult = RunAdbDeviceCommand("shell pm disable-user --user 0 com.samsung.android.game.gos", true, true);
+				if (AdbResult.ExitCode != 0)
+				{
+					Log.Info("Failed to disable Samsung GOS");
+				}
+			}
+
+			// disable other popups that take focus away from the running app.
+			// GMS will display 'google play is required'
+			if (Globals.Params.ParseParam("disableAndroidGMS"))
+			{
+				IProcessResult AdbResult = RunAdbDeviceCommand("shell pm disable-user --user 0 com.google.android.gms", true, true);
+				if (AdbResult.ExitCode != 0)
+				{
+					Log.Info("Failed to disable Google Mobile Services");
+				}
+			}
+			else
+			{
+				IProcessResult AdbResult = RunAdbDeviceCommand("shell pm enable --user 0 com.google.android.gms", true, true);
+				if (AdbResult.ExitCode != 0)
+				{
+					Log.Info("Failed to enable Google Mobile Services");
+				}
+			}
+
 			// start the app on device!
 			string CommandLine = "shell am start -W -S -n " + Install.AndroidPackageName + "/" + LaunchActivity;
 			IProcessResult Process = RunAdbDeviceCommand(CommandLine, false, true);
