Removing appium from any run capturing perf data

#Synced-CL 50214590

diff --git a/Engine/Source/Programs/AutomationTool/Gauntlet/Platform/IOS/Gauntlet.IOSAppium.cs b/Engine/Source/Programs/AutomationTool/Gauntlet/Platform/IOS/Gauntlet.IOSAppium.cs
--- a/Engine/Source/Programs/AutomationTool/Gauntlet/Platform/IOS/Gauntlet.IOSAppium.cs
+++ b/Engine/Source/Programs/AutomationTool/Gauntlet/Platform/IOS/Gauntlet.IOSAppium.cs
@@ -303,7 +303,14 @@ namespace Gauntlet
 					throw new DeviceException($"Exception during WebdriverAgent Creation {Ex.Message}");
 				}
 			}
+		}
 
+		/// <summary>
+		/// Starts the monitor thread which checks for alerts and takes screenshots
+		/// Should be called after the app has launched on the device
+		/// </summary>
+		public void StartMonitorThread()
+		{
 			if (!Globals.Params.ParseParam("NoAppiumScreenshots") && AppiumConfig.ScreenshotPeriod > 0)
 			{
 				MonitorThread = new Thread(StartMonitorLoop);
@@ -674,12 +681,12 @@ namespace Gauntlet
 
 		private void StartMonitorLoop()
 		{
-			// Alert check runs frequently (every 2 seconds)
-			TimeSpan AlertCheckPeriod = TimeSpan.FromSeconds(2);
+			// Alert check runs frequently (every 10 seconds)
+			TimeSpan AlertCheckPeriod = TimeSpan.FromSeconds(10);
 
 			// Calculate how many alert checks happen per screenshot
-			// e.g., if ScreenshotPeriod = 30 seconds, and AlertCheckPeriod = 2 seconds
-			// then screenshotCheckCount = 30 / 2 = 15 iterations
+			// e.g., if ScreenshotPeriod = 30 seconds, and AlertCheckPeriod = 10 seconds
+			// then screenshotCheckCount = 30 / 10 = 3 iterations
 			int screenshotCheckCount = (int)Math.Max(1, AppiumConfig.ScreenshotPeriod / AlertCheckPeriod.TotalSeconds);
 			int iterationCounter = 0;
 
diff --git a/Engine/Source/Programs/AutomationTool/Gauntlet/Platform/IOS/Gauntlet.TargetDeviceIOS.cs b/Engine/Source/Programs/AutomationTool/Gauntlet/Platform/IOS/Gauntlet.TargetDeviceIOS.cs
--- a/Engine/Source/Programs/AutomationTool/Gauntlet/Platform/IOS/Gauntlet.TargetDeviceIOS.cs
+++ b/Engine/Source/Programs/AutomationTool/Gauntlet/Platform/IOS/Gauntlet.TargetDeviceIOS.cs
@@ -903,17 +903,13 @@ namespace Gauntlet
 			{
 				throw new DeviceException("Failed to copy uecommandline file to device!");
 			}
-			Log.Info("Appium Enabled: {0} on {1}", Install.Name, this);
-			if (AppiumContainer.Enabled)
+			
+			if (AppiumContainer != null)
 			{
-				if (AppiumContainer == null)
-				{
-					AppiumContainer = new AppiumContainer(UUID);
-				}
+				Log.Info("Appium Enabled: {0} on {1}", Install.Name, this);
 				AppiumContainer.Start(IOSInstall.PackageName);
 			}
 
-
 			Log.Info("Launching {0} on {1}", Install.Name, this);
 			Log.Verbose(IOSInstall.CommandLine);
 
@@ -948,6 +944,11 @@ namespace Gauntlet
 				throw new DeviceException("Failed to launch on {0}. {1}", Name, AppProcess.Output);
 			}
 
+			if (AppiumContainer != null)
+			{
+				AppiumContainer.StartMonitorThread();
+			}
+
 			return new IOSAppInstance(IOSInstall, AppProcess, IOSInstall.CommandLine);
 		}
 
@@ -1517,7 +1518,6 @@ namespace Gauntlet
 		}
 	}
 
-
 	/// <summary>
 	/// Helper class to parses LLDB crash threads and generate Unreal compatible log callstack
 	/// </summary>
