When zero tests are run (opt-in), convert test events from error to warnings

#rnx
#Synced-CL 0

diff --git a/Engine/Source/Programs/AutomationTool/Gauntlet/Unreal/Automation/UE.Automation.cs b/Engine/Source/Programs/AutomationTool/Gauntlet/Unreal/Automation/UE.Automation.cs
--- a/Engine/Source/Programs/AutomationTool/Gauntlet/Unreal/Automation/UE.Automation.cs
+++ b/Engine/Source/Programs/AutomationTool/Gauntlet/Unreal/Automation/UE.Automation.cs
@@ -1233,7 +1233,14 @@ namespace UE
 							Log.Info("   See log above for details.");
 							foreach (var Error in Errors.Take(MaxEventsDisplayPerTest))
 							{
-								Log.Error(KnownLogEvents.Gauntlet_UnrealEngineTestEvent, "    " + Error.Summary);
+								if (bAllowZeroTestResults)
+								{
+									Log.Warning(KnownLogEvents.Gauntlet_UnrealEngineTestEvent, "    " + Error.Summary);
+								}
+								else
+								{
+									Log.Error(KnownLogEvents.Gauntlet_UnrealEngineTestEvent, "    " + Error.Summary);
+								}
 							}
 							NotifyMoreIfNeeded(Errors);
 							foreach (var Warning in Warnings.Take(MaxEventsDisplayPerTest))
