Updating file copy from device to use devicectl instead of ios-deploy

#Synced-CL 50214590

diff --git a/Engine/Source/Programs/AutomationTool/Gauntlet/Platform/IOS/Gauntlet.TargetDeviceIOS.cs b/Engine/Source/Programs/AutomationTool/Gauntlet/Platform/IOS/Gauntlet.TargetDeviceIOS.cs
--- a/Engine/Source/Programs/AutomationTool/Gauntlet/Platform/IOS/Gauntlet.TargetDeviceIOS.cs
+++ b/Engine/Source/Programs/AutomationTool/Gauntlet/Platform/IOS/Gauntlet.TargetDeviceIOS.cs
@@ -1485,9 +1485,23 @@ namespace Gauntlet
 			}
 			Directory.CreateDirectory(TempPath);
 
-			string CopyCommand = string.Format("--bundle_id {0} --download={1} --to {2}",
-				Install.PackageName, SourceArtifactPath, TempPath);
-			DownloadResult = IOSDevice.ExecuteIOSDeployCommand(CopyCommand, 180);
+			if (TargetDeviceIOS.UseDeviceCtl && !IOSDevice.UseIOSDeployForCopy)
+			{
+				// Use devicectl to copy artifacts from device
+				// devicectl requires paths without leading slashes
+				string DevicectlSource = SourceArtifactPath.TrimStart('/');
+
+				string CopyCommand = string.Format("device copy from --device {0} --source {1} --destination {2} --domain-type appDataContainer --domain-identifier {3}",
+					IOSDevice.UUID, DevicectlSource, TempPath, Install.PackageName);
+				DownloadResult = IOSDevice.ExecuteDevicectlCommand(CopyCommand, UseDeviceID: false, WaitTime: 180);
+			}
+			else
+			{
+				// Fallback to ios-deploy for XCode < 16 or when UseIOSDeployForCopy is true
+				string CopyCommand = string.Format("--bundle_id {0} --download={1} --to {2}",
+					Install.PackageName, SourceArtifactPath, TempPath);
+				DownloadResult = IOSDevice.ExecuteIOSDeployCommand(CopyCommand, 180);
+			}
 			if (DownloadResult.ExitCode != 0)
 			{
 				string Error = string.Format("Copying artifacts from device failed with code {0}:\n{1}\nArtifacts will not be saved.",
@@ -1497,24 +1511,35 @@ namespace Gauntlet
 
 			try
 			{
+				// ios-deploy wraps files in a Documents subdirectory, devicectl does not
 				string DocumentsPath = Path.Combine(TempPath, "Documents");
 				if (Directory.Exists(DocumentsPath))
 				{
+					// ios-deploy case: move from Documents subdirectory
+					Log.Verbose("Moving artifacts from {0} to {1}", DocumentsPath, DestinationArtifactPath);
 					Directory.Move(DocumentsPath, DestinationArtifactPath);
 				}
+				else if (Directory.Exists(TempPath) && Directory.GetFileSystemEntries(TempPath).Length > 0)
+				{
+					// devicectl case: files are directly in TempPath, move TempPath itself
+					Log.Verbose("Moving artifacts from {0} to {1}", TempPath, DestinationArtifactPath);
+					Directory.Move(TempPath, DestinationArtifactPath);
+				}
+				else
+				{
+					Log.Warning("No artifacts found in {0} to move", TempPath);
+				}
 
-				Directory.Delete(TempPath, true);
+				// Clean up temp directory if it still exists
+				if (Directory.Exists(TempPath))
+				{
+					Directory.Delete(TempPath, true);
+				}
 			}
 			catch (Exception Ex)
 			{
 				Log.Error(KnownLogEvents.Gauntlet_DeviceEvent, "Failed to move artifacts out of temp directory:\n{Exception}", Ex.Message);
 			}
-
-			AppiumContainer.Config AppiumConfig = AppiumContainer.AppiumConfig;
-			if (!Globals.Params.ParseParam("NoAppiumScreenshots") && AppiumConfig != null && AppiumConfig.ScreenshotPeriod > 0)
-			{
-				IOSDevice.AppiumContainer.StopMonitorLoop(Path.Combine(DestinationArtifactPath, "Screenshots"));
-			}
 		}
 	}
 
