Add option to force usage of process log file over stdout redirection in Gauntlet

#rnx

#Synced-CL 51005601

diff --git a/Engine/Source/Programs/AutomationTool/Gauntlet/Framework/Base/Gauntlet.TargetDeviceDesktopCommon.cs b/Engine/Source/Programs/AutomationTool/Gauntlet/Framework/Base/Gauntlet.TargetDeviceDesktopCommon.cs
--- a/Engine/Source/Programs/AutomationTool/Gauntlet/Framework/Base/Gauntlet.TargetDeviceDesktopCommon.cs
+++ b/Engine/Source/Programs/AutomationTool/Gauntlet/Framework/Base/Gauntlet.TargetDeviceDesktopCommon.cs
@@ -434,9 +434,10 @@ namespace Gauntlet
 					LogFile = !string.IsNullOrEmpty(M.Groups[1].Value) ? M.Groups[1].Value : M.Groups[2].Value;
 				}
 
-				// Explicitly set log file when not already defined if not build machine
+				// Explicitly set log file either when stdout processing is disabled, or when not already defined if not build machine
 				// -abslog makes sure Unreal dynamically update the log window when using -log
-				if (!CommandUtils.IsBuildMachine && (!string.IsNullOrEmpty(LogFile) || AppConfig.CommandLine.Contains("-log")))
+				if (AppConfig.DisableStdoutProcessing ||
+					!CommandUtils.IsBuildMachine && (!string.IsNullOrEmpty(LogFile) || AppConfig.CommandLine.Contains("-log")))
 				{
 					string LogFolder = string.IsNullOrEmpty(LogFile) ?
 						(string.IsNullOrEmpty(UserDir) ? Path.Combine(Device.LocalCachePath, "Logs") : Path.Combine(UserDir, "Saved", "Logs")) :
diff --git a/Engine/Source/Programs/AutomationTool/Gauntlet/Unreal/Base/Gauntlet.UnrealAppConfig.cs b/Engine/Source/Programs/AutomationTool/Gauntlet/Unreal/Base/Gauntlet.UnrealAppConfig.cs
--- a/Engine/Source/Programs/AutomationTool/Gauntlet/Unreal/Base/Gauntlet.UnrealAppConfig.cs
+++ b/Engine/Source/Programs/AutomationTool/Gauntlet/Unreal/Base/Gauntlet.UnrealAppConfig.cs
@@ -124,6 +124,11 @@ namespace Gauntlet
 		/// </summary>
 		public int MaxDuration { get; set; }
 
+		/// <summary>
+		/// Disable stdout redirection and write the output of the app process to a log file
+		/// </summary>
+		public bool DisableStdoutProcessing { get; set; }
+
 		/// <summary>
 		/// Constructor that sets some required values to defaults
 		/// </summary>
