Add logging for folder deletion

#rnx
#Synced-CL 50931014

diff --git a/Engine/Source/Programs/AutomationTool/Gauntlet/Framework/Gauntlet.Utils.cs b/Engine/Source/Programs/AutomationTool/Gauntlet/Framework/Gauntlet.Utils.cs
--- a/Engine/Source/Programs/AutomationTool/Gauntlet/Framework/Gauntlet.Utils.cs
+++ b/Engine/Source/Programs/AutomationTool/Gauntlet/Framework/Gauntlet.Utils.cs
@@ -1721,6 +1721,7 @@ namespace Gauntlet
 			{
 				try
 				{
+					Log.Info("Deleting directory {Directory}.", Directory);
 					Directory.Delete(bRecursive);
 				}
 				catch (Exception Ex)
@@ -1737,6 +1738,7 @@ namespace Gauntlet
 
 						foreach(FileSystemInfo Info in Directory.EnumerateFiles("*", SearchOption.AllDirectories))
 						{
+							Log.Info("Deleting file after attributes change: {FileName}", Info.FullName);
 							Info.Attributes = FileAttributes.Normal;
 							// Try to delete the file, with retry logic for locked files
 							bool bDeleted = false;
