Fix issue with loading OBBs

#android

#Synced-CL 51027549

diff --git a/Engine/Source/Programs/AutomationTool/Android/AndroidPlatform.Automation.cs b/Engine/Source/Programs/AutomationTool/Android/AndroidPlatform.Automation.cs
--- a/Engine/Source/Programs/AutomationTool/Android/AndroidPlatform.Automation.cs
+++ b/Engine/Source/Programs/AutomationTool/Android/AndroidPlatform.Automation.cs
@@ -1092,7 +1092,7 @@ public class AndroidPlatform : Platform
 
 			foreach (FileReference FileRef in FilesForObb)
 			{
-				string RelativePath = FileRef.MakeRelativeTo(SC.StageDirectory);
+				string RelativePath = FileRef.MakeRelativeTo(SC.StageDirectory).Replace("\\", "/");
 				ObbFile.CreateEntryFromFile_CrossPlatform(FileRef.FullName, RelativePath, CompressionLevel.NoCompression);
 				ObbFileCount += 1;
 				Logger.LogInformation("[{Count}/{Total}] Adding {File} to OBB", ObbFileCount, FilesForObb.Count, RelativePath);
