Disabling appium as the default option.
Adding retry logic around appium driver creation
#Synced-CL 51038270

diff --git a/Engine/Source/Programs/AutomationTool/Gauntlet/Platform/IOS/Gauntlet.IOSAppium.cs b/Engine/Source/Programs/AutomationTool/Gauntlet/Platform/IOS/Gauntlet.IOSAppium.cs
--- a/Engine/Source/Programs/AutomationTool/Gauntlet/Platform/IOS/Gauntlet.IOSAppium.cs
+++ b/Engine/Source/Programs/AutomationTool/Gauntlet/Platform/IOS/Gauntlet.IOSAppium.cs
@@ -171,7 +171,7 @@ namespace Gauntlet
 		private string PackageName = null;
 		static AppiumContainer()
 		{
-			if (Globals.Params.ParseParam("NoAppium"))
+			if (!Globals.Params.ParseParam("UseAppium"))
 			{
 				return;
 			}
@@ -234,73 +234,113 @@ namespace Gauntlet
 
 			this.PackageName = PackageName;
 
-			// Check if WebDriverAgent is already installed
-			bool isWdaInstalled = CheckIfWdaInstalled(UUID, AppiumConfig.WdaBundleId);
-			if (isWdaInstalled)
-			{
-				Log.Info("WebDriverAgent ({0}) already installed on device {1}, skipping install", AppiumConfig.WdaBundleId, UUID);
-			}
-			else
+			// Install WebdriverAgent
+			try
 			{
-				// Install WebdriverAgent
-				try
+				Log.Info("Installing WebDriverAgent ({0}) on device {1}...", AppiumConfig.WdaBundleId, UUID);
+				IProcessResult InstallResult;
+				using (ScopedSuspendECErrorParsing ErrorSuspension = new())
 				{
-					Log.Info("Installing WebDriverAgent ({0}) on device {1}...", AppiumConfig.WdaBundleId, UUID);
-					IProcessResult InstallResult;
-					using (ScopedSuspendECErrorParsing ErrorSuspension = new())
+					if (TargetDeviceIOS.UseDeviceCtl)
 					{
-						if (TargetDeviceIOS.UseDeviceCtl)
-						{
-							InstallResult = TargetDeviceIOS.ExecuteDevicectlCommand(
-								string.Format("device install app \"{0}\" -v", AppiumConfig.WdaAppPath), UUID, 60, AdditionalOptions: ERunOptions.NoStdOutRedirect);
-						}
-						else
-						{
-							InstallResult = TargetDeviceIOS.ExecuteIOSDeployCommand(
-								string.Format("-b \"{0}\"", AppiumConfig.WdaAppPath), UUID, 60, AdditionalOptions: ERunOptions.NoStdOutRedirect);
-						}
+						InstallResult = TargetDeviceIOS.ExecuteDevicectlCommand(
+							string.Format("device install app \"{0}\" -v", AppiumConfig.WdaAppPath), UUID, 60, AdditionalOptions: ERunOptions.NoStdOutRedirect);
 					}
-					if (InstallResult.ExitCode != 0)
+					else
 					{
-						throw new AutomationException("Failed to install WDA app ({0}): {1}", InstallResult.ExitCode, InstallResult.Output);
+						InstallResult = TargetDeviceIOS.ExecuteIOSDeployCommand(
+							string.Format("-b \"{0}\"", AppiumConfig.WdaAppPath), UUID, 60, AdditionalOptions: ERunOptions.NoStdOutRedirect);
 					}
-					Log.Info("WebDriverAgent installed successfully on device {0}", UUID);
 				}
-				catch (Exception Ex)
+				if (InstallResult.ExitCode != 0)
 				{
-					throw new AutomationException($"Exception during WDA Install {Ex.Message}");
+					throw new AutomationException("Failed to install WDA app ({0}): {1}", InstallResult.ExitCode, InstallResult.Output);
 				}
+				Log.Info("WebDriverAgent installed successfully on device {0}", UUID);
+			}
+			catch (Exception Ex)
+			{
+				throw new AutomationException($"Exception during WDA Install {Ex.Message}");
 			}
 
 			lock (Mutex)
 			{
-				// Start Appium Server
-				try
+				// Retry logic for both server start and WebDriver creation
+				// Sometimes servers fail to start or devices appear temporarily unavailable
+				const int MaxRetries = 3;
+				int attemptCount = 0;
+				Exception lastException = null;
+
+				while (attemptCount < MaxRetries)
 				{
-					AppiumServiceBuilder ServerBuilder = new AppiumServiceBuilder()
-						.UsingAnyFreePort()
-						.WithIPAddress("127.0.0.1");
+					attemptCount++;
 
-					if (!string.IsNullOrEmpty(AppiumConfig.AppiumLocation) && FileExists(AppiumConfig.AppiumLocation))
+					try
 					{
-						ServerBuilder.WithAppiumJS(new System.IO.FileInfo(AppiumConfig.AppiumLocation));
+						// Clean up any previous server instance
+						if (Server != null)
+						{
+							Log.Verbose("Disposing previous server instance before retry");
+							try
+							{
+								Server.Dispose();
+							}
+							catch (Exception ex)
+							{
+								Log.Verbose("Error disposing previous server: {0}", ex.Message);
+							}
+							Server = null;
+						}
+
+						// Start Appium Server
+						Log.Info("Starting Appium server for device {0} (attempt {1}/{2})", UUID, attemptCount, MaxRetries);
+						AppiumServiceBuilder ServerBuilder = new AppiumServiceBuilder()
+							.UsingAnyFreePort()
+							.WithIPAddress("127.0.0.1");
+
+						if (!string.IsNullOrEmpty(AppiumConfig.AppiumLocation) && FileExists(AppiumConfig.AppiumLocation))
+						{
+							ServerBuilder.WithAppiumJS(new System.IO.FileInfo(AppiumConfig.AppiumLocation));
+						}
+
+						Server = ServerBuilder.Build();
+						Server.Start();
+						Log.Verbose("Appium server started successfully");
+
+						// Create Webdriver instance
+						Log.Info("Creating WebDriver for device {0} (attempt {1}/{2})", UUID, attemptCount, MaxRetries);
+						Driver = new IOSDriver(Server, AppiumConfig.GetCapabilities(UUID, PackageName));
+						Log.Info("Successfully created WebDriver for device {0}", UUID);
+						break; // Success - exit retry loop
 					}
+					catch (Exception Ex)
+					{
+						lastException = Ex;
+						Log.Warning("Appium start attempt {0} failed: {1}", attemptCount, Ex.Message);
 
-					Server = ServerBuilder.Build();
-					Server.Start();
-				}
-				catch (Exception Ex)
-				{
-					throw new AutomationException($"Exception during Appium Server Start {Ex.Message}");
-				}
-				// Create Webdriver instance
-				try
-				{
-					Driver = new IOSDriver(Server, AppiumConfig.GetCapabilities(UUID, PackageName));
+						// Clean up failed server
+						if (Server != null)
+						{
+							try
+							{
+								Server.Dispose();
+							}
+							catch { }
+							Server = null;
+						}
+
+						if (attemptCount < MaxRetries)
+						{
+							int waitSeconds = attemptCount * 5; // Progressive backoff: 5s, 10s, 15s
+							Log.Info("Waiting {0} seconds before retry...", waitSeconds);
+							Thread.Sleep(waitSeconds * 1000);
+						}
+					}
 				}
-				catch (Exception Ex)
+
+				if (Driver == null || Server == null)
 				{
-					throw new DeviceException($"Exception during WebdriverAgent Creation {Ex.Message}");
+					throw new DeviceException($"Failed to start Appium and create WebDriver after {MaxRetries} attempts: {lastException?.Message}");
 				}
 			}
 		}
@@ -368,7 +408,7 @@ namespace Gauntlet
 					// Check for both the base bundle ID and the .xctrunner variant
 					// (actual installed app often has .xctrunner suffix)
 					bool isInstalled = checkResult.Output.Contains(bundleId) ||
-					                   checkResult.Output.Contains(bundleId + ".xctrunner");
+									   checkResult.Output.Contains(bundleId + ".xctrunner");
 					Log.Verbose("Checked for {0} (or .xctrunner variant) on device {1}: {2}", bundleId, deviceUUID, isInstalled ? "found" : "not found");
 					return isInstalled;
 				}
@@ -621,8 +661,8 @@ namespace Gauntlet
 						{
 							string trimmedName = buttonName.Trim();
 							if (trimmedName.Equals("Cancel", StringComparison.OrdinalIgnoreCase) ||
-							    trimmedName.Equals("Dismiss", StringComparison.OrdinalIgnoreCase) ||
-							    trimmedName.Equals("Close", StringComparison.OrdinalIgnoreCase))
+								trimmedName.Equals("Dismiss", StringComparison.OrdinalIgnoreCase) ||
+								trimmedName.Equals("Close", StringComparison.OrdinalIgnoreCase))
 							{
 								Log.Info("Clicking dismiss button on popup: {0}", buttonName);
 								button.Click();
