Add support for custom build filters for StagedBuilds in Gauntlet

#rnx

#Synced-CL 51093944

diff --git a/Engine/Source/Programs/AutomationTool/Gauntlet/Unreal/BuildSource/Gauntlet.StagedBuildSource.cs b/Engine/Source/Programs/AutomationTool/Gauntlet/Unreal/BuildSource/Gauntlet.StagedBuildSource.cs
--- a/Engine/Source/Programs/AutomationTool/Gauntlet/Unreal/BuildSource/Gauntlet.StagedBuildSource.cs
+++ b/Engine/Source/Programs/AutomationTool/Gauntlet/Unreal/BuildSource/Gauntlet.StagedBuildSource.cs
@@ -106,6 +106,11 @@ namespace Gauntlet
 			return InRoleType == Role;
 		}
 
+		public virtual IStagedBuildFilter GetInstallFilter(UnrealAppConfig AppConfig)
+		{
+			return IStagedBuildFilter.GetFilterForPlatform(AppConfig.Platform.Value);
+		}
+
 		public StagedBuild(UnrealTargetPlatform InPlatform, UnrealTargetConfiguration InConfig, UnrealTargetRole InRole, string InBuildPath, string InExecutablePath, string InFlavor="")
 		{
 			Platform = InPlatform;
@@ -178,7 +183,7 @@ namespace Gauntlet
 				{
 					Log.Info("Installing {0} to {1}", AppConfig.Name, Desc);
 
-					IStagedBuildFilter Filter = IStagedBuildFilter.GetFilterForPlatform(AppConfig.Platform.Value);
+					IStagedBuildFilter Filter = InBuild.GetInstallFilter(AppConfig);
 					if (Filter != null)
 					{
 						foreach (FileInfo File in Filter.GetFilteredFiles(AppConfig, InBuild))
