Allow overriding file filtering logic in WindowsStagedBuildFilter

#rnx

#Synced-CL 51093944

diff --git a/Engine/Source/Programs/AutomationTool/Gauntlet/Platform/Windows/Gauntlet.WindowsBuildSources.cs b/Engine/Source/Programs/AutomationTool/Gauntlet/Platform/Windows/Gauntlet.WindowsBuildSources.cs
--- a/Engine/Source/Programs/AutomationTool/Gauntlet/Platform/Windows/Gauntlet.WindowsBuildSources.cs
+++ b/Engine/Source/Programs/AutomationTool/Gauntlet/Platform/Windows/Gauntlet.WindowsBuildSources.cs
@@ -47,8 +47,8 @@ namespace Gauntlet
 					UnrealHelpers.ConfigInfo FileConfig = UnrealHelpers.GetUnrealConfigFromFileName(AppConfig.ProjectName, File.Name);
 					if (FileConfig.Configuration != UnrealTargetConfiguration.Unknown)
 					{
-						// We found a config, that means it is relevant file. Now ensure it matches the config and flavor
-						if (FileConfig.Configuration != AppConfig.Configuration || !FileConfig.Flavor.Equals(Build.Flavor, StringComparison.OrdinalIgnoreCase))
+						// We found a config, that means it is relevant file
+						if (ShouldFilterFile(FileConfig, AppConfig, Build))
 						{
 							bFilter = true;
 							Log.Verbose("\tSkipping copy of {FileName}", File.FullName);
@@ -64,5 +64,11 @@ namespace Gauntlet
 
 			return FilteredFiles;
 		}
+
+		protected virtual bool ShouldFilterFile(UnrealHelpers.ConfigInfo FileConfig, UnrealAppConfig AppConfig, StagedBuild Build)
+		{
+			// Filter if not matches the config and flavor
+			return FileConfig.Configuration != AppConfig.Configuration || !FileConfig.Flavor.Equals(Build.Flavor, StringComparison.OrdinalIgnoreCase);
+		}
 	}
 }
