Add support for $(ProjectName) in [Staging]:RemapDirectories and AllowedDirectories
- include use case for PC GDK

#rnx

#Synced-CL 51094919

diff --git a/Engine/Config/Windows/BaseWindowsGame.ini b/Engine/Config/Windows/BaseWindowsGame.ini
new file mode 100644
--- /dev/null
+++ b/Engine/Config/Windows/BaseWindowsGame.ini
@@ -0,0 +1,3 @@
+[Staging]
+; maintain PC GDK compatibility with other GDK-based platforms
++AllowedDirectories=$(ProjectName)/Platforms/GDK/Config/OSS
diff --git a/Engine/Source/Programs/AutomationTool/AutomationUtils/DeploymentContext.cs b/Engine/Source/Programs/AutomationTool/AutomationUtils/DeploymentContext.cs
--- a/Engine/Source/Programs/AutomationTool/AutomationUtils/DeploymentContext.cs
+++ b/Engine/Source/Programs/AutomationTool/AutomationUtils/DeploymentContext.cs
@@ -562,6 +562,11 @@ public class DeploymentContext //: ProjectParams
 		ProjectArgForCommandLines = ProjectArgForCommandLines.Replace("\\", "/");
 		ProjectBinariesFolder = DirectoryReference.Combine(ProjectUtils.GetClientProjectBinariesRootPath(RawProjectPath, TargetType.Game, IsCodeBasedProject), PlatformDir);
 
+		Dictionary<string,string> SubstitutionVariables = new()
+		{
+			{ "ProjectName", ShortProjectName }
+		};
+
 		// Build a list of restricted folder names. This will comprise all other restricted platforms, plus standard restricted folder names such as NoRedist, NotForLicensees, etc...
 		RestrictedFolderNames.UnionWith(PlatformExports.GetPlatformFolderNames());
 		RestrictedFolderNames.UnionWith(RestrictedFolder.GetNames());
@@ -591,12 +596,14 @@ public class DeploymentContext //: ProjectParams
 				{
 					throw new AutomationException("Missing 'From' property in '{0}'", RemapDirectory);
 				}
+				FromDir = Utils.ExpandVariables(FromDir, SubstitutionVariables, bUseAdditionalVariablesOnly:true);
 
 				string ToDir;
 				if (!Properties.TryGetValue("To", out ToDir))
 				{
 					throw new AutomationException("Missing 'To' property in '{0}'", RemapDirectory);
 				}
+				ToDir = Utils.ExpandVariables(ToDir, SubstitutionVariables, bUseAdditionalVariablesOnly:true);
 
 				RemapDirectories.Add(Tuple.Create(new StagedDirectoryReference(FromDir), new StagedDirectoryReference(ToDir)));
 			}
@@ -608,7 +615,9 @@ public class DeploymentContext //: ProjectParams
 		{
 			foreach (string AllowedDir in DirectoriesAllowListStrings)
 			{
-				DirectoriesAllowList.Add(new StagedDirectoryReference(AllowedDir));
+				string AllowedDirAdjusted = Utils.ExpandVariables(AllowedDir, SubstitutionVariables, bUseAdditionalVariablesOnly:true);
+
+				DirectoriesAllowList.Add(new StagedDirectoryReference(AllowedDirAdjusted));
 			}
 		}
 
