[Gauntlet] Overlay executables will now support finidng executables for custom modules

#Synced-CL 50811283

diff --git a/Engine/Source/Programs/AutomationTool/Gauntlet/Framework/Base/Gauntlet.OverlayExecutable.cs b/Engine/Source/Programs/AutomationTool/Gauntlet/Framework/Base/Gauntlet.OverlayExecutable.cs
--- a/Engine/Source/Programs/AutomationTool/Gauntlet/Framework/Base/Gauntlet.OverlayExecutable.cs
+++ b/Engine/Source/Programs/AutomationTool/Gauntlet/Framework/Base/Gauntlet.OverlayExecutable.cs
@@ -55,11 +55,19 @@ namespace Gauntlet
 				throw new AutomationException("Overlay Error: Could not find base executable {0}.", BaseExecutable);
 			}
 
+			// If this is a custom module, search for those executables instead
+			string ModuleName = ProjectName;
+			string ExecutableModule = Path.GetFileNameWithoutExtension(BaseExecutable);
+			if (UnrealHelpers.CustomModules.TryGetValue(ExecutableModule, out _))
+			{
+				ModuleName = ExecutableModule;
+			}
+
 			// Ex. /UnrealGame/Binaries/Win64/UnrealClient.exe
 			string PlatformBinariesDirectory = Path.Combine(Environment.CurrentDirectory, ProjectName, "Binaries", Platform.ToString());
 			string ExecutableExtension = string.IsNullOrEmpty(ExtensionOverride) ? Path.GetExtension(BaseExecutable) : ExtensionOverride;
 
-			string ExecutableName = UnrealHelpers.GetExecutableName(ProjectName, Platform, Configuration, Role, Flavor, string.Empty);
+			string ExecutableName = UnrealHelpers.GetExecutableName(ModuleName, Platform, Configuration, Role, Flavor, string.Empty);
 			string PreflightExecutableName = ExecutableName + "-Preflight";
 
 			string LocalExecutable = Path.Combine(PlatformBinariesDirectory, ExecutableName + ExecutableExtension);
