ProjectLauncher: add snapshot import to UAT BuildAcquire
- move  snapshot import UI section to the BuildSync extension

#rnx

#Synced-CL 51155769

diff --git a/Engine/Plugins/Developer/ProjectLauncher/Source/CommonLaunchExtensions/Private/BuildSync/BuildInfoHelper.cpp b/Engine/Plugins/Developer/ProjectLauncher/Source/CommonLaunchExtensions/Private/BuildSync/BuildInfoHelper.cpp
--- a/Engine/Plugins/Developer/ProjectLauncher/Source/CommonLaunchExtensions/Private/BuildSync/BuildInfoHelper.cpp
+++ b/Engine/Plugins/Developer/ProjectLauncher/Source/CommonLaunchExtensions/Private/BuildSync/BuildInfoHelper.cpp
@@ -168,6 +168,11 @@ void FBuildInfoHelper::SetBuildType( const FString& InBuildType )
 	}
 }
 
+FString FBuildInfoHelper::GetBuildType() const
+{
+	return BuildType;
+}
+
 void FBuildInfoHelper::SetBuildsRefreshedHandler(TFunction<void()> InOnBuildsRefreshed)
 {
 	OnBuildsRefreshed = InOnBuildsRefreshed;
diff --git a/Engine/Plugins/Developer/ProjectLauncher/Source/CommonLaunchExtensions/Private/BuildSync/BuildSyncLaunchExtension.cpp b/Engine/Plugins/Developer/ProjectLauncher/Source/CommonLaunchExtensions/Private/BuildSync/BuildSyncLaunchExtension.cpp
--- a/Engine/Plugins/Developer/ProjectLauncher/Source/CommonLaunchExtensions/Private/BuildSync/BuildSyncLaunchExtension.cpp
+++ b/Engine/Plugins/Developer/ProjectLauncher/Source/CommonLaunchExtensions/Private/BuildSync/BuildSyncLaunchExtension.cpp
@@ -23,6 +23,7 @@
 #include "Internationalization/Regex.h"
 #include "GenericPlatform/GenericPlatformHttp.h"
 #include "BuildSync/SBuildSyncBuildCombo.h"
+#include "Extension/BuildCookRunLaunchExtension.h"
 
 #define LOCTEXT_NAMESPACE "FBuildSyncLaunchExtensionInstance"
 
@@ -735,6 +736,23 @@ void FBuildSyncLaunchExtensionInstance::OnValidateProfile()
 			}
 		}
 	}
+
+	if (BuildInfo->GetBuildType() == TEXT("oplog") && SelectedBuild.IsValid())
+	{
+		ILauncherProfileBuildCookRunPtr BCR = GetProfile()->GetFirstBuildCookRun();
+
+		if (BCR.IsValid() && BCR->IsEnabled() && BCR->IsUsingZenStreaming() && BCR->GetCookMode() == ELauncherProfileCookModes::DoNotCook)
+		{
+			if (FString* SyncedClStr = GetProfile()->GetCustomStringProperties().Find(TEXT("SyncedChangelist")))
+			{
+				int32 SyncedCl = FCString::Atoi(**SyncedClStr);
+				if (SelectedBuild->Group->Changelist != SyncedCl)
+				{
+					GetProfile()->AddCustomError( TEXT("Validation_BuildSync_NeedCook"), LOCTEXT("Validation_BuildSync_NeedCook", "The downloaded cook snapshot version is different than the synced project version. Enable cooking or incremental cooking"));
+				}
+			}
+		}
+	}
 }
 
 
@@ -753,10 +771,8 @@ void FBuildSyncLaunchExtensionInstance::CustomizeUATCommandLine( FString& InOutC
 	CommandLine += FString::Printf(TEXT(" -CL=%s"), *SelectedBuild->Group->CommitIdentifier);
 
 	FString Branch = FEngineVersion::Current().GetBranch();
-	Branch = Branch.Replace(TEXT("//"), TEXT("")).Replace(TEXT("/"), TEXT("-")).ToLower();
 	CommandLine += FString::Printf(TEXT(" -Stream=%s"), *Branch);
 
-	CommandLine += TEXT(" -getbuild");
 	CommandLine += FString::Printf(TEXT(" -buildtype=%s"), *GetBuildType());
 	CommandLine += FString::Printf(TEXT(" -buildplatform=%s"), *FString::Join(CachedNamedArtifacts, TEXT("+")));
 
diff --git a/Engine/Plugins/Developer/ProjectLauncher/Source/CommonLaunchExtensions/Private/UgsSync/UgsSyncLaunchExtension.cpp b/Engine/Plugins/Developer/ProjectLauncher/Source/CommonLaunchExtensions/Private/UgsSync/UgsSyncLaunchExtension.cpp
--- a/Engine/Plugins/Developer/ProjectLauncher/Source/CommonLaunchExtensions/Private/UgsSync/UgsSyncLaunchExtension.cpp
+++ b/Engine/Plugins/Developer/ProjectLauncher/Source/CommonLaunchExtensions/Private/UgsSync/UgsSyncLaunchExtension.cpp
@@ -155,7 +155,15 @@ void FUGSSyncLaunchExtensionInstance::RefreshLatestSyncedState()
 		return;
 	}
 
-	StateObject->TryGetNumberField(TEXTVIEW("LastSyncChangeNumber"), SyncedChangelist);
+	int32 NewlySyncedChangelist = SyncedChangelist;
+	if (StateObject->TryGetNumberField(TEXTVIEW("LastSyncChangeNumber"), NewlySyncedChangelist))
+	{
+		if (NewlySyncedChangelist >= 0 && NewlySyncedChangelist != SyncedChangelist)
+		{
+			SyncedChangelist = NewlySyncedChangelist;
+			GetProfile()->GetCustomStringProperties().Add(TEXT("SyncedChangelist"), ::LexToString(SyncedChangelist));
+		}
+	}
 }
 
 FUGSSyncLaunchExtension::FUGSSyncLaunchExtension()
diff --git a/Engine/Plugins/Developer/ProjectLauncher/Source/CommonLaunchExtensions/Public/BuildSync/BuildInfoHelper.h b/Engine/Plugins/Developer/ProjectLauncher/Source/CommonLaunchExtensions/Public/BuildSync/BuildInfoHelper.h
--- a/Engine/Plugins/Developer/ProjectLauncher/Source/CommonLaunchExtensions/Public/BuildSync/BuildInfoHelper.h
+++ b/Engine/Plugins/Developer/ProjectLauncher/Source/CommonLaunchExtensions/Public/BuildSync/BuildInfoHelper.h
@@ -32,6 +32,7 @@ public:
 
 	UE_API void SetProjectName( const FString& InProjectName );
 	UE_API void SetBuildType( const FString& InBuildType );
+	UE_API FString GetBuildType() const;
 	UE_API void SetBuildsRefreshedHandler(TFunction<void()> OnBuildsRefreshed);
 
 
diff --git a/Engine/Plugins/Developer/ProjectLauncher/Source/ProjectLauncher/Private/Model/ProjectLauncherModel.cpp b/Engine/Plugins/Developer/ProjectLauncher/Source/ProjectLauncher/Private/Model/ProjectLauncherModel.cpp
--- a/Engine/Plugins/Developer/ProjectLauncher/Source/ProjectLauncher/Private/Model/ProjectLauncherModel.cpp
+++ b/Engine/Plugins/Developer/ProjectLauncher/Source/ProjectLauncher/Private/Model/ProjectLauncherModel.cpp
@@ -824,12 +824,14 @@ PRAGMA_ENABLE_DEPRECATION_WARNINGS
 	{
 		bool bPakFiles = (ContentScheme == EContentScheme::PakFiles || ContentScheme == EContentScheme::DevelopmentPackage || ContentScheme == EContentScheme::SubmissionPackage);
 		bool bUseZen = (ContentScheme != EContentScheme::LooseFiles && ContentScheme != EContentScheme::CookOnTheFly);
+		bool bUseZenStreaming = (ContentScheme == EContentScheme::ZenStreaming);
 		bool bCOTF = (ContentScheme == EContentScheme::CookOnTheFly);
 		bool bPackage = (ContentScheme == EContentScheme::DevelopmentPackage || ContentScheme == EContentScheme::SubmissionPackage);
 		bool bZenPakStreaming = (ContentScheme == EContentScheme::ZenPakStreaming);
 		bool bSubmissionPackage = (ContentScheme == EContentScheme::SubmissionPackage);
 		// FIXME: may need to turn off Zen/Pak in UAT somehow if it's set in project defaults?
 
+		BuildCookRun->SetUseZenStreaming(bUseZenStreaming);
 		BuildCookRun->SetUseZenPakStreaming(bZenPakStreaming);
 		BuildCookRun->SetDeployWithUnrealPak(bPakFiles);
 		BuildCookRun->SetUseZenStore(bUseZen);
@@ -1367,7 +1369,6 @@ PRAGMA_ENABLE_DEPRECATION_WARNINGS
 		if (bShowUATCommand)
 		{
 			MsgTextBuilder.AppendLine(FText::FromString(UATCommand->GetDescription()));
-			MsgTextBuilder.Indent();
 		}
 
 		if (CustomErrors.Num() > 0 || Profile->HasValidationError(UATCommand))
diff --git a/Engine/Plugins/Developer/ProjectLauncher/Source/ProjectLauncher/Private/ProfileTree/CustomProfileTreeBuilder.cpp b/Engine/Plugins/Developer/ProjectLauncher/Source/ProjectLauncher/Private/ProfileTree/CustomProfileTreeBuilder.cpp
--- a/Engine/Plugins/Developer/ProjectLauncher/Source/ProjectLauncher/Private/ProfileTree/CustomProfileTreeBuilder.cpp
+++ b/Engine/Plugins/Developer/ProjectLauncher/Source/ProjectLauncher/Private/ProfileTree/CustomProfileTreeBuilder.cpp
@@ -58,8 +58,6 @@ namespace ProjectLauncher
 			FLaunchProfileTreeNode& ContentSchemeHeader = AddHeading( GeneralBuildCookRunHeader, TEXT("ContentScheme"), LOCTEXT("ContentSchemeHeading", "Content Scheme") );
 			BuildCookRun.AddCompressPakFilesProperty(ContentSchemeHeader);
 			BuildCookRun.AddGenerateChunksProperty(ContentSchemeHeader);
-			BuildCookRun.AddZenSnapshotProperty(ContentSchemeHeader);
-			BuildCookRun.AddImportZenSnapshotProperty(ContentSchemeHeader);
 			BuildCookRun.AddZenPakStreamingPathProperty(ContentSchemeHeader);
 
 			FLaunchProfileTreeNode& SubmissionHeader = AddHeading( GeneralBuildCookRunHeader,  TEXT("Release"), LOCTEXT("ReleaseHeading", "Release / DLC / Patching") );
diff --git a/Engine/Plugins/Developer/ProjectLauncher/Source/ProjectLauncher/Private/Styling/ProjectLauncherStyle.cpp b/Engine/Plugins/Developer/ProjectLauncher/Source/ProjectLauncher/Private/Styling/ProjectLauncherStyle.cpp
--- a/Engine/Plugins/Developer/ProjectLauncher/Source/ProjectLauncher/Private/Styling/ProjectLauncherStyle.cpp
+++ b/Engine/Plugins/Developer/ProjectLauncher/Source/ProjectLauncher/Private/Styling/ProjectLauncherStyle.cpp
@@ -98,6 +98,14 @@ void FProjectLauncherStyle::Initialize()
 	StyleSet->Set("Profile.NoPlatform", new CORE_IMAGE_BRUSH_SVG("Starship/Launcher/PaperAirplane", Icon24x24));
 	StyleSet->Set("Profile.NoPlatform.Large", new CORE_IMAGE_BRUSH_SVG("Starship/Launcher/PaperAirplane", Icon64x64));
 
+	const FTextBlockStyle NormalText = FAppStyle::GetWidgetStyle<FTextBlockStyle>("NormalText");
+	FTextBlockStyle BoldText = FTextBlockStyle(NormalText)
+		.SetColorAndOpacity(FSlateColor::UseForeground())
+		.SetHighlightColor(FLinearColor(0.3f, 0.3f, 0.3f))
+		.SetFont(DEFAULT_FONT("Bold", 10));
+
+	StyleSet->Set("RichTextBlock.Bold", BoldText);
+
 	FSlateStyleRegistry::RegisterSlateStyle( *StyleSet.Get() );
 }
 
diff --git a/Engine/Plugins/Developer/ProjectLauncher/Source/ProjectLauncher/Private/Widgets/Profiles/SCustomLaunchCustomProfileEditor.cpp b/Engine/Plugins/Developer/ProjectLauncher/Source/ProjectLauncher/Private/Widgets/Profiles/SCustomLaunchCustomProfileEditor.cpp
--- a/Engine/Plugins/Developer/ProjectLauncher/Source/ProjectLauncher/Private/Widgets/Profiles/SCustomLaunchCustomProfileEditor.cpp
+++ b/Engine/Plugins/Developer/ProjectLauncher/Source/ProjectLauncher/Private/Widgets/Profiles/SCustomLaunchCustomProfileEditor.cpp
@@ -693,6 +693,7 @@ void SCustomLaunchCustomProfileEditor::Construct(const FArguments& InArgs, const
 					.Padding(4)
 					[
 						SNew(STextBlock)
+						.AutoWrapText(true)
 						.Text_Lambda(GetErrorText)
 					]
 				]
diff --git a/Engine/Source/Developer/LauncherServices/Private/Profiles/LauncherProfile.h b/Engine/Source/Developer/LauncherServices/Private/Profiles/LauncherProfile.h
--- a/Engine/Source/Developer/LauncherServices/Private/Profiles/LauncherProfile.h
+++ b/Engine/Source/Developer/LauncherServices/Private/Profiles/LauncherProfile.h
@@ -2566,7 +2566,7 @@ protected:
 		}
 
 		// Deploy: deployment by copying to devices requires cooking by the book unless the data is available from zen
-		if ((BuildCookRun->GetDeploymentMode() == ELauncherProfileDeploymentModes::CopyToDevice) && ((BuildCookRun->GetCookMode() != ELauncherProfileCookModes::ByTheBook)&&(BuildCookRun->GetCookMode()!=ELauncherProfileCookModes::ByTheBookInEditor)) && !BuildCookRun->IsUsingZenPakStreaming() && !BuildCookRun->IsImportingZenSnapshot())
+		if ((BuildCookRun->GetDeploymentMode() == ELauncherProfileDeploymentModes::CopyToDevice) && ((BuildCookRun->GetCookMode() != ELauncherProfileCookModes::ByTheBook)&&(BuildCookRun->GetCookMode()!=ELauncherProfileCookModes::ByTheBookInEditor)) && !BuildCookRun->IsUsingZenPakStreaming() && !BuildCookRun->IsUsingZenStreaming())
 		{
 			ValidationData.Errors.Add(ELauncherProfileValidationErrors::CopyToDeviceRequiresCookByTheBook);
 		}
diff --git a/Engine/Source/Developer/LauncherServices/Private/Profiles/LauncherProfileBuildCookRun.h b/Engine/Source/Developer/LauncherServices/Private/Profiles/LauncherProfileBuildCookRun.h
--- a/Engine/Source/Developer/LauncherServices/Private/Profiles/LauncherProfileBuildCookRun.h
+++ b/Engine/Source/Developer/LauncherServices/Private/Profiles/LauncherProfileBuildCookRun.h
@@ -780,6 +780,7 @@ public:
 		Writer.WriteValue("BuildTargetSpecified", BuildTargetSpecified);
 		Writer.WriteValue("UseZenStore", bUseZenStore);
 		Writer.WriteValue("ImportingZenSnapshot", bImportingZenSnapshot);
+		Writer.WriteValue("UseZenStreaming", bUseZenStreaming);
 		Writer.WriteValue("UseZenPakStreaming", bUseZenPakStreaming);
 		Writer.WriteValue("ZenPakStreamingPath", ZenPakStreamingPath);
 		Writer.WriteValue("ClientArchitectures", ClientArchitectures);
@@ -1072,6 +1073,7 @@ public:
 			bUseZenStore = Object.GetBoolField(TEXT("UseZenStore"));
 			bImportingZenSnapshot = Object.GetBoolField(TEXT("ImportingZenSnapshot"));
 			bUseZenPakStreaming = Object.GetBoolField(TEXT("UseZenPakStreaming"));
+			bUseZenStreaming = Object.GetBoolField(TEXT("UseZenStreaming"));
 			ZenPakStreamingPath = Object.GetStringField(TEXT("ZenPakStreamingPath"));
 		}
 
@@ -1241,6 +1243,7 @@ public:
 		bUseIoStore = false;
 		bUseZenStore = false;
 		bImportingZenSnapshot = false;
+		bUseZenStreaming = false;
 		bUseZenPakStreaming = false;
 		bIsDeviceASimulator = false;
 		bMakeBinaryConfig = false;
@@ -1658,6 +1661,21 @@ public:
 		return bImportingZenSnapshot;
 	}
 
+	virtual bool IsUsingZenStreaming() const override
+	{
+		return bUseZenStreaming;
+	}
+
+	virtual void SetUseZenStreaming( bool UseZenStreaming) override
+	{
+		if (bUseZenStreaming != UseZenStreaming)
+		{
+			bUseZenStreaming = UseZenStreaming;
+
+			Validate();
+		}
+	}
+
 	virtual void SetUseZenPakStreaming( bool UseZenPakStreaming ) override
 	{
 		if (bUseZenPakStreaming != UseZenPakStreaming)
@@ -2087,6 +2105,8 @@ private:
 	// Import a Zen snapshot before cooking
 	bool bImportingZenSnapshot;
 
+	bool bUseZenStreaming;
+
 	// Use Zen pak streaming
 	bool bUseZenPakStreaming;
 
diff --git a/Engine/Source/Developer/LauncherServices/Public/ILauncherProfileBuildCookRun.h b/Engine/Source/Developer/LauncherServices/Public/ILauncherProfileBuildCookRun.h
--- a/Engine/Source/Developer/LauncherServices/Public/ILauncherProfileBuildCookRun.h
+++ b/Engine/Source/Developer/LauncherServices/Public/ILauncherProfileBuildCookRun.h
@@ -1109,6 +1109,20 @@ public:
 	 */
 	virtual bool IsImportingZenSnapshot() const = 0;
 
+	/**
+	 * Checks whether Zen streaming will be used.
+	 * 
+	 * @return true if using Zen streaming, false otherwise
+	 */
+	virtual bool IsUsingZenStreaming() const = 0;
+	
+	/**
+	 * Set whether to use Zen streaming.
+	 *
+	 * @param UseZenStreaming Whether use Zen stremaing.
+	 */
+	virtual void SetUseZenStreaming( bool UseZenStreaming ) = 0;
+
 	/**
 	 * Set whether to use Zen pak streaming.
 	 * 
diff --git a/Engine/Source/Programs/AutomationTool/Scripts/BuildAcquire.Automation.cs b/Engine/Source/Programs/AutomationTool/Scripts/BuildAcquire.Automation.cs
--- a/Engine/Source/Programs/AutomationTool/Scripts/BuildAcquire.Automation.cs
+++ b/Engine/Source/Programs/AutomationTool/Scripts/BuildAcquire.Automation.cs
@@ -13,10 +13,12 @@ using System.ComponentModel;
 using System.Linq;
 using System.Reflection;
 using System.Diagnostics;
+using System.IO;
 using UnrealBuildBase;
-using EpicGames.UHT.Types;
 using System.Text.RegularExpressions;
-using Microsoft.Extensions.FileSystemGlobbing.Internal.PathSegments;
+using UnrealBuildTool;
+using EpicGames.Horde;
+using System.Reflection.Metadata;
 
 class BuildInfo
 {
@@ -51,6 +53,11 @@ public static class EnumExtensions
 
 	public static string GetEnumDescription(Enum value)
 	{
+		if (value == null)
+		{
+			return String.Empty;
+		}
+
 		FieldInfo fi = value.GetType().GetField(value.ToString());
 		DescriptionAttribute[] attributes = (DescriptionAttribute[])fi.GetCustomAttributes(typeof(DescriptionAttribute), false);
 
@@ -89,8 +96,11 @@ class BuildAcquire : BuildCommand, IProjectParamsHelpers
 	{
 		[Description("staged-build")]
 		Staged,
-		[Description("packaged")]
-		Packaged
+		[Description("packaged-build")]
+		Packaged,
+		[Description("oplog")]
+		Snapshot,
+		Unspecified
 	}
 
 	public BuildAcquire()
@@ -110,12 +120,19 @@ class BuildAcquire : BuildCommand, IProjectParamsHelpers
 				RawProjectPath: ProjectPath
 			);
 
+		Params.Validate();
+
 		if (DownloadBuild)
 		{
 			DownloadBuilds(Params);
 		}
 
-		if (SyncUGS)
+		if (DownloadSnapshot)
+		{
+			ImportSnapshot(Params);
+		}
+
+		if (SyncUGS && !DryRun)
 		{
 			if (!IsEditorRunning(Params))
 			{
@@ -141,11 +158,15 @@ class BuildAcquire : BuildCommand, IProjectParamsHelpers
 			BuildConfigType = BuildVariant.Development;
 		}
 
+		DryRun = ParseParam("DryRun");
 		SyncUGS = ParseParam("SyncUGS");
+		ClosestSnapshot = ParseParam("ClosestSnapshot");
+		DownloadSnapshot = (BuildType == BuildTypeEnum.Snapshot) || ClosestSnapshot;
 		DownloadClient = ParseParam("DownloadClient");
 		DownloadServer = ParseParam("DownloadServer");
 		DownloadGame = ParseParam("DownloadGame");
-		DownloadBuild = DownloadClient || DownloadServer || DownloadGame || ParseParam("BuildPlatform");
+		BuildPlatform = ParseParamValue("BuildPlatform");
+		DownloadBuild = DownloadClient || DownloadServer || DownloadGame || (!String.IsNullOrEmpty(BuildPlatform) && !DownloadSnapshot && BuildType != BuildTypeEnum.Unspecified);
 
 		Stream = ParseParamValue("Stream");
 		if (Stream == null)
@@ -159,9 +180,9 @@ class BuildAcquire : BuildCommand, IProjectParamsHelpers
 			throw new AutomationException("No CL number specified. Use -CL=<changelist>");
 		}
 
-		if (!DownloadBuild && !SyncUGS)
+		if (!DownloadBuild && !SyncUGS && !DownloadSnapshot)
 		{
-			throw new AutomationException("No action to perform. To download a build use -DownloadGame|-DownloadServer|-DownloadClient or -BuildPlatform. To sync UGS pass in -SyncUGS");
+			throw new AutomationException("No action to perform specified on the commandline.");
 		}
 
 		UGSExePath = FileReference.Combine(new DirectoryReference(Environment.GetFolderPath(Environment.SpecialFolder.LocalApplicationData)), "UnrealGameSync", "Latest", "ugs.exe");
@@ -173,6 +194,15 @@ class BuildAcquire : BuildCommand, IProjectParamsHelpers
 			}
 		}
 
+		ZenExePath = new FileReference(Path.Combine(Unreal.EngineDirectory.FullName, "Binaries", HostPlatform.Current.HostEditorPlatform.ToString(), String.Format("zen{0}", RuntimePlatform.ExeExtension)));
+		if (DownloadSnapshot)
+		{
+			if (!FileReference.Exists(ZenExePath))
+			{
+				throw new AutomationException("Zen executable not found at {}, cannot continue", ZenExePath.ToString());
+			}
+		}
+
 		CloudStorageService = CommandUtils.ServiceScope.ServiceProvider.GetService<ICloudStorage>();
 		if (CloudStorageService == null)
 		{
@@ -199,25 +229,35 @@ class BuildAcquire : BuildCommand, IProjectParamsHelpers
 
 	private void DownloadBuilds(ProjectParams Params)
 	{
+		string BuildTypeDesc = EnumExtensions.GetEnumDescription(BuildType);
+		if (String.IsNullOrEmpty(BuildTypeDesc))
+		{
+			throw new AutomationException("Build type description is uncrecognized");
+		}
+
 		CloudStorageQueryBuilder query = new CloudStorageQueryBuilder();
-		query.Equals("type", BuildType);
+		query.Equals("type", BuildTypeDesc);
 		query.Equals("commit", Changelist);
 		query.Equals("isPreflight", "false");
 
-		string[] BuildPlatforms = ParseOptionalStringParam("BuildPlatform")?.Split('+');
+		string[] BuildPlatforms = BuildPlatform?.Split('+');
+		string[] Overlay = ParseOptionalStringParam("Overlay")?.Split('+');
 
 		CbObject obj = query.Done();
 
-		List<FindBuildResponse> foundBuilds = CloudStorageService.FindBuilds(ProjectPath.GetFileNameWithoutExtension(), BuildType, Stream, obj).ToList();
+		string SanitizedStream = StringId.Sanitize(Stream).ToString().Replace(".", "-", StringComparison.Ordinal);
+
+		List<FindBuildResponse> foundBuilds = CloudStorageService.FindBuilds(ProjectPath.GetFileNameWithoutExtension(), BuildTypeDesc, SanitizedStream, obj).ToList();
         List<FindBuildResponse> BuildsToDownload = new();
+		List<FindBuildResponse> OverlayBuilds = new();
 		
 		if (foundBuilds.Count > 0)
 		{
-			if (BuildPlatforms != null && BuildPlatforms.Length > 0 && ParseParam("GetBuild"))
+			if (BuildPlatforms?.Length > 0)
 			{
-				foreach (string BuildPlatform in BuildPlatforms)
+				foreach (string BuildPlatformEntry in BuildPlatforms)
 				{
-					BuildsToDownload.Add(foundBuilds.Where(Build => Build.Name.Equals(BuildPlatform, StringComparison.OrdinalIgnoreCase)).First());
+					BuildsToDownload.Add(foundBuilds.Where(Build => Build.Name.Equals(BuildPlatformEntry, StringComparison.OrdinalIgnoreCase)).First());
 				}
 			}
 			else
@@ -243,6 +283,19 @@ class BuildAcquire : BuildCommand, IProjectParamsHelpers
 			}
 
 			BuildsToDownload = BuildsToDownload.DistinctBy(Build => Build.Name).ToList();
+
+			if (Overlay?.Length > 0)
+			{
+				foreach (string OverlayBuild in Overlay)
+				{
+					OverlayBuilds.Add(foundBuilds.Where(Build => Build.Name.Equals(OverlayBuild, StringComparison.OrdinalIgnoreCase)).First());
+				}
+			}
+		}
+
+		if (DryRun)
+		{
+			Logger.LogInformation("Will be performing a dry run");
 		}
 
 		if (BuildsToDownload.Count == 0)
@@ -258,15 +311,134 @@ class BuildAcquire : BuildCommand, IProjectParamsHelpers
 
         foreach (FindBuildResponse Build in BuildsToDownload)
         {
-            DirectoryReference FinalDownloadDir = DirectoryReference.Combine(new DirectoryReference(Params.BaseStageDirectory), Build.Metadata.Find("cookPlatform").AsString());
-			CloudStorageService.DownloadBuild(Build.BucketId, Build.BuildId, FinalDownloadDir, downloadOptions: new CloudBuildDownloadOptions
+			DirectoryReference BaseDir;
+			//DirectoryReference FinalDownloadDir = DirectoryReference.Combine(new DirectoryReference(Params.BaseStageDirectory), Build.Metadata.Find("cookPlatform").AsString());
+			if (BuildType == BuildTypeEnum.Snapshot)
+			{
+				BaseDir = new DirectoryReference(Path.Combine(Params.RawProjectPath.Directory.FullName, "Saved", "Cooked"));
+			}
+			else
+			{
+				BaseDir = new DirectoryReference(Params.BaseStageDirectory);
+			}
+
+			DirectoryReference FinalDownloadDir = DirectoryReference.Combine(BaseDir, Build.Metadata.Find("cookPlatform").AsString());
+            
+			Logger.LogInformation("Downloading {Name}:{BuildId} to {FinalDownloadDir}", Build.Name, Build.BuildId, FinalDownloadDir);
+
+			if (!DryRun)
 			{
-				AssumeHttp2 = true,
-				BoostWorkers = true
-			});
+				CloudStorageService.DownloadBuild(Build.BucketId, Build.BuildId, FinalDownloadDir, downloadOptions: new CloudBuildDownloadOptions
+				{
+					AssumeHttp2 = true,
+					BoostWorkers = true
+				});
+
+				foreach (FindBuildResponse OverlayBuild in OverlayBuilds)
+				{
+					CloudStorageService.DownloadBuild(OverlayBuild.BucketId, OverlayBuild.BuildId, FinalDownloadDir, downloadOptions: new CloudBuildDownloadOptions
+					{
+						AssumeHttp2 = true,
+						BoostWorkers = true,
+						Append = true
+					});
+				}
+			}
         }
     }
 
+	private struct CloudStorageConfig
+	{
+		public string Host = "";
+		public string BuildsNamespace = "";
+		public string OAuthProviderIdentifier = "";
+
+		public CloudStorageConfig()
+		{
+		}
+	}
+
+	private List<string> GetCookPlatforms(ProjectParams Params)
+	{
+		var PlatformsSet = new HashSet<string>();
+
+		if (!Params.NoClient)
+		{
+			foreach (var ClientPlatform in Params.ClientTargetPlatforms)
+			{
+				var PlatformDescriptor = Params.GetCookedDataPlatformForClientTarget(ClientPlatform);
+				string CookPlatform = Platform.Platforms[PlatformDescriptor].GetCookPlatform(false, Params.Client);
+				PlatformsSet.Add(CookPlatform);
+			}
+		}
+
+		if (Params.DedicatedServer)
+		{
+			foreach (var ServerPlatform in Params.ServerTargetPlatforms)
+			{
+				var PlatformDescriptor = Params.GetCookedDataPlatformForServerTarget(ServerPlatform);
+				string CookPlatform = Platform.Platforms[PlatformDescriptor].GetCookPlatform(true, false);
+				PlatformsSet.Add(CookPlatform);
+			}
+		}
+
+		if (!String.IsNullOrEmpty(BuildPlatform))
+		{
+			PlatformsSet.AddRange(BuildPlatform.Split('+'));
+		}
+
+		return PlatformsSet.ToList();
+	}
+
+	private void ImportSnapshot(ProjectParams Params)
+	{
+
+		ConfigHierarchy HostPlatformConfig = Params.EngineConfigs[HostPlatform.Current.HostEditorPlatform];
+		if (!HostPlatformConfig.TryGetValueGeneric("StorageServers", "Cloud", out CloudStorageConfig cloudConfig))
+		{
+			throw new AutomationException("Failed to read config values for storage servers");
+		}
+
+		UnrealCloudDDCBuildIndex Index = new UnrealCloudDDCBuildIndex(cloudConfig.Host.Split(";")[0], cloudConfig.BuildsNamespace, cloudConfig.OAuthProviderIdentifier, (string app, string commandLine) => CommandUtils.Run(app, commandLine));
+		DirectoryReference CookOutputDir;
+		if (String.IsNullOrEmpty(Params.CookOutputDir))
+		{
+			CookOutputDir = new DirectoryReference(Path.Combine(Params.RawProjectPath.Directory.FullName, "Saved", "Cooked", "{Platform}"));
+		}
+		else
+		{
+			CookOutputDir = new DirectoryReference(Params.CookOutputDir);
+		}
+
+		foreach (string CookPlatform in GetCookPlatforms(Params))
+		{
+			string Platform;
+			string Runtime;
+
+			if (CookPlatform.EndsWith("client", StringComparison.OrdinalIgnoreCase))
+			{
+				Platform = CookPlatform.Substring(0, CookPlatform.Length - "client".Length);
+				Runtime = "client";
+			}
+			else if (CookPlatform.EndsWith("server", StringComparison.OrdinalIgnoreCase))
+			{
+				Platform = CookPlatform.Substring(0, CookPlatform.Length - "server".Length);
+				Runtime = "server";
+			}
+			else
+			{
+				Platform = CookPlatform;
+				Runtime = "game";
+			}
+
+			EpicGames.ProjectStore.Build Snapshot = Index.GetBuild(Stream, ProjectPath.GetFileNameWithoutExtension(), Platform, Runtime, Int32.Parse(Changelist), ClosestSnapshot);
+			if (!DryRun)
+			{
+				Snapshot.Import(Params.RawProjectPath, CookOutputDir, false, false);
+			}
+		}
+	}
+
 	private bool IsEditorRunning(ProjectParams Params)
 	{
 		Process[] ProcessArr = Process.GetProcesses();
@@ -279,6 +451,11 @@ class BuildAcquire : BuildCommand, IProjectParamsHelpers
 				continue;
 			}
 
+			if (process.HasExited)
+			{
+				continue;
+			}
+
 			try
 			{
 				FileReference ProcessFile = new FileReference(process.MainModule.FileName);
@@ -300,15 +477,20 @@ class BuildAcquire : BuildCommand, IProjectParamsHelpers
 
 	private ICloudStorage CloudStorageService;
 
-	string _BuildType;
-	private string BuildType
+	//string _BuildType;
+	BuildTypeEnum? _BuildType;
+	private BuildTypeEnum? BuildType
 	{
 		get
 		{
 			if (_BuildType == null)
 			{
-                String CmdBuildType = ParseRequiredStringParam("BuildType");
-				if (!Enum.TryParse(CmdBuildType, true, out BuildTypeEnum BuildType))
+				String CmdBuildType = ParseOptionalStringParam("BuildType");
+				if (String.IsNullOrEmpty(CmdBuildType))
+				{
+					_BuildType = BuildTypeEnum.Unspecified;
+				}
+				else if (!Enum.TryParse(CmdBuildType, true, out BuildTypeEnum OutBuildType))
 				{
 					EnumDescription BuildTypeDescriptions = EnumExtensions.ConvertEnumWithDescription<BuildTypeEnum>();
 					int BuildTypeIndex = BuildTypeDescriptions.Descriptions.IndexOf(CmdBuildType);
@@ -319,10 +501,14 @@ class BuildAcquire : BuildCommand, IProjectParamsHelpers
 					}
 					else
 					{
-						_BuildType = BuildTypeDescriptions.Descriptions[BuildTypeIndex];
+						_BuildType = (BuildTypeEnum)BuildTypeIndex;
 					}
 				}
-				_BuildType = EnumExtensions.GetEnumNameOrDescription(BuildType);
+				else
+				{
+					//_BuildType = EnumExtensions.GetEnumNameOrDescription(OutBuildType);
+					_BuildType = OutBuildType;
+				}
 			}
 
 			return _BuildType;
@@ -332,14 +518,19 @@ class BuildAcquire : BuildCommand, IProjectParamsHelpers
 	string Stream;
 	string Changelist;
 	FileReference ProjectPath;
+	string BuildPlatform;
 	bool DownloadServer;
 	bool DownloadClient;
 	bool DownloadGame;
+	bool DownloadSnapshot;
+	bool ClosestSnapshot;
 
 	bool SyncUGS;
 	bool DownloadBuild;
+	bool DryRun;
 
 	FileReference UGSExePath;
+	FileReference ZenExePath;
 	// TODO: Add support for build config once established how PL2 wants to pass in build configs. Currently it uses build platform which contains the full set: platform, config, executable type
 	BuildVariant BuildConfigType;
 }
