[Gauntlet] Revert behavior that inserts default devices in front of local devices

#Synced-CL 51172295

diff --git a/Engine/Source/Programs/AutomationTool/Gauntlet/Framework/Devices/Gauntlet.DevicePool.cs b/Engine/Source/Programs/AutomationTool/Gauntlet/Framework/Devices/Gauntlet.DevicePool.cs
--- a/Engine/Source/Programs/AutomationTool/Gauntlet/Framework/Devices/Gauntlet.DevicePool.cs
+++ b/Engine/Source/Programs/AutomationTool/Gauntlet/Framework/Devices/Gauntlet.DevicePool.cs
@@ -692,8 +692,15 @@ namespace Gauntlet
 							Def.Name = DeviceAddress;
 							Def.Platform = DevicePlatform;
 							Def.Model = Model;
-							UnprovisionedDevices.Insert(UserDeviceIndex, Def); // add user-specified devices at the start of the list so they are given priority
-							UserDeviceIndex++;
+							if (Def.Name.Equals("default", StringComparison.OrdinalIgnoreCase))
+							{
+								UnprovisionedDevices.Add(Def);
+							}
+							else
+							{
+								// add user-specified devices at the start of the list so they are given priority
+								UnprovisionedDevices.Insert(UserDeviceIndex++, Def);
+							}
 						}
 					}
 				}
