[Gauntlet] Adds IgnoreRoleFailures option to UnrealTestConfig
When set to true, roles that encounter problems will log as information instead of errors.
This will prevent horde jobs from being marked as a fail when the job is actively looking for failures

#Synced-CL 50811283

diff --git a/Engine/Source/Programs/AutomationTool/Gauntlet/Unreal/Base/Gauntlet.UnrealTestConfiguration.cs b/Engine/Source/Programs/AutomationTool/Gauntlet/Unreal/Base/Gauntlet.UnrealTestConfiguration.cs
--- a/Engine/Source/Programs/AutomationTool/Gauntlet/Unreal/Base/Gauntlet.UnrealTestConfiguration.cs
+++ b/Engine/Source/Programs/AutomationTool/Gauntlet/Unreal/Base/Gauntlet.UnrealTestConfiguration.cs
@@ -21,7 +21,7 @@ namespace Gauntlet
 		public static readonly Regex CommandlineOptionRegex = new Regex(@"-(?<option>-?[\w\d.:!\[\]\/\\?-]+)(=(?<value>(""(?:(?:\\""|[^""])+)"")|(?:\S+)))?", RegexOptions.Compiled);
 
 		// Give external people read-only access
-		public IReadOnlyDictionary<string, object> Arguments {  get { return Params;  } }
+		public IReadOnlyDictionary<string, object> Arguments { get { return Params; } }
 
 		/// <summary>
 		/// String containing extra commandline args that do not conform to UE commandline arg specs.
@@ -622,7 +622,7 @@ namespace Gauntlet
 		/// The minimum time interval between retrieval of heartbeat logs
 		/// </summary>
 		public float LogHeartbeatInterval;
-		
+
 		public UnrealHeartbeatOptions(float InHeartbeatPeriod = 30f, bool bShouldExpectHeartbeats = false, float InTimeoutBeforeFirstActiveHeartbeat = 0f, float InTimeoutBetweenActiveHeartbeats = 0f, float InTimeoutBetweenAnyHeartbeats = 90f, float InLogHeartbeatInterval = 0f)
 		{
 			HeartbeatPeriod = InHeartbeatPeriod;
@@ -842,6 +842,11 @@ namespace Gauntlet
 		[AutoParam]
 		public string TagFilter { get; set; }
 
+		/// <summary>
+		/// If true, an individual process role failing (i.e. crash) does not flag the test as a failure
+		/// </summary>
+		public bool IgnoreRoleFailures { get; set; }
+
 		// Member variables 
 
 		/// <summary>
@@ -892,7 +897,7 @@ namespace Gauntlet
 		/// <returns></returns>
 		public UnrealTestRole RequireRole(UnrealTargetRole InRole)
 		{
-			if(InRole.IsEditor())
+			if (InRole.IsEditor())
 			{
 				return GetEditorRole();
 			}
@@ -920,7 +925,7 @@ namespace Gauntlet
 		/// <param name="InRole"></param>
 		/// <param name="Count"></param>
 		/// <returns></returns>
-		public IEnumerable<UnrealTestRole> RequireRoles(UnrealTargetRole InRole, int Count)
+		public List<UnrealTestRole> RequireRoles(UnrealTargetRole InRole, int Count)
 		{
 			return RequireRoles(InRole, null, Count);
 		}
@@ -933,7 +938,7 @@ namespace Gauntlet
 			RequiredRoles.Clear();
 		}
 
-		public IEnumerable<UnrealTestRole> RequireRoles(UnrealTargetRole InRole, UnrealTargetPlatform? PlatformOverride, int Count, ERoleModifier roleType = ERoleModifier.None)
+		public List<UnrealTestRole> RequireRoles(UnrealTargetRole InRole, UnrealTargetPlatform? PlatformOverride, int Count, ERoleModifier roleType = ERoleModifier.None)
 		{
 			if (RequiredRoles.ContainsKey(InRole) == false)
 			{
@@ -977,7 +982,7 @@ namespace Gauntlet
 		/// </summary>
 		/// <param name="InRole"></param>
 		/// <returns></returns>
-		public IEnumerable<UnrealTestRole> GetRequiredRoles(UnrealTargetRole InRole)
+		public List<UnrealTestRole> GetRequiredRoles(UnrealTargetRole InRole)
 		{
 			if (RequiredRoles.ContainsKey(InRole))
 			{
diff --git a/Engine/Source/Programs/AutomationTool/Gauntlet/Unreal/Base/Gauntlet.UnrealTestNode.cs b/Engine/Source/Programs/AutomationTool/Gauntlet/Unreal/Base/Gauntlet.UnrealTestNode.cs
--- a/Engine/Source/Programs/AutomationTool/Gauntlet/Unreal/Base/Gauntlet.UnrealTestNode.cs
+++ b/Engine/Source/Programs/AutomationTool/Gauntlet/Unreal/Base/Gauntlet.UnrealTestNode.cs
@@ -609,7 +609,7 @@ namespace Gauntlet
 			return TestPriority.Normal;
 		}
 
-		protected virtual IEnumerable<UnrealSessionInstance.RoleInstance> FindAnyMissingRoles()
+		protected virtual IEnumerable<RoleInstance> FindAnyMissingRoles()
 		{
 			return TestInstance.RunningRoles.Where(R => R.AppInstance.HasExited);
 		}
@@ -1310,7 +1310,7 @@ namespace Gauntlet
 					Horde.UpdateStepMessage(Horde.GauntletHordeStepMessage.TestTimeoutExceeded);
 					SetTestResult(TestResult.TimedOut);
 				}
-				else
+				else if (!GetCachedConfiguration().IgnoreRoleFailures)
 				{
                     SetTestResult(TestResult.OperationalException);
                 }
@@ -1389,7 +1389,6 @@ namespace Gauntlet
 		/// </summary>
 		public virtual void AddProcessResultEventsFromTestNode()
 		{
-
 			if (RoleResults != null)
 			{
 				// First, make sure we have all of our process exit events.
@@ -2522,7 +2521,7 @@ namespace Gauntlet
 			bool HasFailed = InRoleResult.ProcessResult != UnrealProcessResult.ExitOk && InRoleResult.ProcessResult != UnrealProcessResult.EngineTestError;
 			string RoleState = HasFailed ? "failed:" : "completed:";
 			string StatusMessage = string.Format(" #### {0} {1} {2} ({3}, ExitCode={4})", RoleArtifacts.SessionRole.RoleType, RoleState, InRoleResult.Summary, InRoleResult.ProcessResult, InRoleResult.ExitCode);
-			if (HasFailed)
+			if (HasFailed && !GetCachedConfiguration().IgnoreRoleFailures)
 			{
 				Log.Error(KnownLogEvents.Gauntlet_TestEvent, StatusMessage);
 			}
@@ -2553,7 +2552,16 @@ namespace Gauntlet
 			foreach (UnrealTestEvent Event in Asserts)
 			{
 				LogEvent AssertData = DecomposeAssert(Event, MaxCallstackLines);
-				Log.Error(AssertData.Id, AssertData.Format ?? AssertData.Message, ((Dictionary<string, object>)AssertData.Properties)?.Values.ToArray());
+				
+				if (GetCachedConfiguration().IgnoreRoleFailures)
+				{
+					Log.Info(AssertData.Id, AssertData.Format ?? AssertData.Message, ((Dictionary<string, object>)AssertData.Properties)?.Values.ToArray());
+				}
+				else
+				{
+					Log.Error(AssertData.Id, AssertData.Format ?? AssertData.Message, ((Dictionary<string, object>)AssertData.Properties)?.Values.ToArray());
+				}
+					
 				Log.Info("");
 			}
 
@@ -2579,7 +2587,16 @@ namespace Gauntlet
 				{
 					Callstack.Add("Could not parse callstack. See log for full callstack");
 				}
-				Log.Warning(KnownLogEvents.Gauntlet_TestEvent, " * Ensure: {Summary}\n{Callstack}", Event.Summary, string.Join("\n", Callstack.Select(C => "    " + C)));
+
+				if (GetCachedConfiguration().IgnoreRoleFailures)
+				{
+					Log.Info(KnownLogEvents.Gauntlet_TestEvent, " * Ensure: {Summary}\n{Callstack}", Event.Summary, string.Join("\n", Callstack.Select(C => "    " + C)));
+				}
+				else
+				{
+					Log.Warning(KnownLogEvents.Gauntlet_TestEvent, " * Ensure: {Summary}\n{Callstack}", Event.Summary, string.Join("\n", Callstack.Select(C => "    " + C)));
+				}
+					
 				Log.Info("");
 			}
 
@@ -2607,7 +2624,14 @@ namespace Gauntlet
 
 				foreach (var Error in PrintedErrorList)
 				{
-					Log.Error(KnownLogEvents.Gauntlet_TestEvent, " * {Message}", Error);
+					if (GetCachedConfiguration().IgnoreRoleFailures)
+					{
+						Log.Info(KnownLogEvents.Gauntlet_TestEvent, " * {Message}", Error);
+					}
+					else
+					{
+						Log.Error(KnownLogEvents.Gauntlet_TestEvent, " * {Message}", Error);
+					}	
 				}
 				if (!string.IsNullOrEmpty(TrimStatement))
 				{
