Fix deletion of files with names invalid on Windows

#rnx

#Synced-CL 51037731

diff --git a/Engine/Source/Programs/AutomationTool/Gauntlet/Framework/Gauntlet.Utils.cs b/Engine/Source/Programs/AutomationTool/Gauntlet/Framework/Gauntlet.Utils.cs
--- a/Engine/Source/Programs/AutomationTool/Gauntlet/Framework/Gauntlet.Utils.cs
+++ b/Engine/Source/Programs/AutomationTool/Gauntlet/Framework/Gauntlet.Utils.cs
@@ -1739,13 +1739,24 @@ namespace Gauntlet
 						foreach(FileSystemInfo Info in Directory.EnumerateFiles("*", SearchOption.AllDirectories))
 						{
 							Log.Info("Deleting file after attributes change: {FileName}", Info.FullName);
-							Info.Attributes = FileAttributes.Normal;
+
+							try
+							{
+								Info.Attributes = FileAttributes.Normal;
+							}
+							catch (ArgumentException ArgEx)
+							{
+								Log.Warning("Could not set attributes for {File}. Error: {Error}", Info.FullName, ArgEx.Message);
+								// Continue to deletion attempt - we'll try using extended-length path
+							}
+
 							// Try to delete the file, with retry logic for locked files
 							bool bDeleted = false;
 							for (int Attempt = 0; Attempt < 3 && !bDeleted; Attempt++)
 							{
 								try
 								{
+									// First try normal deletion
 									Info.Delete();
 									bDeleted = true;
 								}
@@ -1768,6 +1779,19 @@ namespace Gauntlet
 										throw;
 									}
 								}
+								catch (Exception)
+								{
+									// Handle files with invalid Windows filenames using extended-length path
+									Log.Warning("Attempting deletion with extended-length path: {File}", Info.FullName);
+									if (Info.FullName.StartsWith(@"\\"))
+									{
+										// Network or already an extended-length path
+										throw;
+									}
+									string ExtendedPath = @"\\?\" + Info.FullName;
+									System.IO.File.Delete(ExtendedPath);
+									bDeleted = true;
+								}
 							}
 						}
 					}
