Fix issue with code-only projects using wrong architectures for packaging on Android

#android

#Synced-CL 51027549

diff --git a/Engine/Source/Programs/AutomationTool/Android/AndroidPlatform.Automation.cs b/Engine/Source/Programs/AutomationTool/Android/AndroidPlatform.Automation.cs
--- a/Engine/Source/Programs/AutomationTool/Android/AndroidPlatform.Automation.cs
+++ b/Engine/Source/Programs/AutomationTool/Android/AndroidPlatform.Automation.cs
@@ -888,7 +888,20 @@ public class AndroidPlatform : Platform
 
 	private UnrealArchitectures GetDeploymentArchitectures(ProjectParams Params, DeploymentContext SC)
 	{
-		return Params.ClientArchitecture ?? SC.StageTargets[0].Receipt.Architectures;
+		if (Params.ClientArchitecture != null)
+		{
+			return Params.ClientArchitecture;
+		}
+
+		// use Receipt architectures if bIsCodeBasedProject (UnrealGame.target used if no code or plugin differences)
+		if (SC.StageTargets[0].Receipt.TargetName != "UnrealGame")
+		{
+			Logger.LogInformation("GetDeploymentArchitectures: not UnrealGame");
+			return SC.StageTargets[0].Receipt.Architectures;
+		}
+		
+		// otherwise figure it out from project ini
+		return UnrealArchitectureConfig.ForPlatform(UnrealTargetPlatform.Android).ActiveArchitectures(Params.RawProjectPath, null);
 	}
 
 	private List<string> CollectPluginDataPaths(DeploymentContext SC)
