Adding wait and retry logic to IOS deployment scripts.

This is needed so the project launcher tool and other UAT usages do not just error out if the device is locked

This can be bypassed by by passing the commandline argument -DeviceLockedTimeout=0

#tests I tested that I can

#Synced-CL 51305131

diff --git a/Engine/Source/Programs/AutomationTool/AutomationUtils/ProjectParams.cs b/Engine/Source/Programs/AutomationTool/AutomationUtils/ProjectParams.cs
--- a/Engine/Source/Programs/AutomationTool/AutomationUtils/ProjectParams.cs
+++ b/Engine/Source/Programs/AutomationTool/AutomationUtils/ProjectParams.cs
@@ -468,6 +468,7 @@ namespace AutomationTool
 			this.GenerateOnDemandPakForNonChunkedBuild = InParams.GenerateOnDemandPakForNonChunkedBuild;
 			this.XcodeBuildOptions = InParams.XcodeBuildOptions;
 			this.bMakeMacNative = InParams.bMakeMacNative;
+			this.DeviceLockedTimeout = InParams.DeviceLockedTimeout;
 		}
 
 		/// <summary>
@@ -1332,6 +1333,10 @@ namespace AutomationTool
 			{
 				this.NumClients = Command.ParseParamInt("numclients");
 			}
+			if (Command != null)
+			{
+				this.DeviceLockedTimeout = Command.ParseParamInt("DeviceLockedTimeout", 120);
+			}
             if (CrashIndex.HasValue)
             {
                 this.CrashIndex = CrashIndex.Value;
@@ -2350,6 +2355,9 @@ namespace AutomationTool
 		[Help("numclients=n", "Start extra clients, n should be 2 or more")]
 		public int NumClients;
 
+		[Help("DeviceLockedTimeout=n", "Seconds to wait for a locked iOS device before failing deploy (default 120, 0 to disable)")]
+		public int DeviceLockedTimeout = 120;
+
 		/// <summary>
 		/// Run: Additional command line arguments to pass to the program
 		/// </summary>
diff --git a/Engine/Source/Programs/AutomationTool/IOS/IOSPlatform.Automation.cs b/Engine/Source/Programs/AutomationTool/IOS/IOSPlatform.Automation.cs
--- a/Engine/Source/Programs/AutomationTool/IOS/IOSPlatform.Automation.cs
+++ b/Engine/Source/Programs/AutomationTool/IOS/IOSPlatform.Automation.cs
@@ -981,6 +981,37 @@ public class IOSPlatform : ApplePlatform
 		File.Delete(Directory.GetCurrentDirectory() + "\\CommandsToPush.txt");
 	}
 
+	private static bool IsDeviceLockedError(string Output)
+	{
+		if (string.IsNullOrEmpty(Output))
+		{
+			return false;
+		}
+		return Output.Contains("for reason: Locked", StringComparison.OrdinalIgnoreCase) ||
+		       Output.Contains("could not be, unlocked", StringComparison.OrdinalIgnoreCase) ||
+		       Output.Contains("Password protected", StringComparison.OrdinalIgnoreCase);
+	}
+
+	private static bool WaitForDeviceUnlockRetry(Stopwatch Timer, int MaxWaitSeconds)
+	{
+		const int RetryIntervalSeconds = 3;
+		const int RetryIntervalMs = RetryIntervalSeconds * 1000;
+		
+		if (Timer.Elapsed.TotalSeconds >= MaxWaitSeconds)
+		{
+			Logger.LogError("[iOS] Timed out after {MaxWaitSeconds}s waiting for device unlock.", MaxWaitSeconds);
+			return false;
+		}
+
+		Logger.LogInformation("[iOS] Device is locked. Please unlock your device to continue | Retrying in {RetryTimeSeconds}s ...", RetryIntervalSeconds);
+		
+		Thread.Sleep(RetryIntervalMs);
+
+		Logger.LogInformation("[iOS] Retrying...");
+		
+		return true;
+	}
+
 	public override void Deploy(ProjectParams Params, DeploymentContext SC)
 	{
 		if (Params.Devices.Count != 1)
@@ -1089,7 +1120,35 @@ public class IOSPlatform : ApplePlatform
 		Directory.SetCurrentDirectory(CombinePaths(CmdEnv.LocalRoot, "Engine/Binaries/DotNET/IOS/"));
 		if (!Params.IterativeDeploy || bCreatedIPA || bNeedsIPA)
 		{
-			RunAndLog(CmdEnv, DeviceInstaller, LibimobileDeviceArguments);
+			if (Params.DeviceLockedTimeout > 0)
+			{
+				Stopwatch Timer = Stopwatch.StartNew();
+				int ExitCode;
+				string Output;
+				do
+				{
+					Output = RunAndLog(CmdEnv, DeviceInstaller, LibimobileDeviceArguments, out ExitCode);
+					if (ExitCode == 0)
+					{
+						break;
+					}
+					string OutputText = Output ?? string.Empty;
+					if (!IsDeviceLockedError(OutputText))
+					{
+						break;
+					}
+				}
+				while (WaitForDeviceUnlockRetry(Timer, Params.DeviceLockedTimeout));
+
+				if (ExitCode != 0)
+				{
+					throw new AutomationException("Failed to deploy to iOS device: {0}", Output);
+				}
+			}
+			else
+			{
+				RunAndLog(CmdEnv, DeviceInstaller, LibimobileDeviceArguments);
+			}
 		}
 
 		// deploy the assets
@@ -1258,8 +1317,39 @@ public class IOSPlatform : ApplePlatform
 				Arguments += " " + ClientCmdLine;
 			}
 			
-			IProcessResult ClientProcess = Run(Program, Arguments, null, ClientRunFlags);
-			
+			IProcessResult ClientProcess;
+			if (Params.DeviceLockedTimeout > 0 && !Params.DeviceNames[0].Contains(AppleToolChainSettings.LocalMacUUID) && !Params.bMakeMacNative)
+			{
+				Stopwatch Timer = Stopwatch.StartNew();
+				do
+				{
+					ClientProcess = Run(Program, Arguments, null, ClientRunFlags);
+					ClientProcess.WaitForExit();
+
+					if (ClientProcess.ExitCode == 0)
+					{
+						break;
+					}
+
+					string Output = ClientProcess.Output ?? string.Empty;
+					if (!IsDeviceLockedError(Output))
+					{
+						break;
+					}
+				}
+				while (WaitForDeviceUnlockRetry(Timer, Params.DeviceLockedTimeout));
+
+				if (ClientProcess.ExitCode != 0)
+				{
+					string FinalOutput = ClientProcess.Output ?? string.Empty;
+					throw new AutomationException("Failed to launch app on device: {0}", FinalOutput);
+				}
+			}
+			else
+			{
+				ClientProcess = Run(Program, Arguments, null, ClientRunFlags);
+			}
+
 			// This failure mode is now valid only for use of libimobiledevice on Windows.
 			if (OperatingSystem.IsWindows() && ClientProcess.ExitCode == -1)
 			{
