separate out bootstrappper details from gdk manifest generation

#rnx

#Synced-CL 51169632

diff --git a/Engine/Source/Programs/AutomationTool/Microsoft/GDKAutomationUtils.cs b/Engine/Source/Programs/AutomationTool/Microsoft/GDKAutomationUtils.cs
--- a/Engine/Source/Programs/AutomationTool/Microsoft/GDKAutomationUtils.cs
+++ b/Engine/Source/Programs/AutomationTool/Microsoft/GDKAutomationUtils.cs
@@ -35,12 +35,10 @@ namespace AutomationUtils.GDK
 		{
 			ForStaging,
 			ForPackaging,
-			ForBootstrapping,
 		}
 		public static Dictionary<UnrealTargetConfiguration, string> GetGameConfigExecutables(ProjectParams Params, DeploymentContext SC, GameConfigPurpose Purpose, string IniSection_TargetSettings, UnrealTargetConfiguration? OnlyConfig = null)
 		{
 			bool bForStaging = (Purpose == GameConfigPurpose.ForStaging);
-			bool bOnlyBootstrap = (Purpose == GameConfigPurpose.ForBootstrapping);
 
 			// build filtered stage target list
 			List<StageTarget> FilteredStageTargets = SC.StageTargets.Where(X => OnlyConfig == null || (X.Receipt.Configuration == OnlyConfig!.Value)).ToList();
@@ -56,33 +54,31 @@ namespace AutomationUtils.GDK
 			if (!Params.NoBootstrapExe && SC.StageTargetPlatform.PlatformType.IsInGroup(UnrealPlatformGroup.Windows))
 			{
 				// bootstrap naming logic from WinPlatform.Automation.cs -> GetFilesToDeployOrStage
-				Dictionary<UnrealTargetConfiguration, string> BootstrapExecutablePairs;
-				if (SC.StageTargetConfigurations.Count > 1) // note: must check unfiltered stage target here to make sure the correct bootstrap executable is referenced
+				Dictionary<UnrealTargetConfiguration, List<string>> PerConfigExecutables = GetPerConfigBootstrapExecutables(Params, SC);
+					
+				bool bAllExist = PerConfigExecutables.Any();
+				Dictionary<UnrealTargetConfiguration, string> BootstrapExecutablePairs = [];
+				foreach (KeyValuePair<UnrealTargetConfiguration, List<string>> ConfigExecutables in PerConfigExecutables)
 				{
-					BootstrapExecutablePairs = ExecutablePairs.ToDictionary(X => X.Key, X => Path.GetFileName(X.Value));
-				}
-				else if (Params.IsCodeBasedProject)
-				{
-					BootstrapExecutablePairs = ExecutablePairs.ToDictionary(X => X.Key, X => SC.StageTargets.First().Receipt.TargetName + ".exe");
-				}
-				else
-				{
-					BootstrapExecutablePairs = ExecutablePairs.ToDictionary(X => X.Key, X => SC.ShortProjectName + ".exe");
+					if (OnlyConfig != null && OnlyConfig!.Value != ConfigExecutables.Key)
+					{
+						continue;
+					}
+
+					if (ConfigExecutables.Value.Any())
+					{
+						string Executable = ConfigExecutables.Value.First();
+						BootstrapExecutablePairs.Add( ConfigExecutables.Key, Executable );
+
+						bAllExist &= bForStaging ? SC.FilesToStage.NonUFSSystemFiles.ContainsKey(new StagedFileReference(Executable)) : FileReference.Exists(FileReference.Combine(SC.StageDirectory, Executable));
+					}
 				}
 
 				// only use the bootstrap executables if they exist (to allow for development iteration scenarios). Bootstrap must exist for distribution
-				bool bAllExist = BootstrapExecutablePairs.All(X =>
-					bForStaging ? SC.FilesToStage.NonUFSSystemFiles.ContainsKey(new StagedFileReference(X.Value)) : FileReference.Exists(FileReference.Combine(SC.StageDirectory, X.Value)));
-
 				if (Params.Distribution || bAllExist)
 				{
 					ExecutablePairs = BootstrapExecutablePairs;
 				}
-				else if (bOnlyBootstrap)
-				{
-					// bootstrap executables not found
-					ExecutablePairs.Clear();
-				}
 			}
 			else if (Params.NoBootstrapExe)
 			{
@@ -97,13 +93,66 @@ namespace AutomationUtils.GDK
 					}
 				}
 			}
-			else if (bOnlyBootstrap)
+
+			return ExecutablePairs;
+		}
+
+		public static Dictionary<UnrealTargetConfiguration, List<string>> GetPerConfigBootstrapExecutables(ProjectParams Params, DeploymentContext SC)
+		{
+			Dictionary<UnrealTargetConfiguration, List<string>> PerConfigExecutables = [];
+
+			if (!Params.NoBootstrapExe && SC.StageTargetPlatform.PlatformType.IsInGroup(UnrealPlatformGroup.Windows))
 			{
-				// this configuration will never return bootstrap executables
-				ExecutablePairs.Clear();
+				// bootstrap logic from WinPlatform.Automation.cs -> GetFilesToDeployOrStage
+				// @todo: unify this logic
+
+				// one bootstrap per configuration, per target
+				foreach (UnrealTargetConfiguration Configuration in SC.StageTargetConfigurations)
+				{
+					List<string> BootstrapExecutables = [];
+
+					foreach (string TargetName in SC.StageTargets.Select( T => T.Receipt.TargetName ).Distinct())
+					{
+						BuildProduct PrimaryExecutable = null;
+						StageTarget? PrimaryTarget = null;
+
+						// collected the staged executable files for each staged architecture
+						foreach (StageTarget Target in SC.StageTargets.Where( T => T.Receipt.Configuration == Configuration && T.Receipt.TargetName == TargetName))
+						{
+							BuildProduct Executable = Target.Receipt.BuildProducts.FirstOrDefault(x => x.Type == BuildProductType.Executable);
+							if (Executable != null)
+							{
+								UnrealArch Architecture = Target.Receipt.Architectures.SingleArchitecture;
+								if (PrimaryExecutable == null || Architecture.bIsX64) // prefer x64 for the BootstrapExeName because it is undecorated
+								{
+									PrimaryExecutable = Executable;
+									PrimaryTarget = Target;
+								}
+							}
+						}
+					
+						string BootstrapExeName;
+						if(SC.StageTargetConfigurations.Count > 1)
+						{
+							BootstrapExeName = PrimaryExecutable.Path.GetFileName();
+						}
+						else if(Params.IsCodeBasedProject)
+						{
+							BootstrapExeName = PrimaryTarget.Value.Receipt.TargetName + ".exe";
+						}
+						else
+						{
+							BootstrapExeName = SC.ShortProjectName + ".exe";
+						}
+
+						BootstrapExecutables.Add(BootstrapExeName);
+					}
+
+					PerConfigExecutables.Add(Configuration, BootstrapExecutables);
+				}
 			}
 
-			return ExecutablePairs;
+			return PerConfigExecutables;
 		}
 
 		/// <summary>
diff --git a/Engine/Source/Programs/AutomationTool/Win/MSGaming/MSGameStorePackageGenerator.cs b/Engine/Source/Programs/AutomationTool/Win/MSGaming/MSGameStorePackageGenerator.cs
--- a/Engine/Source/Programs/AutomationTool/Win/MSGaming/MSGameStorePackageGenerator.cs
+++ b/Engine/Source/Programs/AutomationTool/Win/MSGaming/MSGameStorePackageGenerator.cs
@@ -45,10 +45,10 @@ namespace AutomationUtils.GDK
 			// Exclude bootstrap executables that are associated with other configurations
 			if (!Params.NoBootstrapExe && !Params.HasDLCName && OnlyConfig != null && SrcFile.HasExtension(".exe") && SrcFile.Directory == SC.StageDirectory)
 			{
-				Dictionary<UnrealTargetConfiguration, string> ExecutablePairs = GDKAutomationUtils.GetGameConfigExecutables(Params, SC, GDKAutomationUtils.GameConfigPurpose.ForBootstrapping, IniSection_TargetSettings);
-				bool bExcluded = ExecutablePairs.Any(X =>
+				Dictionary<UnrealTargetConfiguration, List<string>> BootstrapExecutables = GDKAutomationUtils.GetPerConfigBootstrapExecutables(Params, SC);
+				bool bExcluded = BootstrapExecutables.Any( X => 
 					X.Key != OnlyConfig &&
-					SrcFile.GetFileName().Equals(X.Value, StringComparison.InvariantCultureIgnoreCase)
+					X.Value.Any( F => SrcFile.GetFileName().Equals(F, StringComparison.InvariantCultureIgnoreCase))
 				);
 				if (bExcluded)
 				{
