update custom deployment & staging handlers
- deployment handler construction takes the project parameters & deployment context
- custom staging handler won't hide exceptions from TryInitialize

#rnx
#Synced-CL 51169632

diff --git a/Engine/Source/Programs/AutomationTool/AutomationUtils/CustomDeploymentHandler.cs b/Engine/Source/Programs/AutomationTool/AutomationUtils/CustomDeploymentHandler.cs
--- a/Engine/Source/Programs/AutomationTool/AutomationUtils/CustomDeploymentHandler.cs
+++ b/Engine/Source/Programs/AutomationTool/AutomationUtils/CustomDeploymentHandler.cs
@@ -12,6 +12,10 @@ using System.Linq;
 public abstract class CustomDeploymentHandler
 {
 	protected readonly Platform Platform;
+	protected CustomDeploymentHandler( Platform AutomationPlatform, ProjectParams Params, DeploymentContext SC )
+		: this(AutomationPlatform)
+	{
+	}
 	protected CustomDeploymentHandler( Platform AutomationPlatform )
 	{
 		this.Platform = AutomationPlatform;
@@ -42,7 +46,7 @@ public abstract class CustomDeploymentHandler
 	/// <summary>
 	/// Validate and instantiate the given custom deployment handler
 	/// </summary>
-	public static CustomDeploymentHandler Create(string Name, Platform AutomationPlatform)
+	public static CustomDeploymentHandler Create(string Name, Platform AutomationPlatform, ProjectParams Params, DeploymentContext SC)
 	{
 		// one-time static init. NB. Can't use the partial struct convention because automation classes are all in different assemblies
 		if (HandlerRegistry == null)
@@ -66,7 +70,7 @@ public abstract class CustomDeploymentHandler
 			throw new BuildException($"Unknown custom deployment handler {Name}");
 		}
 
-		CustomDeploymentHandler Handler = (CustomDeploymentHandler)Activator.CreateInstance(SelectedType, AutomationPlatform);
+		CustomDeploymentHandler Handler = (CustomDeploymentHandler)Activator.CreateInstance(SelectedType, AutomationPlatform, Params, SC);
 		if (Handler == null)
 		{
 			throw new BuildException($"Could not instantiate the custom deployment handler {Name}");
diff --git a/Engine/Source/Programs/AutomationTool/AutomationUtils/CustomStagingHandler.cs b/Engine/Source/Programs/AutomationTool/AutomationUtils/CustomStagingHandler.cs
--- a/Engine/Source/Programs/AutomationTool/AutomationUtils/CustomStagingHandler.cs
+++ b/Engine/Source/Programs/AutomationTool/AutomationUtils/CustomStagingHandler.cs
@@ -86,16 +86,19 @@ public abstract class CustomStagingHandler
 		List<CustomStagingHandler> AvailableHandlers = [];
 		foreach (Type HandlerType in HandlerRegistry)
 		{
+			CustomStagingHandler Handler;
 			try
 			{
-				CustomStagingHandler Handler = (CustomStagingHandler)Activator.CreateInstance(HandlerType);
-				if (Handler != null && Handler.TryInitialize(Params, SC))
-				{
-					AvailableHandlers.Add(Handler);
-				}
+				Handler = (CustomStagingHandler)Activator.CreateInstance(HandlerType);
 			}
 			catch (Exception)
 			{
+				continue;
+			}
+
+			if (Handler != null && Handler.TryInitialize(Params, SC))
+			{
+				AvailableHandlers.Add(Handler);
 			}
 		}
 
diff --git a/Engine/Source/Programs/AutomationTool/Win/MSGaming/MSGameStoreCustomDeploymentHandler.cs b/Engine/Source/Programs/AutomationTool/Win/MSGaming/MSGameStoreCustomDeploymentHandler.cs
--- a/Engine/Source/Programs/AutomationTool/Win/MSGaming/MSGameStoreCustomDeploymentHandler.cs
+++ b/Engine/Source/Programs/AutomationTool/Win/MSGaming/MSGameStoreCustomDeploymentHandler.cs
@@ -15,6 +15,8 @@ using Microsoft.Extensions.Logging;
 using static AutomationTool.CommandUtils;
 using System.Text.Json;
 
+// @note this code will soon be merged into MSGameStoreCustomStagingHandler
+
 [CustomDeploymentHandlerAttribute("MSGameStore")]
 public class MSGameStoreCustomDeploymentHandler : CustomDeploymentHandler
 {
@@ -24,8 +26,8 @@ public class MSGameStoreCustomDeploymentHandler : CustomDeploymentHandler
 
 	protected static string GDKBinariesDir => Path.Combine(GDKExports.GetGSDKRoot(), "bin");
 
-	public MSGameStoreCustomDeploymentHandler(Platform AutomationPlatform)
-		: base(AutomationPlatform)
+	public MSGameStoreCustomDeploymentHandler(Platform AutomationPlatform, ProjectParams Params, DeploymentContext SC)
+		: base(AutomationPlatform, Params, SC)
 	{
 		if (GDKEdition == 0 || !Directory.Exists(GDKBinariesDir))
 		{
diff --git a/Engine/Source/Programs/AutomationTool/Win/WinPlatform.Automation.cs b/Engine/Source/Programs/AutomationTool/Win/WinPlatform.Automation.cs
--- a/Engine/Source/Programs/AutomationTool/Win/WinPlatform.Automation.cs
+++ b/Engine/Source/Programs/AutomationTool/Win/WinPlatform.Automation.cs
@@ -456,7 +456,7 @@ public class Win64Platform : Platform
 
 		if (!string.IsNullOrEmpty(CustomDeploymentName))
 		{
-			return CustomDeploymentHandler.Create(CustomDeploymentName, this);
+			return CustomDeploymentHandler.Create(CustomDeploymentName, this, Params, SC);
 		}
 
 		return base.GetCustomDeploymentHandler(Params, SC);
