Add experimental -prestaged parameter to BuildCookRun which allows deployment and launching of an externally-produced staged build (such as one created by an automated build process, or from an earlier -archive operation)

e.g.
- copy an existing staged build to the project's staged build folder, e.g. Samples/Games/Lyra/Saved/StagedBuilds/Windows/
- RunUAT BuildCookRun -project=Lyra -platform=Win64 -prestaged -deploy -run

#rnx

#Synced-CL 51169632

diff --git a/Engine/Source/Programs/AutomationTool/AutomationUtils/ProjectParams.cs b/Engine/Source/Programs/AutomationTool/AutomationUtils/ProjectParams.cs
--- a/Engine/Source/Programs/AutomationTool/AutomationUtils/ProjectParams.cs
+++ b/Engine/Source/Programs/AutomationTool/AutomationUtils/ProjectParams.cs
@@ -381,6 +381,7 @@ namespace AutomationTool
 			this.LogWindow = InParams.LogWindow;
 			this.Stage = InParams.Stage;
 			this.SkipStage = InParams.SkipStage;
+			this.PreStaged = InParams.PreStaged;
             this.StageDirectoryParam = InParams.StageDirectoryParam;
 			this.Manifests = InParams.Manifests;
             this.CreateChunkInstall = InParams.CreateChunkInstall;
@@ -608,6 +609,7 @@ namespace AutomationTool
             bool? PrePak = null,
             bool? SkipStage = null,
 			bool? Stage = null,
+			bool? PreStaged = null,
 			bool? Manifests = null,
             bool? CreateChunkInstall = null,
 			bool? SkipEncryption = null,
@@ -908,6 +910,16 @@ namespace AutomationTool
 			{
 				this.Stage = true;
 			}
+			this.PreStaged = GetParamValueIfNotSpecified(Command, PreStaged, this.PreStaged, "prestaged");
+			if (this.PreStaged)
+			{
+				this.Build = false;
+				this.Cook = true;
+				this.SkipCook = true;
+				this.Stage = true;
+				this.SkipStage = true;
+			}
+
 			this.StageDirectoryParam = ParseParamValueIfNotSpecified(Command, StageDirectoryParam, "stagingdirectory", String.Empty, true);
 			this.OptionalFileStagingDirectory = ParseParamValueIfNotSpecified(Command, OptionalFileStagingDirectory, "optionalfilestagingdirectory", String.Empty, true);
 			this.OptionalFileInputDirectory = ParseParamValueIfNotSpecified(Command, OptionalFileInputDirectory, "optionalfileinputdirectory", String.Empty, true);
@@ -1618,6 +1630,12 @@ namespace AutomationTool
 		[Help("skipstage", "uses a stage directory, but assumes everything is already there, implies -stage")]
 		public bool SkipStage { private set; get; }
 
+		/// <summary>
+		/// Shared: true if the staged build contains a build that may not have been created locally, command line: -prestaged (experimental)
+		/// </summary>
+		[Help("prestaged", "the staged directory contains a build that wasn't necessarily created locally, implies -skipbuild -skipcook -skipstage (experimental)")]
+		public bool PreStaged { private set; get; }
+
 		/// <summary>
 		/// Shared: true if this build is using streaming install manifests, command line: -manifests
 		/// </summary>
@@ -2177,7 +2195,7 @@ namespace AutomationTool
 		/// true if the staging directory is to be cleaned: -cleanstage (also true if -clean is specified)
 		/// </summary>
 		[Help("nocleanstage", "skip cleaning the stage directory")]
-		public bool NoCleanStage { set { bNoCleanStage = value; } get { return SkipStage || bNoCleanStage; } }
+		public bool NoCleanStage { set { bNoCleanStage = value; } get { return SkipStage || PreStaged || bNoCleanStage; } }
 		private bool bNoCleanStage;
 
 		/// <summary>
@@ -2661,22 +2679,22 @@ namespace AutomationTool
 				bool bHasGameTarget = false;
 				foreach (SingleTargetProperties Target in Targets)
 				{
-					if (Target.Rules.Type == TargetType.Game)
+					if (Target.Rules.Type == TargetType.Game && !ClientCookedTargetsList.Contains(Target.TargetName))
 					{
 						ClientCookedTargetsList.Add(Target.TargetName);
 						bHasGameTarget = true;
 					}
-					else if (Target.Rules.Type == TargetType.Client)
+					else if (Target.Rules.Type == TargetType.Client && !ClientCookedTargetsList.Contains(Target.TargetName))
 					{
 						ClientCookedTargetsList.Add(Target.TargetName);
 						Client = true;
 					}
-					else if (Target.Rules.Type == TargetType.Server)
+					else if (Target.Rules.Type == TargetType.Server && !ServerCookedTargetsList.Contains(Target.TargetName))
 					{
 						ServerCookedTargetsList.Add(Target.TargetName);
 						DedicatedServer = true;
 					}
-					else
+					else if (!ProgramTargetsList.Contains(Target.TargetName))
 					{
 						ProgramTargetsList.Add(Target.TargetName);
 						ProjectType = TargetType.Program;
diff --git a/Engine/Source/Programs/AutomationTool/Scripts/CopyBuildToStagingDirectory.Automation.cs b/Engine/Source/Programs/AutomationTool/Scripts/CopyBuildToStagingDirectory.Automation.cs
--- a/Engine/Source/Programs/AutomationTool/Scripts/CopyBuildToStagingDirectory.Automation.cs
+++ b/Engine/Source/Programs/AutomationTool/Scripts/CopyBuildToStagingDirectory.Automation.cs
@@ -6406,6 +6406,42 @@ namespace AutomationScripts
 							continue;
 						}
 
+						// we need to synthesize the receipt if we are using an externally-created staged build
+						if (Params.PreStaged)
+						{
+							SingleTargetProperties TargetProperties = Params.ProjectTargets.FirstOrDefault( T => T.TargetName == Target );
+							if (TargetProperties == null)
+							{
+								throw new AutomationException("Cannot find pre-staged target properties for {0}", Target);
+							}
+
+							FileReference DefaultReceiptPath = TargetReceipt.GetDefaultPath(Params.RawProjectPath.Directory, Target, ReceiptPlatform, Config, Architecture);
+							FileReference DefaultExeName = DefaultReceiptPath.ChangeExtension( Platform.GetExeExtension(ReceiptPlatform) );
+							TargetReceipt MockReceipt = new(
+								Params.RawProjectPath,
+								Target,
+								TargetProperties.Rules.Type,
+								ReceiptPlatform,
+								Config,
+								new BuildVersion(),
+								Architecture,
+								false,
+								BuildSettingsVersion.Latest,
+								TargetBuildEnvironment.UniqueIfNeeded);
+
+							MockReceipt.Launch = DefaultExeName;
+							MockReceipt.LaunchCmd = DefaultExeName;
+							MockReceipt.AddBuildProduct(DefaultExeName, BuildProductType.Executable);
+
+							if (!string.IsNullOrEmpty(TargetProperties.Rules.CustomConfig))
+							{
+								MockReceipt.AdditionalProperties.Add(new ReceiptProperty("CustomConfig", TargetProperties.Rules.CustomConfig));
+							}
+
+							TargetsToStage.Add(new StageTarget { Receipt = MockReceipt, RequireFilesExist = true });
+							continue;
+						}
+
 						FileReference ReceiptFileName = TargetReceipt.GetDefaultPath(ReceiptBaseDir, Target, ReceiptPlatform, Config, Architecture);
 						if (!FileReference.Exists(ReceiptFileName))
 						{
@@ -6511,7 +6547,7 @@ namespace AutomationScripts
 
 		public static void CopyBuildToStagingDirectory(ProjectParams Params)
 		{
-			if (ShouldCreatePak(Params) || (Params.Stage && !Params.SkipStage))
+			if ((ShouldCreatePak(Params) || (Params.Stage && !Params.SkipStage)) && !Params.PreStaged)
 			{
 				Params.ValidateAndLog();
 
