Remove .po and .archive files to be automatically staged when a directory is requested to be staged from config.
.po and .archive files are similar to the .uasset and .umap which shouldn't be needed in the cooked build via this method.

diff --git a/Engine/Source/Programs/AutomationTool/Scripts/CopyBuildToStagingDirectory.Automation.cs b/Engine/Source/Programs/AutomationTool/Scripts/CopyBuildToStagingDirectory.Automation.cs
--- a/Engine/Source/Programs/AutomationTool/Scripts/CopyBuildToStagingDirectory.Automation.cs
+++ b/Engine/Source/Programs/AutomationTool/Scripts/CopyBuildToStagingDirectory.Automation.cs
@@ -833,7 +833,7 @@ namespace AutomationScripts
 						if (bUFS)
 						{
 							List<FileReference> Files = SC.FindFilesToStage(InputDir, StageFilesSearch.AllDirectories);
-							Files.RemoveAll(x => x.HasExtension(".uasset") || x.HasExtension(".umap") || (SC.DedicatedServer && x.HasExtension(".mp4")));
+							Files.RemoveAll(x => x.HasExtension(".uasset") || x.HasExtension(".umap") || x.HasExtension(".po") || x.HasExtension(".archive") || (SC.DedicatedServer && x.HasExtension(".mp4")));
 							SC.StageFiles(StagedFileType.UFS, InputDir, Files, OutputDir);
 						}
 						else
