Passing virtual bundles as ignored tags to ChunkDeltaOptimise

#Synced-CL 51327498

diff --git a/Engine/Source/Programs/AutomationTool/AutomationUtils/BundleUtils.cs b/Engine/Source/Programs/AutomationTool/AutomationUtils/BundleUtils.cs
--- a/Engine/Source/Programs/AutomationTool/AutomationUtils/BundleUtils.cs
+++ b/Engine/Source/Programs/AutomationTool/AutomationUtils/BundleUtils.cs
@@ -2,6 +2,8 @@
 using AutomationTool;
 using EpicGames.Core;
 using EpicGames.Horde.Projects;
+using EpicGames.MCP.Automation;
+using EpicGames.UHT.Utils;
 using System;
 using System.Collections.Generic;
 using System.Linq;
@@ -418,5 +420,49 @@ namespace AutomationUtils.Automation
 				BundleConfig.GetArray("InstallBundleManager.BundleSources", "DefaultBundleSources", out List<string> InstallBundleSources) && 
 				InstallBundleSources.Contains("Platform");
 		}
+
+		// Searches through all tagged files in a manifest, and determines virtual bundle tags from any files with multiple tags
+		//	Virtual bundles should always contain 'virtual' in their name
+		public static List<string> InferVirtualBundleTagsFromManifest(BuildPatchToolBase.ManifestMetadataOutput ManifestContents)
+		{
+			List<string> VirtualBundleTags = new List<string>();
+			if (ManifestContents != null)
+			{
+				foreach (BuildPatchToolBase.ManifestMetadataOutput.ManifestMetadataOutputFile TaggedFile in ManifestContents.Files)
+				{
+					if (TaggedFile.Tags.Count > 1)
+					{
+						foreach (string TagName in TaggedFile.Tags)
+						{
+							if (TagName.Contains("virtual", StringComparison.OrdinalIgnoreCase))
+							{
+								VirtualBundleTags.AddUnique(TagName);
+							}
+						}
+					}
+				}
+			}
+			return VirtualBundleTags;
+		}
+
+		// Loads the metadata for a given manifest and returns any virtual bundle tags inferred from its contents
+		public static List<string> InferVirtualBundleTagsFromManifest(string ManifestPath, BuildPatchToolBase.ToolVersion BPTVersion = BuildPatchToolBase.ToolVersion.Live)
+		{
+			List<string> VirtualBundleTags = new List<string>();
+			if (!string.IsNullOrEmpty(ManifestPath))
+			{
+				BuildPatchToolBase.ManifestMetadataOptions MetadataOptions = new BuildPatchToolBase.ManifestMetadataOptions
+				{
+					SourceManifestFilename = ManifestPath
+				};
+				BuildPatchToolBase.ManifestMetadataOutput ManifestMetadataOutput = null;
+				BuildPatchToolBase.Get().Execute(MetadataOptions, out ManifestMetadataOutput, BPTVersion);
+				if (ManifestMetadataOutput != null)
+				{
+					VirtualBundleTags = InferVirtualBundleTagsFromManifest(ManifestMetadataOutput);
+				}
+			}
+			return VirtualBundleTags;
+		}
 	}
 }
