Add Thin Client test node and test config

#rnx

#Synced-CL 51261596

diff --git a/Engine/Source/Programs/AutomationTool/Gauntlet/Unreal/Automation/UE.Automation.cs b/Engine/Source/Programs/AutomationTool/Gauntlet/Unreal/Automation/UE.Automation.cs
--- a/Engine/Source/Programs/AutomationTool/Gauntlet/Unreal/Automation/UE.Automation.cs
+++ b/Engine/Source/Programs/AutomationTool/Gauntlet/Unreal/Automation/UE.Automation.cs
@@ -658,8 +658,19 @@ namespace UE
 		/// </summary>
 		public override void TickTest()
 		{
-			// We are primarily interested in what the editor is doing
-			var AppInstance = TestInstance.EditorApp;
+			IAppInstance AppInstance = null;
+			
+			// If we have only one Client running without Editor
+			if (TestInstance.EditorApp == null && TestInstance.ClientApps.Length == 1)
+			{
+				AppInstance = TestInstance.ClientApps[0];
+			}
+			else
+			{
+				// We are primarily interested in what the editor is doing
+				AppInstance = TestInstance.EditorApp;
+			}
+
 			if (LogReader == null)
 			{
 				LogReader = new UnrealLogStreamParser(AppInstance.GetLogBufferReader());
