[Gauntlet] Terminate tests that fail during startup if they are multi-iteration and use device reservation blocks
This prevents multiple devices from being used

#Synced-CL 51289644

diff --git a/Engine/Source/Programs/AutomationTool/Gauntlet/Unreal/Base/Gauntlet.UnrealTestNode.cs b/Engine/Source/Programs/AutomationTool/Gauntlet/Unreal/Base/Gauntlet.UnrealTestNode.cs
--- a/Engine/Source/Programs/AutomationTool/Gauntlet/Unreal/Base/Gauntlet.UnrealTestNode.cs
+++ b/Engine/Source/Programs/AutomationTool/Gauntlet/Unreal/Base/Gauntlet.UnrealTestNode.cs
@@ -963,6 +963,11 @@ namespace Gauntlet
 			TestInstance = UnrealApp.LaunchSession();
 			if (TestInstance == null)
 			{
+				if (InNumPasses > 1 && DevicePool.DeviceReservationBlock)
+				{
+					// We throw here because we don't want different devices used for a multi-iteration test when using a reservation block
+					throw new AutomationException("Failed to start test iteration {0} when using a device reservation block. Terminating the test.", Pass);
+				}
 				return false;
 			}
 
