Update Bulk and PGO Builds to use Pak files directly

#android

#Synced-CL 51289644

#Synced-CL 51289644

diff --git a/Engine/Source/Programs/AutomationTool/Android/AndroidPlatform.Automation.cs b/Engine/Source/Programs/AutomationTool/Android/AndroidPlatform.Automation.cs
--- a/Engine/Source/Programs/AutomationTool/Android/AndroidPlatform.Automation.cs
+++ b/Engine/Source/Programs/AutomationTool/Android/AndroidPlatform.Automation.cs
@@ -1194,202 +1194,205 @@ public class AndroidPlatform : Platform
 		string LocalObbName = SC.StageDirectory.FullName+".obb";
 		string LocalPatchName = SC.StageDirectory.FullName + ".patch.obb";
 		string LocalOverflowNameTemplate = SC.StageDirectory.FullName + ".overflow{0}.obb";
-
-		FileFilter ObbFileFilter = new FileFilter(FileFilterType.Include);
-		ConfigHierarchy EngineIni = ConfigCache.ReadHierarchy(ConfigHierarchyType.Engine, DirectoryReference.FromFile(Params.RawProjectPath), UnrealTargetPlatform.Android);
-		List<string> ObbFilters;
-		EngineIni.GetArray("/Script/AndroidRuntimeSettings.AndroidRuntimeSettings", "ObbFilters", out ObbFilters);
-		if (ObbFilters != null)
+		
+		if (!Deploy.bNoObbs)
 		{
-			ObbFileFilter.AddRules(ObbFilters);
-		}
-		// Filter out dynamic libraries from obb
-		ObbFileFilter.Exclude("*.so");
+			FileFilter ObbFileFilter = new FileFilter(FileFilterType.Include);
+			ConfigHierarchy EngineIni = ConfigCache.ReadHierarchy(ConfigHierarchyType.Engine, DirectoryReference.FromFile(Params.RawProjectPath), UnrealTargetPlatform.Android);
+			List<string> ObbFilters;
+			EngineIni.GetArray("/Script/AndroidRuntimeSettings.AndroidRuntimeSettings", "ObbFilters", out ObbFilters);
+			if (ObbFilters != null)
+			{
+				ObbFileFilter.AddRules(ObbFilters);
+			}
+			// Filter out dynamic libraries from obb
+			ObbFileFilter.Exclude("*.so");
 
-		List<FileReference> FilesForObb = new List<FileReference>();
-		// Add staged Engine files
-		{
-			string EngineStageDirectoryPath = Path.Combine(SC.StageDirectory.FullName, "Engine");
-			FilesForObb.AddRange(ObbFileFilter.ApplyToDirectory(new DirectoryReference(EngineStageDirectoryPath), true));
-		}
-		// Add staged project files
-		{
-			string ProjectStageDirectoryPath = Path.Combine(SC.StageDirectory.FullName, SC.ShortProjectName);
-			FilesForObb.AddRange(ObbFileFilter.ApplyToDirectory(new DirectoryReference(ProjectStageDirectoryPath), true));
-		}
+			List<FileReference> FilesForObb = new List<FileReference>();
+			// Add staged Engine files
+			{
+				string EngineStageDirectoryPath = Path.Combine(SC.StageDirectory.FullName, "Engine");
+				FilesForObb.AddRange(ObbFileFilter.ApplyToDirectory(new DirectoryReference(EngineStageDirectoryPath), true));
+			}
+			// Add staged project files
+			{
+				string ProjectStageDirectoryPath = Path.Combine(SC.StageDirectory.FullName, SC.ShortProjectName);
+				FilesForObb.AddRange(ObbFileFilter.ApplyToDirectory(new DirectoryReference(ProjectStageDirectoryPath), true));
+			}
 
-		bool OBBNeedsUpdate = false;
+			bool OBBNeedsUpdate = false;
 
-		if (File.Exists(LocalObbName))
-		{ 
-			System.DateTime OBBTimeStamp = File.GetLastWriteTimeUtc(LocalObbName);
-			foreach (FileReference FileToObb in FilesForObb)
-			{
-				System.DateTime FileTimeStamp = File.GetLastWriteTimeUtc(FileToObb.FullName);
-				if (FileTimeStamp > OBBTimeStamp)
+			if (File.Exists(LocalObbName))
+			{ 
+				System.DateTime OBBTimeStamp = File.GetLastWriteTimeUtc(LocalObbName);
+				foreach (FileReference FileToObb in FilesForObb)
 				{
-					OBBNeedsUpdate = true;
-					break;
+					System.DateTime FileTimeStamp = File.GetLastWriteTimeUtc(FileToObb.FullName);
+					if (FileTimeStamp > OBBTimeStamp)
+					{
+						OBBNeedsUpdate = true;
+						break;
+					}
 				}
 			}
-		}
-		else
-		{
-			OBBNeedsUpdate = true;
-		}
-		Int64 OBBSizeAllowed = GetMaxOBBSizeAllowed(SC);
-		string LimitString = (OBBSizeAllowed < NormalOBBSizeAllowed) ? "1 GiB" : ((OBBSizeAllowed < MaxOBBSizeAllowed) ? "2 GiB" : "4 GiB");
-
-		if (!OBBNeedsUpdate)
-		{
-			Logger.LogInformation("{Text}", "OBB is up to date: " + LocalObbName);
-		}
-		else
-		{
-			// Always delete the target OBB file if it exists
-			if (File.Exists(LocalObbName))
+			else
 			{
-				File.Delete(LocalObbName);
+				OBBNeedsUpdate = true;
 			}
+			Int64 OBBSizeAllowed = GetMaxOBBSizeAllowed(SC);
+			string LimitString = (OBBSizeAllowed < NormalOBBSizeAllowed) ? "1 GiB" : ((OBBSizeAllowed < MaxOBBSizeAllowed) ? "2 GiB" : "4 GiB");
 
-			// Always delete the target patch OBB file if it exists
-			if (File.Exists(LocalPatchName))
+			if (!OBBNeedsUpdate)
 			{
-				File.Delete(LocalPatchName);
+				Logger.LogInformation("{Text}", "OBB is up to date: " + LocalObbName);
 			}
-
-			// Always delete all target overflow OBB files if they exists
-			int OverflowIndex = 1;
-			while (File.Exists(string.Format(LocalOverflowNameTemplate, OverflowIndex)))
+			else
 			{
-				File.Delete(string.Format(LocalOverflowNameTemplate, OverflowIndex));
-				// move to next overflow
-				OverflowIndex++;
-			}
-
-			List<FileReference> FilesToObb = FilesForObb;
-			List<FileReference> FilesToPatch = new List<FileReference>();
-			List<OverflowFileInfo> OverflowInfos = new List<OverflowFileInfo>();
+				// Always delete the target OBB file if it exists
+				if (File.Exists(LocalObbName))
+				{
+					File.Delete(LocalObbName);
+				}
 
-			if (AllowPatchOBBFile(SC))
-			{
-				int AllowOverflowOBBLimit = AllowOverflowOBBFiles(SC);
+				// Always delete the target patch OBB file if it exists
+				if (File.Exists(LocalPatchName))
+				{
+					File.Delete(LocalPatchName);
+				}
 
-				FilesToObb = new List<FileReference>();
+				// Always delete all target overflow OBB files if they exists
+				int OverflowIndex = 1;
+				while (File.Exists(string.Format(LocalOverflowNameTemplate, OverflowIndex)))
+				{
+					File.Delete(string.Format(LocalOverflowNameTemplate, OverflowIndex));
+					// move to next overflow
+					OverflowIndex++;
+				}
 
-				// Collect the filesize and place into Obb or Patch list
-				Int64 StagingDirLength = SC.StageDirectory.FullName.Length;
-				Int64 MinimumObbSize = 22 + 10;		// EOCD with comment (store version)
-				Int64 MainObbSize = MinimumObbSize;
-				Int64 PatchObbSize = MinimumObbSize;
+				List<FileReference> FilesToObb = FilesForObb;
+				List<FileReference> FilesToPatch = new List<FileReference>();
+				List<OverflowFileInfo> OverflowInfos = new List<OverflowFileInfo>();
 
-				foreach (FileReference FileRef in FilesForObb)
+				if (AllowPatchOBBFile(SC))
 				{
-					FileInfo LocalFileInfo = new FileInfo(FileRef.FullName);
-					Int64 LocalFileLength = LocalFileInfo.Length;
-					Int64 FilenameLength = FileRef.FullName.Length - StagingDirLength - 1;
+					int AllowOverflowOBBLimit = AllowOverflowOBBFiles(SC);
 
-					Int64 LocalOverhead = (30 + FilenameLength + 36);       // local file descriptor
-					Int64 GlobalOverhead = (46 + FilenameLength + 36);      // central directory cost
-					Int64 FileRequirements = LocalFileLength + LocalOverhead + GlobalOverhead;
+					FilesToObb = new List<FileReference>();
 
-					if (MainObbSize + FileRequirements < OBBSizeAllowed)
-					{
-						FilesToObb.Add(FileRef);
-						MainObbSize += FileRequirements;
-					}
-					else if (PatchObbSize + FileRequirements < OBBSizeAllowed)
-					{
-						FilesToPatch.Add(FileRef);
-						PatchObbSize += FileRequirements;
-					}
-					else if (AllowOverflowOBBLimit > 0)
+					// Collect the filesize and place into Obb or Patch list
+					Int64 StagingDirLength = SC.StageDirectory.FullName.Length;
+					Int64 MinimumObbSize = 22 + 10;		// EOCD with comment (store version)
+					Int64 MainObbSize = MinimumObbSize;
+					Int64 PatchObbSize = MinimumObbSize;
+
+					foreach (FileReference FileRef in FilesForObb)
 					{
-						// find an overflow we can fit this file
-						bool bFoundOverflow = false;				
-						foreach (OverflowFileInfo OverflowRef in OverflowInfos)
+						FileInfo LocalFileInfo = new FileInfo(FileRef.FullName);
+						Int64 LocalFileLength = LocalFileInfo.Length;
+						Int64 FilenameLength = FileRef.FullName.Length - StagingDirLength - 1;
+
+						Int64 LocalOverhead = (30 + FilenameLength + 36);       // local file descriptor
+						Int64 GlobalOverhead = (46 + FilenameLength + 36);      // central directory cost
+						Int64 FileRequirements = LocalFileLength + LocalOverhead + GlobalOverhead;
+
+						if (MainObbSize + FileRequirements < OBBSizeAllowed)
+						{
+							FilesToObb.Add(FileRef);
+							MainObbSize += FileRequirements;
+						}
+						else if (PatchObbSize + FileRequirements < OBBSizeAllowed)
+						{
+							FilesToPatch.Add(FileRef);
+							PatchObbSize += FileRequirements;
+						}
+						else if (AllowOverflowOBBLimit > 0)
 						{
-							if (OverflowRef.OverflowObbSize + FileRequirements < OBBSizeAllowed)
+							// find an overflow we can fit this file
+							bool bFoundOverflow = false;				
+							foreach (OverflowFileInfo OverflowRef in OverflowInfos)
 							{
-								OverflowRef.FilesForOverflow.Add(FileRef);
-								OverflowRef.OverflowObbSize += FileRequirements;
-								bFoundOverflow = true;
+								if (OverflowRef.OverflowObbSize + FileRequirements < OBBSizeAllowed)
+								{
+									OverflowRef.FilesForOverflow.Add(FileRef);
+									OverflowRef.OverflowObbSize += FileRequirements;
+									bFoundOverflow = true;
+								}
 							}
-						}
 
-						if (!bFoundOverflow)
-						{
-							if (AllowOverflowOBBLimit >= OverflowInfos.Count)
+							if (!bFoundOverflow)
 							{
-								// create a new overflow obb
-								if (MinimumObbSize + FileRequirements < OBBSizeAllowed)
+								if (AllowOverflowOBBLimit >= OverflowInfos.Count)
 								{
-									List<FileReference> NewObb = new List<FileReference>();
-									NewObb.Add(FileRef);
-									OverflowInfos.Add(new OverflowFileInfo(NewObb, MinimumObbSize + FileRequirements));
+									// create a new overflow obb
+									if (MinimumObbSize + FileRequirements < OBBSizeAllowed)
+									{
+										List<FileReference> NewObb = new List<FileReference>();
+										NewObb.Add(FileRef);
+										OverflowInfos.Add(new OverflowFileInfo(NewObb, MinimumObbSize + FileRequirements));
+									}
+									else
+									{
+										Logger.LogInformation("{Text}", "Failed to add " + FileRef.FullName + " to a new overflow as it is bigger than the allowed OBB size ");
+										throw new AutomationException(ExitCode.Error_AndroidOBBError, "Stage Failed. Could not add {1} ({0} bytes) to an OBB as OBBs are limited to {2} bytes.", FileRequirements, FileRef.FullName, OBBSizeAllowed);
+									}
 								}
 								else
 								{
-									Logger.LogInformation("{Text}", "Failed to add " + FileRef.FullName + " to a new overflow as it is bigger than the allowed OBB size ");
-									throw new AutomationException(ExitCode.Error_AndroidOBBError, "Stage Failed. Could not add {1} ({0} bytes) to an OBB as OBBs are limited to {2} bytes.", FileRequirements, FileRef.FullName, OBBSizeAllowed);
+									Logger.LogInformation("{Text}", "Failed to add required overflow OBB: " + LocalObbName);
+									throw new AutomationException(ExitCode.Error_AndroidOBBError, "Stage Failed. Overflow OBBs limited to a count of {0}. Contents are to big to fit.", OverflowInfos.Count);
 								}
 							}
-							else
-							{
-								Logger.LogInformation("{Text}", "Failed to add required overflow OBB: " + LocalObbName);
-								throw new AutomationException(ExitCode.Error_AndroidOBBError, "Stage Failed. Overflow OBBs limited to a count of {0}. Contents are to big to fit.", OverflowInfos.Count);
-							}
 						}
-					}
-					else
-					{
-						// no room in either file and no overflows allowed
-						Logger.LogInformation("{Text}", "Failed to build OBB: " + LocalObbName);
-						throw new AutomationException(ExitCode.Error_AndroidOBBError, "Stage Failed. Could not build OBB {0}. The file may be too big to fit in an OBB ({1} limit and no overflows are permitted)", LocalObbName, LimitString);
+						else
+						{
+							// no room in either file and no overflows allowed
+							Logger.LogInformation("{Text}", "Failed to build OBB: " + LocalObbName);
+							throw new AutomationException(ExitCode.Error_AndroidOBBError, "Stage Failed. Could not build OBB {0}. The file may be too big to fit in an OBB ({1} limit and no overflows are permitted)", LocalObbName, LimitString);
+						}
 					}
 				}
-			}
 
-			// Now create the main OBB as a ZIP archive.
-			if (!CreateOBBFile(SC, LocalObbName, FilesToObb))
-			{
-				Logger.LogInformation("{Text}", "Failed to build OBB: " + LocalObbName);
-				throw new AutomationException(ExitCode.Error_AndroidOBBError, "Stage Failed. Could not build OBB {0}. The file may be too big to fit in an OBB ({1} limit)", LocalObbName, LimitString);
-			}
+				// Now create the main OBB as a ZIP archive.
+				if (!CreateOBBFile(SC, LocalObbName, FilesToObb))
+				{
+					Logger.LogInformation("{Text}", "Failed to build OBB: " + LocalObbName);
+					throw new AutomationException(ExitCode.Error_AndroidOBBError, "Stage Failed. Could not build OBB {0}. The file may be too big to fit in an OBB ({1} limit)", LocalObbName, LimitString);
+				}
 
-			// Now create the patch OBB as a ZIP archive if required.
-			if (FilesToPatch.Count() > 0)
-			{
-				if (!CreateOBBFile(SC, LocalPatchName, FilesToPatch))
+				// Now create the patch OBB as a ZIP archive if required.
+				if (FilesToPatch.Count() > 0)
 				{
-					Logger.LogInformation("{Text}", "Failed to build OBB: " + LocalPatchName);
-					throw new AutomationException(ExitCode.Error_AndroidOBBError, "Stage Failed. Could not build OBB {0}. The file may be too big to fit in an OBB ({1} limit)", LocalPatchName, LimitString);
+					if (!CreateOBBFile(SC, LocalPatchName, FilesToPatch))
+					{
+						Logger.LogInformation("{Text}", "Failed to build OBB: " + LocalPatchName);
+						throw new AutomationException(ExitCode.Error_AndroidOBBError, "Stage Failed. Could not build OBB {0}. The file may be too big to fit in an OBB ({1} limit)", LocalPatchName, LimitString);
+					}
 				}
-			}
 
-			OverflowIndex = 1;
-			foreach (OverflowFileInfo OverflowRef in OverflowInfos)
-			{
-				if (OverflowRef.FilesForOverflow.Count() > 0)
+				OverflowIndex = 1;
+				foreach (OverflowFileInfo OverflowRef in OverflowInfos)
 				{
-					string LocalOverflowName = string.Format(LocalOverflowNameTemplate, OverflowIndex);
-					if (!CreateOBBFile(SC, LocalOverflowName, OverflowRef.FilesForOverflow))
+					if (OverflowRef.FilesForOverflow.Count() > 0)
 					{
-						Logger.LogInformation("{Text}", "Failed to build OBB: " + LocalOverflowName);
-						throw new AutomationException(ExitCode.Error_AndroidOBBError, "Stage Failed. Could not build OBB {0}. The file may be too big to fit in an OBB ({1} limit)", LocalOverflowName, LimitString);
+						string LocalOverflowName = string.Format(LocalOverflowNameTemplate, OverflowIndex);
+						if (!CreateOBBFile(SC, LocalOverflowName, OverflowRef.FilesForOverflow))
+						{
+							Logger.LogInformation("{Text}", "Failed to build OBB: " + LocalOverflowName);
+							throw new AutomationException(ExitCode.Error_AndroidOBBError, "Stage Failed. Could not build OBB {0}. The file may be too big to fit in an OBB ({1} limit)", LocalOverflowName, LimitString);
+						}
 					}
+					OverflowIndex++;
 				}
-				OverflowIndex++;
 			}
-		}
 
-		// make sure the OBB is <= 2GiB (or 4GiB if large OBB enabled)
-		FileInfo OBBFileInfo = new FileInfo(LocalObbName);
-		Int64 ObbFileLength = OBBFileInfo.Length;
-		if (ObbFileLength > OBBSizeAllowed)
-		{
-			Logger.LogInformation("{Text}", "OBB exceeds " + LimitString + " limit: " + ObbFileLength + " bytes");
-			throw new AutomationException(ExitCode.Error_AndroidOBBError, "Stage Failed. OBB {0} exceeds {1} limit)", LocalObbName, LimitString);
+			// make sure the OBB is <= 2GiB (or 4GiB if large OBB enabled)
+			FileInfo OBBFileInfo = new FileInfo(LocalObbName);
+			Int64 ObbFileLength = OBBFileInfo.Length;
+			if (ObbFileLength > OBBSizeAllowed)
+			{
+				Logger.LogInformation("{Text}", "OBB exceeds " + LimitString + " limit: " + ObbFileLength + " bytes");
+				throw new AutomationException(ExitCode.Error_AndroidOBBError, "Stage Failed. OBB {0} exceeds {1} limit)", LocalObbName, LimitString);
+			}
 		}
 
 		// collect plugin extra data paths from target receipts
@@ -1461,7 +1464,7 @@ public class AndroidPlatform : Platform
 
 			List<OverflowBatchInstallInfo> OverflowInfos = new List<OverflowBatchInstallInfo>();
 			
-			if (!bPackageDataInsideApk)
+			if (!Deploy.bNoObbs && !bPackageDataInsideApk)
 			{
 				DeviceObbName = GetDeviceObbName(ApkName, SC);
 				ObbName = GetFinalObbName(ApkName, SC);
@@ -1534,7 +1537,7 @@ public class AndroidPlatform : Platform
 				string BatchName = GetFinalBatchName(ApkName, SC, false, EBatchType.Install, Target);
 				List<string> InstallBatchLines = GenerateInstallBatchFile(bPackageDataInsideApk, PackageName, ApkName, Params, ObbName, DeviceObbName, false, PatchName, DevicePatchName, false, OverflowInfos,
 					bIsPC, Params.Distribution, TargetSDKVersion > 22, bDisablePerfHarden, bUseAFS, bUseAFSProject, AFSToken, Target);
-				if (bHaveAPK)
+				if (!Deploy.bNoObbs && bHaveAPK)
 				{
 					// make a batch file that can be used to install the .apk and .obb files
 					File.WriteAllText(BatchName, string.Join(LineEnding, InstallBatchLines) + LineEnding);
@@ -1567,7 +1570,7 @@ public class AndroidPlatform : Platform
 
 				if (!OperatingSystem.IsWindows())
 				{
-					if (bHaveAPK)
+					if (!Deploy.bNoObbs && bHaveAPK)
 					{
 						CommandUtils.FixUnixFilePermissions(BatchName);
 					}
@@ -2060,7 +2063,8 @@ public class AndroidPlatform : Platform
 		UnrealTargetConfiguration TargetConfiguration = SC.StageTargetConfigurations[0];
 		UnrealArchitectures Architectures = GetDeploymentArchitectures(Params, SC);
 		bool bMakeSeparateApks = UnrealBuildTool.AndroidExports.ShouldMakeSeparateApks();
-		bool bPackageDataInsideApk = UnrealBuildTool.AndroidExports.CreateDeploymentHandler(Params.RawProjectPath, Params.ForcePackageData).GetPackageDataInsideApk();
+		IAndroidDeploy Deploy = UnrealBuildTool.AndroidExports.CreateDeploymentHandler(Params.RawProjectPath, Params.ForcePackageData);
+		bool bPackageDataInsideApk = Deploy.GetPackageDataInsideApk();
 		bool bBundleEnabled = GetEnableBundle(SC);
 
 		bool bAFSEnablePlugin;
@@ -2173,7 +2177,7 @@ public class AndroidPlatform : Platform
 				SC.ArchiveFiles(Path.GetDirectoryName(ApkName), Path.GetFileName(ApkName));
 			}
 
-			if (!bPackageDataInsideApk && !FileExists(ObbName))
+			if (!Deploy.bNoObbs && !bPackageDataInsideApk && !FileExists(ObbName))
 			{
                 throw new AutomationException(ExitCode.Error_ObbNotFound, "ARCHIVE FAILED - {0} was not found", ObbName);
 			}
@@ -2199,7 +2203,7 @@ public class AndroidPlatform : Platform
 				}
 			}
 
-			if (!bPackageDataInsideApk)
+			if (!Deploy.bNoObbs && !bPackageDataInsideApk)
 			{
 				// only add if not already in archive list
 				if (!AddedObbFiles.Contains(ObbName))
diff --git a/Engine/Source/Programs/UnrealBuildTool/Platform/Android/AndroidExports.cs b/Engine/Source/Programs/UnrealBuildTool/Platform/Android/AndroidExports.cs
--- a/Engine/Source/Programs/UnrealBuildTool/Platform/Android/AndroidExports.cs
+++ b/Engine/Source/Programs/UnrealBuildTool/Platform/Android/AndroidExports.cs
@@ -39,6 +39,11 @@ namespace UnrealBuildTool
 		/// </summary>
 		bool bFastIterate { get; }
 		
+		/// <summary>
+		/// 
+		/// </summary>
+		bool bNoObbs { get; }
+		
 		/// <summary>
 		/// 
 		/// </summary>
diff --git a/Engine/Source/Programs/UnrealBuildTool/Platform/Android/UEDeployAndroid.cs b/Engine/Source/Programs/UnrealBuildTool/Platform/Android/UEDeployAndroid.cs
--- a/Engine/Source/Programs/UnrealBuildTool/Platform/Android/UEDeployAndroid.cs
+++ b/Engine/Source/Programs/UnrealBuildTool/Platform/Android/UEDeployAndroid.cs
@@ -111,6 +111,13 @@ namespace UnrealBuildTool
 		[CommandLine("-OnlyGradleGeneration", Value = "true")]
 		public bool bOnlyGradleGeneration = false;
 
+		/// <summary>
+		/// Won't generate or use OBBs if set
+		/// </summary>
+		[CommandLine("-NoObbs", Value = "true")]
+		public bool bNoObbs = false;
+		bool IAndroidDeploy.bNoObbs => bNoObbs;
+
 		public UEDeployAndroid(FileReference? InProjectFile, bool InForcePackageData, ILogger InLogger)
 			: base(InLogger)
 		{
@@ -2014,6 +2021,7 @@ LD_HWASAN=1 exec $cmd";
 			CurrentSettings.AppendLine(String.Format("bEnableBundle={0}", bEnableBundle));
 			CurrentSettings.AppendLine(String.Format("bDontBundleLibrariesInAPK={0}", bDontBundleLibrariesInAPK));
 			CurrentSettings.AppendLine(String.Format("bUsesZenStreaming={0}", bUsesZenStreaming));
+			CurrentSettings.AppendLine(String.Format("bNoObbs={0}", bNoObbs));
 
 			// all AndroidRuntimeSettings ini settings in here
 			ConfigHierarchy Ini = GetConfigCacheIni(ConfigHierarchyType.Engine);
@@ -2186,7 +2194,7 @@ LD_HWASAN=1 exec $cmd";
 				InputFiles.AddRange(SettingsFiles);
 
 				// rebuild if .pak files exist for OBB in APK case
-				if (bPackageDataInsideApk)
+				if (!bNoObbs && bPackageDataInsideApk)
 				{
 					string PAKFileLocation = ProjectDirectory + "/Saved/StagedBuilds/Android" + CookFlavor + "/" + ProjectName + "/Content/Paks";
 					if (Directory.Exists(PAKFileLocation))
@@ -3375,7 +3383,12 @@ LD_HWASAN=1 exec $cmd";
 
 		private bool RequiresOBB(bool bDisallowPackageInAPK, string OBBLocation)
 		{
-			if (bDisallowPackageInAPK)
+			if (bNoObbs)
+			{
+				Logger.LogInformation("OBBs have been disabled.");
+				return false;
+			}
+			else if (bDisallowPackageInAPK)
 			{
 				Logger.LogInformation("APK contains data.");
 				return false;
@@ -4348,7 +4361,7 @@ dependencies.implementation('androidx.tracing:tracing:1.2.0')
 
 			// cache if we want data in the Apk
 			bool bPackageDataInsideApk = !bDisallowPackagingDataInApk && GetPackageDataInsideApk();
-			bool bDisableVerifyOBBOnStartUp = DisableVerifyOBBOnStartUp();
+			bool bDisableVerifyOBBOnStartUp = bNoObbs || DisableVerifyOBBOnStartUp();
 			bool bUseExternalFilesDir = UseExternalFilesDir(bDisallowExternalFilesDir);
 
 			// Generate Java files
@@ -4370,15 +4383,15 @@ dependencies.implementation('androidx.tracing:tracing:1.2.0')
 			// Generate the OBB and Shim files here
 			string ObbFileLocation = ProjectDirectory + "/Saved/StagedBuilds/Android" + CookFlavor + ".obb";
 			string PatchFileLocation = ProjectDirectory + "/Saved/StagedBuilds/Android" + CookFlavor + ".patch.obb";
-			List<string> RequiredOBBFiles = new List<string> { ObbFileLocation };
-			if (File.Exists(PatchFileLocation))
+			List<string> RequiredOBBFiles = bNoObbs ? [] : [ObbFileLocation];
+			if (!bNoObbs && File.Exists(PatchFileLocation))
 			{
 				RequiredOBBFiles.Add(PatchFileLocation);
 			}
 
 			// Generate the OBBData.java file if out of date (can skip rewriting it if packaging inside Apk in some cases)
 			// Note: this may be replaced per architecture later if store version is different
-			WriteJavaOBBDataFile(UnrealOBBDataFileName, PackageName, RequiredOBBFiles, CookFlavor, bPackageDataInsideApk, Architectures.Architectures[0]);
+			WriteJavaOBBDataFile(UnrealOBBDataFileName, PackageName, RequiredOBBFiles, CookFlavor, !bNoObbs && bPackageDataInsideApk, Architectures.Architectures[0]);
 
 			// Make sure any existing proguard file in project is NOT used (back it up)
 			string ProjectBuildProguardFile = Path.Combine(GameBuildFilesPath, "proguard-project.txt");
@@ -4888,28 +4901,31 @@ dependencies.implementation('androidx.tracing:tracing:1.2.0')
 				string UnrealBuildGradlePath = Path.Combine(UnrealBuildPath, "gradle");
 
 				bool bInApkAssetChange = false;
-				// If we are packaging for Amazon then we need to copy the  file to the correct location
-				Logger.LogInformation("bPackageDataInsideApk = {PackageDataInsideApk}", bPackageDataInsideApk);
-				if (bPackageDataInsideApk)
+				if (!bNoObbs)
 				{
-					Logger.LogInformation("Obb location {ObbFileLocation}", ObbFileLocation);
-					string ObbFileDestination = UnrealBuildPath + "/assets";
-					Logger.LogInformation("Obb destination location {ObbFileDestination}", ObbFileDestination);
-					if (File.Exists(ObbFileLocation))
+					// If we are packaging for Amazon then we need to copy the  file to the correct location
+					Logger.LogInformation("bPackageDataInsideApk = {PackageDataInsideApk}", bPackageDataInsideApk);
+					if (bPackageDataInsideApk)
 					{
-						Directory.CreateDirectory(UnrealBuildPath);
-						Directory.CreateDirectory(ObbFileDestination);
-						Logger.LogInformation("Obb file exists...");
-						string DestFileName = Path.Combine(ObbFileDestination, "main.obb.png"); // Need a rename to turn off compression
-						string SrcFileName = ObbFileLocation;
-						bInApkAssetChange = CopyIfDifferent(SrcFileName, DestFileName, true, false);
+						Logger.LogInformation("Obb location {ObbFileLocation}", ObbFileLocation);
+						string ObbFileDestination = UnrealBuildPath + "/assets";
+						Logger.LogInformation("Obb destination location {ObbFileDestination}", ObbFileDestination);
+						if (File.Exists(ObbFileLocation))
+						{
+							Directory.CreateDirectory(UnrealBuildPath);
+							Directory.CreateDirectory(ObbFileDestination);
+							Logger.LogInformation("Obb file exists...");
+							string DestFileName = Path.Combine(ObbFileDestination, "main.obb.png"); // Need a rename to turn off compression
+							string SrcFileName = ObbFileLocation;
+							bInApkAssetChange = CopyIfDifferent(SrcFileName, DestFileName, true, false);
+						}
+					}
+					else // try to remove the file it we aren't packaging inside the APK
+					{
+						string ObbFileDestination = UnrealBuildPath + "/assets";
+						string DestFileName = Path.Combine(ObbFileDestination, "main.obb.png");
+						bInApkAssetChange = SafeDeleteFile(DestFileName);
 					}
-				}
-				else // try to remove the file it we aren't packaging inside the APK
-				{
-					string ObbFileDestination = UnrealBuildPath + "/assets";
-					string DestFileName = Path.Combine(ObbFileDestination, "main.obb.png");
-					bInApkAssetChange = SafeDeleteFile(DestFileName);
 				}
 
 				if (bInApkAssetChange)
@@ -4963,6 +4979,12 @@ dependencies.implementation('androidx.tracing:tracing:1.2.0')
 				CopyFileDirectory(GameBuildFilesPath_NR, UnrealBuildPath, Replacements, ExcludeFolders);
 				CopyFileDirectory(GameBuildFilesPath_LA, UnrealBuildPath, Replacements, ExcludeFolders);
 
+				if (bNoObbs || !bPackageDataInsideApk)
+				{
+					// Try to remove main.obb.png file if we don't want OBBs
+					SafeDeleteFile(Path.Combine(UnrealBuildPath, "assets", "main.obb.png"));
+				}
+				
 				// Parse Gradle filters (may have been replaced by above copies)
 				ParseFilterFile(Path.Combine(UnrealBuildPath, "GradleFilter.txt"));
 
