Cloud Storage - Added retries for http timeouts (exit code 84)

#Synced-CL 51251964

diff --git a/Engine/Source/Programs/AutomationTool/BuildGraph/Tasks/CreateCloudArtifactTask.cs b/Engine/Source/Programs/AutomationTool/BuildGraph/Tasks/CreateCloudArtifactTask.cs
--- a/Engine/Source/Programs/AutomationTool/BuildGraph/Tasks/CreateCloudArtifactTask.cs
+++ b/Engine/Source/Programs/AutomationTool/BuildGraph/Tasks/CreateCloudArtifactTask.cs
@@ -549,7 +549,7 @@ namespace AutomationTool.Tasks
 					// kHttpBadGateway          = 92,	// BadGateway(502)
 					// kHttpGatewayTimeout      = 93,	// GatewayTimeout(504)
 
-					if (exitCode is 502 or 504 or 69 or 92 or 93)
+					if (exitCode is 502 or 504 or 69 or 92 or 93 or 84)
 					{
 						// zen returns http status code as exit codes if it had an issue
 						Logger.LogInformation("Http Status code {ExitCode} returned from zen. Waiting {WaitDuration}s before trying.", exitCode, waitDuration);
diff --git a/Engine/Source/Programs/AutomationTool/BuildGraph/Tasks/ZenExportSnapshotTask.cs b/Engine/Source/Programs/AutomationTool/BuildGraph/Tasks/ZenExportSnapshotTask.cs
--- a/Engine/Source/Programs/AutomationTool/BuildGraph/Tasks/ZenExportSnapshotTask.cs
+++ b/Engine/Source/Programs/AutomationTool/BuildGraph/Tasks/ZenExportSnapshotTask.cs
@@ -453,7 +453,7 @@ namespace AutomationTool.Tasks
 				}
 
 				// 502 and 504 are the old exit codes for retries the new exit codes for retries that also work on mac and linux are 69, 92 and 93 see zen cli documentation of exit codes for more details
-				bool isGatewayError = exitCode is 502 or 504 or 69 or 92 or 93;
+				bool isGatewayError = exitCode is 502 or 504 or 69 or 92 or 93 or 84;
 				int waitDuration = isGatewayError ? waitDurationWebSeconds[attempt] : 0;
 				bool willRetry = attempt < (attemptLimit - 1);
 				Logger.LogWarning("Attempt {AttemptNum} of exporting the oplog failed, {Action}...", attempt + 1, willRetry ? $"retrying after {waitDuration}s" : "abandoning");
diff --git a/Engine/Source/Programs/Shared/EpicGames.ProjectStore/CloudStorage/CloudStorage.cs b/Engine/Source/Programs/Shared/EpicGames.ProjectStore/CloudStorage/CloudStorage.cs
--- a/Engine/Source/Programs/Shared/EpicGames.ProjectStore/CloudStorage/CloudStorage.cs
+++ b/Engine/Source/Programs/Shared/EpicGames.ProjectStore/CloudStorage/CloudStorage.cs
@@ -582,7 +582,7 @@ namespace EpicGames.ProjectStore
 
 					ZenExitCodes zenExitCode = AsZenExitCodeValue(exitCode);
 					// retry on zen exit codes that indicate transient issues
-					if (zenExitCode is ZenExitCodes.HttpServiceUnavailable or ZenExitCodes.HttpBadGateway or ZenExitCodes.HttpGatewayTimeout)
+					if (zenExitCode is ZenExitCodes.HttpServiceUnavailable or ZenExitCodes.HttpBadGateway or ZenExitCodes.HttpGatewayTimeout or ZenExitCodes.HttpTimeout)
 					{
 						if (retryAttempt + 1 == maxRetryAttempts)
 						{
