Replace var with type

#rnx
#Synced-CL 51479624

#Synced-CL 51479624

diff --git a/Engine/Source/Programs/AutomationTool/AutomationUtils/Automation.cs b/Engine/Source/Programs/AutomationTool/AutomationUtils/Automation.cs
--- a/Engine/Source/Programs/AutomationTool/AutomationUtils/Automation.cs
+++ b/Engine/Source/Programs/AutomationTool/AutomationUtils/Automation.cs
@@ -325,7 +325,7 @@ namespace AutomationTool
 					finally
 					{
 						// dispose of the class if necessary
-						var CommandDisposable = Command as IDisposable;
+						IDisposable CommandDisposable = Command as IDisposable;
 						if (CommandDisposable != null)
 						{
 							CommandDisposable.Dispose();
diff --git a/Engine/Source/Programs/AutomationTool/AutomationUtils/BuildCommand.cs b/Engine/Source/Programs/AutomationTool/AutomationUtils/BuildCommand.cs
--- a/Engine/Source/Programs/AutomationTool/AutomationUtils/BuildCommand.cs
+++ b/Engine/Source/Programs/AutomationTool/AutomationUtils/BuildCommand.cs
@@ -299,7 +299,7 @@ namespace AutomationTool
 			}
 			else
 			{
-				var Branch = new BranchInfo();
+				BranchInfo Branch = new BranchInfo();
 				var GameProj = Branch.FindGame(originalProjectName);
 				if (GameProj != null)
 				{
diff --git a/Engine/Source/Programs/AutomationTool/AutomationUtils/BundleUtils.cs b/Engine/Source/Programs/AutomationTool/AutomationUtils/BundleUtils.cs
--- a/Engine/Source/Programs/AutomationTool/AutomationUtils/BundleUtils.cs
+++ b/Engine/Source/Programs/AutomationTool/AutomationUtils/BundleUtils.cs
@@ -63,7 +63,7 @@ namespace AutomationUtils.Automation
 			Action<TPlatformBundleSettings, ConfigHierarchy, string> GetPlatformSettings) 
 			where TPlatformBundleSettings : BundleSettings, new()
 		{
-			var Results = new List<TPlatformBundleSettings>();
+			List<TPlatformBundleSettings> Results = new List<TPlatformBundleSettings>();
 
 			const string BundleDefinitionPrefix = "InstallBundleDefinition ";
 
@@ -271,7 +271,7 @@ namespace AutomationUtils.Automation
 
 		public static void LoadBundleSubsetConfig(ConfigHierarchy BundleConfig, out IReadOnlyDictionary<string, BundleSubsetSettings> BundlesSubsets)
 		{
-			var Results = new Dictionary<string, BundleSubsetSettings>();
+			Dictionary<string, BundleSubsetSettings> Results = new Dictionary<string, BundleSubsetSettings>();
 
 			const string BundleDefinitionSubsetPrefix = "InstallBundleSubsetDefinition ";
 
@@ -373,7 +373,7 @@ namespace AutomationUtils.Automation
 		public static IEnumerable<string> GetCloudAppsForBundles<TPlatformBundleSettings>(IReadOnlyDictionary<string, TPlatformBundleSettings> InstallBundles) where TPlatformBundleSettings : BundleSettings
 		{
 			// Gather all bundle clouda app name overrides
-			var CloudAppNames = new HashSet<string>();
+			HashSet<string> CloudAppNames = new HashSet<string>();
 			foreach (var Bundle in InstallBundles.Values)
 			{
 				CloudAppNames.Add(Bundle.CloudAppNameOverride);
@@ -384,14 +384,14 @@ namespace AutomationUtils.Automation
 		public static IEnumerable<TPlatformBundleSettings> GetBundleDependencies<TPlatformBundleSettings>(
 			TPlatformBundleSettings Bundle, IReadOnlyDictionary<string, TPlatformBundleSettings> InstallBundles) where TPlatformBundleSettings : BundleSettings
 		{
-			var FoundDependencies = new Dictionary<string, TPlatformBundleSettings>();
+			Dictionary<string, TPlatformBundleSettings> FoundDependencies = new Dictionary<string, TPlatformBundleSettings>();
 
-			var CurrentDependencies = new List<TPlatformBundleSettings>();
+			List<TPlatformBundleSettings> CurrentDependencies = new List<TPlatformBundleSettings>();
 			CurrentDependencies.Add(Bundle);
 
 			while (CurrentDependencies.Count > 0)
 			{
-				var NextDependencies = new List<TPlatformBundleSettings>();
+				List<TPlatformBundleSettings> NextDependencies = new List<TPlatformBundleSettings>();
 				foreach (var Dep in CurrentDependencies)
 				{
 					if (FoundDependencies.ContainsKey(Dep.Name))
diff --git a/Engine/Source/Programs/AutomationTool/AutomationUtils/CommandUtils.cs b/Engine/Source/Programs/AutomationTool/AutomationUtils/CommandUtils.cs
--- a/Engine/Source/Programs/AutomationTool/AutomationUtils/CommandUtils.cs
+++ b/Engine/Source/Programs/AutomationTool/AutomationUtils/CommandUtils.cs
@@ -1273,7 +1273,7 @@ namespace AutomationTool
 				CombinePathMaxLength += (PathPart != null) ? PathPart.Length : 0;
 			}
 			CombinePathMaxLength += Paths.Length;
-			var CombinedPath = new StringBuilder(CombinePathMaxLength);
+			StringBuilder CombinedPath = new StringBuilder(CombinePathMaxLength);
 
 			// Combine all paths
 			CombinedPath.Append(Paths[0].Replace("\"", ""));
@@ -1769,8 +1769,8 @@ namespace AutomationTool
 		public static void ThreadedCopyFiles(string SourceDirectory, string DestDirectory, int MaxThreads = 64)
 		{
 			CreateDirectory(DestDirectory);
-			var SourceFiles = Directory.EnumerateFiles(SourceDirectory, "*", SearchOption.AllDirectories).ToList();
-			var DestFiles = SourceFiles.Select(SourceFile => CommandUtils.MakeRerootedFilePath(SourceFile, SourceDirectory, DestDirectory)).ToList();
+			List<string> SourceFiles = Directory.EnumerateFiles(SourceDirectory, "*", SearchOption.AllDirectories).ToList();
+			List<string> DestFiles = SourceFiles.Select(SourceFile => CommandUtils.MakeRerootedFilePath(SourceFile, SourceDirectory, DestDirectory)).ToList();
 			ThreadedCopyFiles(SourceFiles, DestFiles, MaxThreads);
 		}
 
@@ -1848,8 +1848,8 @@ namespace AutomationTool
 		/// <returns>List of filenames copied to the target directory</returns>
 		public static List<string> ThreadedCopyFiles(string SourceDir, string TargetDir, List<string> RelativePaths, int MaxThreads = 64, bool bRetry = false)
 		{
-			var SourceFileNames = RelativePaths.Select(RelativePath => CommandUtils.CombinePaths(SourceDir, RelativePath)).ToList();
-			var TargetFileNames = RelativePaths.Select(RelativePath => CommandUtils.CombinePaths(TargetDir, RelativePath)).ToList();
+			List<string> SourceFileNames = RelativePaths.Select(RelativePath => CommandUtils.CombinePaths(SourceDir, RelativePath)).ToList();
+			List<string> TargetFileNames = RelativePaths.Select(RelativePath => CommandUtils.CombinePaths(TargetDir, RelativePath)).ToList();
 			CommandUtils.ThreadedCopyFiles(SourceFileNames, TargetFileNames, MaxThreads, bRetry: bRetry);
 			return TargetFileNames;
 		}
@@ -1869,7 +1869,7 @@ namespace AutomationTool
 			// Filter all the relative paths
 			Logger.LogInformation("Applying filter to {SourceDir}...", SourceDir);
 			DirectoryReference SourceDirRef = new DirectoryReference(SourceDir);
-			var RelativePaths = Filter.ApplyToDirectory(SourceDirRef, bIgnoreSymlinks).Select(x => x.MakeRelativeTo(SourceDirRef)).ToList();
+			List<string> RelativePaths = Filter.ApplyToDirectory(SourceDirRef, bIgnoreSymlinks).Select(x => x.MakeRelativeTo(SourceDirRef)).ToList();
 			return ThreadedCopyFiles(SourceDir, TargetDir, RelativePaths, MaxThreads, bRetry: bRetry);
 		}
 
diff --git a/Engine/Source/Programs/AutomationTool/AutomationUtils/CommandletUtils.cs b/Engine/Source/Programs/AutomationTool/AutomationUtils/CommandletUtils.cs
--- a/Engine/Source/Programs/AutomationTool/AutomationUtils/CommandletUtils.cs
+++ b/Engine/Source/Programs/AutomationTool/AutomationUtils/CommandletUtils.cs
@@ -197,12 +197,12 @@ namespace AutomationTool
             {
                 throw new AutomationException("GenerateDistillFileSets did not produce a manifest for {0}.", ProjectFile);
             }
-            var Lines = new List<string>(ReadAllLines(ManifestFile));
+			List<string> Lines = new List<string>(ReadAllLines(ManifestFile));
             if (Lines.Count < 1)
             {
                 throw new AutomationException("GenerateDistillFileSets for {0} did not produce any files.", ProjectFile);
             }
-            var Result = new List<FileReference>();
+			List<FileReference> Result = new List<FileReference>();
             foreach (var ThisFile in Lines)
             {
                 var TestFile = CombinePaths(ThisFile);
@@ -210,9 +210,9 @@ namespace AutomationTool
                 {
                     throw new AutomationException("GenerateDistillFileSets produced {0}, but {1} doesn't exist.", ThisFile, TestFile);
                 }
-                // we correct the case here
-                var TestFileInfo = new FileInfo(TestFile);
-                var FinalFile = new FileReference(TestFileInfo);
+				// we correct the case here
+				FileInfo TestFileInfo = new FileInfo(TestFile);
+				FileReference FinalFile = new FileReference(TestFileInfo);
                 if (!FileReference.Exists(FinalFile))
                 {
                     throw new AutomationException("GenerateDistillFileSets produced {0}, but {1} doesn't exist.", ThisFile, FinalFile);
@@ -563,7 +563,7 @@ namespace AutomationTool
 			}
 			else
 			{
-				var Branch = new BranchInfo();
+				BranchInfo Branch = new BranchInfo();
 				var GameProj = Branch.FindGame(OriginalProjectName);
 				if (GameProj != null)
 				{
diff --git a/Engine/Source/Programs/AutomationTool/AutomationUtils/DeploymentContext.cs b/Engine/Source/Programs/AutomationTool/AutomationUtils/DeploymentContext.cs
--- a/Engine/Source/Programs/AutomationTool/AutomationUtils/DeploymentContext.cs
+++ b/Engine/Source/Programs/AutomationTool/AutomationUtils/DeploymentContext.cs
@@ -1054,7 +1054,7 @@ public class DeploymentContext //: ProjectParams
 			List<string> All = new();
 			CommandUtils.FindFilesAndSymlinks(InPath, Wildcard, bRecursive, All);
 
-			var Exclude = new HashSet<string>();
+			HashSet<string> Exclude = new HashSet<string>();
 			if (ExcludeWildcard != null)
 			{
 				foreach (var Excl in ExcludeWildcard)
diff --git a/Engine/Source/Programs/AutomationTool/AutomationUtils/DeviceReservation.cs b/Engine/Source/Programs/AutomationTool/AutomationUtils/DeviceReservation.cs
--- a/Engine/Source/Programs/AutomationTool/AutomationUtils/DeviceReservation.cs
+++ b/Engine/Source/Programs/AutomationTool/AutomationUtils/DeviceReservation.cs
@@ -177,24 +177,24 @@ namespace AutomationTool.DeviceReservation
 
 		public static Uri AppendPath(this Uri BaseUri, string NewPath)
 		{
-			var Builder = new UriBuilder(BaseUri);
+			UriBuilder Builder = new UriBuilder(BaseUri);
 			Builder.Path += NewPath;
 			return Builder.Uri;
 		}
 
 		public static Uri AppendQuery(this Uri BaseUri, string NewQuery)
 		{
-			var Builder = new UriBuilder(BaseUri);
+			UriBuilder Builder = new UriBuilder(BaseUri);
 			Builder.Query += NewQuery;
 			return Builder.Uri;
 		}
 
 		public static T InvokeAPI<T>(Uri UriToRequest, string Method, object ObjectToSerialize)
 		{
-			var SerializeOptions = new JsonSerializerOptions() { PropertyNameCaseInsensitive = true };
+			JsonSerializerOptions SerializeOptions = new JsonSerializerOptions() { PropertyNameCaseInsensitive = true };
 			SerializeOptions.Converters.Add(new TimeSpanJsonConverter());
 
-			var Request = (HttpWebRequest)WebRequest.Create(UriToRequest);
+			HttpWebRequest Request = (HttpWebRequest)WebRequest.Create(UriToRequest);
 			Request.UseDefaultCredentials = true;
 			Request.Method = Method;
 
@@ -205,14 +205,14 @@ namespace AutomationTool.DeviceReservation
 				using (var RequestStream = Request.GetRequestStream())
 				{
 					var JsonString = JsonSerializer.Serialize(ObjectToSerialize, SerializeOptions);
-					var Writer = new StreamWriter(RequestStream);
+					StreamWriter Writer = new StreamWriter(RequestStream);
 					Writer.Write(JsonString);
 					Writer.Flush();
 					RequestStream.Flush();
 				}
 			}
 
-			using (var Response = (HttpWebResponse)Request.GetResponse())
+			using (HttpWebResponse Response = (HttpWebResponse)Request.GetResponse())
 			using (var ResponseStream = Response.GetResponseStream())
 			{
 				MemoryStream MemoryStream = new MemoryStream();
@@ -224,7 +224,7 @@ namespace AutomationTool.DeviceReservation
 
 		public static void InvokeAPI(Uri UriToRequest, string Method, object ObjectToSerialize = null)
 		{
-			var Request = (HttpWebRequest)WebRequest.Create(UriToRequest);
+			HttpWebRequest Request = (HttpWebRequest)WebRequest.Create(UriToRequest);
 			Request.UseDefaultCredentials = true;
 			Request.Method = Method;
 
@@ -235,7 +235,7 @@ namespace AutomationTool.DeviceReservation
 				using (var RequestStream = Request.GetRequestStream())
 				{
 					var JsonString = JsonSerializer.Serialize(ObjectToSerialize, new JsonSerializerOptions());
-					var Writer = new StreamWriter(RequestStream);
+					StreamWriter Writer = new StreamWriter(RequestStream);
 					Writer.Write(JsonString);
 					Writer.Flush();
 					RequestStream.Flush();
@@ -246,7 +246,7 @@ namespace AutomationTool.DeviceReservation
 				Request.ContentLength = 0;
 			}
 
-			using (var Response = (HttpWebResponse)Request.GetResponse())
+			using (HttpWebResponse Response = (HttpWebResponse)Request.GetResponse())
 			using (var ResponseStream = Response.GetResponseStream())
 			{
 				// Nothing to do here. Error codes throw exceptions.
diff --git a/Engine/Source/Programs/AutomationTool/AutomationUtils/FileRegions.cs b/Engine/Source/Programs/AutomationTool/AutomationUtils/FileRegions.cs
--- a/Engine/Source/Programs/AutomationTool/AutomationUtils/FileRegions.cs
+++ b/Engine/Source/Programs/AutomationTool/AutomationUtils/FileRegions.cs
@@ -37,10 +37,10 @@ namespace AutomationTool
 			// This serialization function must match FFileRegion::SerializeFileRegions in Serialization/FileRegions.cpp
 			try
 			{
-				using (var Reader = new BinaryReader(File.OpenRead(Filename)))
+				using (BinaryReader Reader = new BinaryReader(File.OpenRead(Filename)))
 				{
 					int NumRegions = Reader.ReadInt32();
-					var Results = new List<FileRegion>(NumRegions);
+					List<FileRegion> Results = new List<FileRegion>(NumRegions);
 					for (int Index = 0; Index < NumRegions; ++Index)
 					{
 						FileRegion Region;
@@ -61,7 +61,7 @@ namespace AutomationTool
 
 		public static void PrintSummaryTable(IEnumerable<FileRegion> Regions)
         {
-			var RegionsByType = (from Region in Regions group Region by Region.Type).ToDictionary(g => g.Key, g => g.ToList());
+			Dictionary<FileRegionType, List<FileRegion>> RegionsByType = (from Region in Regions group Region by Region.Type).ToDictionary(g => g.Key, g => g.ToList());
 
 			Logger.LogInformation(" +-------------------+----------------------------------+---------------+");
 			Logger.LogInformation(" |        Type       |           Total Length           |  Num Regions  |");
diff --git a/Engine/Source/Programs/AutomationTool/AutomationUtils/GoogleSheetsHelper.cs b/Engine/Source/Programs/AutomationTool/AutomationUtils/GoogleSheetsHelper.cs
--- a/Engine/Source/Programs/AutomationTool/AutomationUtils/GoogleSheetsHelper.cs
+++ b/Engine/Source/Programs/AutomationTool/AutomationUtils/GoogleSheetsHelper.cs
@@ -101,7 +101,7 @@ namespace SheetsHelper
 			{
 				UserCredential credential;
 
-				using (var stream =
+				using (FileStream stream =
 					new FileStream(SecretKeyPath, FileMode.Open, FileAccess.Read))
 				{
 					string credPath = System.Environment.GetFolderPath(
@@ -190,10 +190,10 @@ namespace SheetsHelper
 
 			if (Sheet != null)
 			{
-				var Request = new Request();
+				Request Request = new Request();
 				Request.DeleteSheet = new DeleteSheetRequest();
 				Request.DeleteSheet.SheetId = Sheet.Properties.SheetId;
-				var BatchRequest = new BatchUpdateSpreadsheetRequest();
+				BatchUpdateSpreadsheetRequest BatchRequest = new BatchUpdateSpreadsheetRequest();
 
 				BatchRequest.Requests = new List<Request> { Request };
 
@@ -205,13 +205,13 @@ namespace SheetsHelper
 
 		public SheetWrapper AddSheet(string Title, int Index = 0)
 		{
-			var Request = new Request();
+			Request Request = new Request();
 			Request.AddSheet = new AddSheetRequest();
 			Request.AddSheet.Properties = new SheetProperties();
 			Request.AddSheet.Properties.Title = Title;
 			Request.AddSheet.Properties.Index = Index;
 
-			var BatchRequest = new BatchUpdateSpreadsheetRequest();
+			BatchUpdateSpreadsheetRequest BatchRequest = new BatchUpdateSpreadsheetRequest();
 
 			BatchRequest.Requests = new List<Request> { Request };
 
@@ -244,7 +244,7 @@ namespace SheetsHelper
 		/// </summary>
 		public void BatchAddSheet(string Title, int index = 0)
         {
-            var Request = new Request();
+			Request Request = new Request();
             Request.AddSheet = new AddSheetRequest();
             Request.AddSheet.Properties = new SheetProperties();
             Request.AddSheet.Properties.Title = Title;
@@ -262,7 +262,7 @@ namespace SheetsHelper
 		/// </summary>
 		public void BatchAppendDimension(SheetWrapper Sheet, int RowsOrColumnsToAppend = 1, ESheetsDimension Dimension = ESheetsDimension.ROWS)
         {
-            var Request = new Request();
+			Request Request = new Request();
             Request.AppendDimension = new AppendDimensionRequest();
             Request.AppendDimension.SheetId = Sheet.SheetID;
             Request.AppendDimension.Length = RowsOrColumnsToAppend;
@@ -289,14 +289,14 @@ namespace SheetsHelper
 		/// </summary>
 		public void BatchInsertDimension(SheetWrapper Sheet, int RowsOrColumnsToAppend = 1, int StartIndex = 0, ESheetsDimension Dimension = ESheetsDimension.ROWS, bool? bInheritFromBefore = null)
         {
-            var Range = new DimensionRange();
+			DimensionRange Range = new DimensionRange();
             Range.SheetId = Sheet.SheetID;
             Range.StartIndex = StartIndex;
             Range.EndIndex = StartIndex + RowsOrColumnsToAppend;
             Range.Dimension = Enum.GetName(typeof(ESheetsDimension), Dimension);
 
-            // create a request to insert the dimension
-            var Request = new Request();
+			// create a request to insert the dimension
+			Request Request = new Request();
             Request.InsertDimension = new InsertDimensionRequest();
             Request.InsertDimension.Range = Range;
             Request.InsertDimension.InheritFromBefore = bInheritFromBefore;
@@ -461,20 +461,20 @@ namespace SheetsHelper
 			 */
 
 			// define the range
-			var Range = new DimensionRange();
+			DimensionRange Range = new DimensionRange();
 			Range.SheetId = SheetID;
 			Range.StartIndex = RowIndex;
 			Range.EndIndex = RowIndex + RowCount;
 			Range.Dimension = "ROWS";
 
 			// create a request to insert the dimension
-			var Request = new Request();
+			Request Request = new Request();
 			Request.InsertDimension = new InsertDimensionRequest();
 			Request.InsertDimension.Range = Range;
 			Request.InsertDimension.InheritFromBefore = false;
 
 			// Request to update the spreadsheet
-			var BatchRequest = new BatchUpdateSpreadsheetRequest();
+			BatchUpdateSpreadsheetRequest BatchRequest = new BatchUpdateSpreadsheetRequest();
 			BatchRequest.Requests = new List<Request> { Request };
 
 			SheetsUtils.TryCheckQuotaSync();
@@ -485,18 +485,18 @@ namespace SheetsHelper
 
         public void AppendColumns(int ColumnCount)
         {
-            /*
+			/*
 			 * https://developers.google.com/sheets/samples/rowcolumn
 			 */
 
-            // create a request to insert the dimension
-            var Request = new Request();
+			// create a request to insert the dimension
+			Request Request = new Request();
             Request.AppendDimension = new AppendDimensionRequest();
             Request.AppendDimension.Length = ColumnCount;
             Request.AppendDimension.Dimension = "Columns";
 
-            // Request to update the spreadsheet
-            var BatchRequest = new BatchUpdateSpreadsheetRequest();
+			// Request to update the spreadsheet
+			BatchUpdateSpreadsheetRequest BatchRequest = new BatchUpdateSpreadsheetRequest();
             BatchRequest.Requests = new List<Request> { Request };
 
 			SheetsUtils.TryCheckQuotaSync();
@@ -512,13 +512,13 @@ namespace SheetsHelper
 			 */
 
 			// create a request to insert the dimension
-			var Request = new Request();
+			Request Request = new Request();
 			Request.AppendDimension = new AppendDimensionRequest();
 			Request.AppendDimension.Length = RowCount;
 			Request.AppendDimension.Dimension = "Rows";
 
 			// Request to update the spreadsheet
-			var BatchRequest = new BatchUpdateSpreadsheetRequest();
+			BatchUpdateSpreadsheetRequest BatchRequest = new BatchUpdateSpreadsheetRequest();
 			BatchRequest.Requests = new List<Request> { Request };
 
 			SheetsUtils.TryCheckQuotaSync();
diff --git a/Engine/Source/Programs/AutomationTool/AutomationUtils/HashCollection.cs b/Engine/Source/Programs/AutomationTool/AutomationUtils/HashCollection.cs
--- a/Engine/Source/Programs/AutomationTool/AutomationUtils/HashCollection.cs
+++ b/Engine/Source/Programs/AutomationTool/AutomationUtils/HashCollection.cs
@@ -90,7 +90,7 @@ namespace AutomationTool
 		/// <param name="HashType"></param>
 		public bool AddFile(FileReference InPath, HashCollection.HashType HashType)
 		{
-			var Fi = new FileInfo(InPath.FullName);
+			FileInfo Fi = new FileInfo(InPath.FullName);
 
 			if (!Fi.Exists)
 			{
@@ -207,7 +207,7 @@ namespace AutomationTool
 			{
 				using (var stream = System.IO.File.OpenRead(InPath))
 				{
-					var Serializer = new XmlSerializer(typeof(HashCollection));
+					XmlSerializer Serializer = new XmlSerializer(typeof(HashCollection));
 					return Serializer.Deserialize(stream) as HashCollection;
 				}
 			}
diff --git a/Engine/Source/Programs/AutomationTool/AutomationUtils/HelpUtils.cs b/Engine/Source/Programs/AutomationTool/AutomationUtils/HelpUtils.cs
--- a/Engine/Source/Programs/AutomationTool/AutomationUtils/HelpUtils.cs
+++ b/Engine/Source/Programs/AutomationTool/AutomationUtils/HelpUtils.cs
@@ -40,8 +40,8 @@ namespace AutomationTool
 		/// <summary>
 		/// Constructor for parameter descriptions.
 		/// </summary>
-		/// <param name="ParamName">Paramter name</param>
-		/// <param name="Description">Paramter description</param>
+		/// <param name="ParamName">Parameter name</param>
+		/// <param name="Description">Parameter description</param>
 		public HelpAttribute(string ParamName, string Description)
 			: base(FormatDescription(ParamName, Description))
 		{
@@ -131,7 +131,7 @@ namespace AutomationTool
 
 				string RequiredStr = Required ? ". Required" : ". Optional";
 				string ChoicesStr = "";
-				var ValidChoices = Choices as IEnumerable<object>;
+				IEnumerable<object> ValidChoices = Choices as IEnumerable<object>;
 				if (ValidChoices != null)
 				{
 					ChoicesStr += ". Choices=[";
@@ -196,7 +196,7 @@ namespace AutomationTool
 			}
 			set
 			{
-				var Enumerable = value as System.Collections.IEnumerable;
+				IEnumerable Enumerable = value as System.Collections.IEnumerable;
 				if (Enumerable != null)
 				{
 					_Choices = Enumerable;
@@ -309,7 +309,7 @@ namespace AutomationTool
 					var SplitPoint = Param.IndexOf(' ');
 					if (SplitPoint > 0)
 					{
-						// Extract the name and description seperately
+						// Extract the name and description separately
 						ParamName = Param.Substring(0, SplitPoint);
 						ParamDesc = Param.Substring(SplitPoint + 1, Param.Length - (SplitPoint + 1));
 					}
@@ -339,7 +339,7 @@ namespace AutomationTool
 		/// </summary>
 		/// <param name="ObjType">Type to get help for.</param>
 		/// <param name="ObjectDescription">Description</param>
-		/// <param name="Params">List of paramters</param>
+		/// <param name="Params">List of parameters</param>
 		public static void GetTypeHelp(Type ObjType, out string ObjectDescription, out List<string> Params)
 		{
 			ObjectDescription = String.Empty;
@@ -348,7 +348,7 @@ namespace AutomationTool
 			var AllAttributes = ObjType.GetCustomAttributes(false);			
 			foreach (var CustomAttribute in AllAttributes)
 			{
-				var HelpAttribute = CustomAttribute as HelpAttribute;
+				HelpAttribute HelpAttribute = CustomAttribute as HelpAttribute;
 				if (HelpAttribute != null)
 				{
 					if (HelpAttribute.IsParam)
@@ -379,7 +379,7 @@ namespace AutomationTool
 				var AllMemeberAtributes = Member.GetCustomAttributes(false);
 				foreach (var CustomAttribute in AllMemeberAtributes)
 				{
-					var HelpAttribute = CustomAttribute as HelpAttribute;
+					HelpAttribute HelpAttribute = CustomAttribute as HelpAttribute;
 					if (HelpAttribute != null && HelpAttribute.IsParam)
 					{
 						Params.Add(HelpAttribute.Description);
@@ -390,7 +390,7 @@ namespace AutomationTool
 
 		private static List<MemberInfo> GetAllMemebers(Type ObjType)
 		{
-			var Members = new List<MemberInfo>();
+			List<MemberInfo> Members = new List<MemberInfo>();
 			for (; ObjType != null; ObjType = ObjType.BaseType)
 			{
 				var ObjMembers = ObjType.GetMembers(BindingFlags.Public | BindingFlags.NonPublic | BindingFlags.Instance | BindingFlags.Static);
diff --git a/Engine/Source/Programs/AutomationTool/AutomationUtils/HttpClientSingleton.cs b/Engine/Source/Programs/AutomationTool/AutomationUtils/HttpClientSingleton.cs
--- a/Engine/Source/Programs/AutomationTool/AutomationUtils/HttpClientSingleton.cs
+++ b/Engine/Source/Programs/AutomationTool/AutomationUtils/HttpClientSingleton.cs
@@ -23,8 +23,8 @@ namespace AutomationUtils
 				.WaitAndRetryAsync(3, retryAttempt => TimeSpan.FromSeconds(Math.Pow(2, retryAttempt)));
 
 				// per MS documentation, with a lifetime specified we shouldn't have socket exhaustion issues
-				var socketHandler = new SocketsHttpHandler { PooledConnectionLifetime = TimeSpan.FromMinutes(15) };
-				var pollyHandler = new PolicyHttpMessageHandler(retryPolicy)
+				SocketsHttpHandler socketHandler = new SocketsHttpHandler { PooledConnectionLifetime = TimeSpan.FromMinutes(15) };
+				PolicyHttpMessageHandler pollyHandler = new PolicyHttpMessageHandler(retryPolicy)
 				{
 					InnerHandler = socketHandler,
 				};
diff --git a/Engine/Source/Programs/AutomationTool/AutomationUtils/LinuxHostPlatform.cs b/Engine/Source/Programs/AutomationTool/AutomationUtils/LinuxHostPlatform.cs
--- a/Engine/Source/Programs/AutomationTool/AutomationUtils/LinuxHostPlatform.cs
+++ b/Engine/Source/Programs/AutomationTool/AutomationUtils/LinuxHostPlatform.cs
@@ -74,7 +74,7 @@ namespace AutomationTool
 
 		public override Process CreateProcess(string AppName)
 		{
-			var NewProcess = new Process();
+			Process NewProcess = new Process();
 			return NewProcess;
 		}
 
diff --git a/Engine/Source/Programs/AutomationTool/AutomationUtils/MacHostPlatform.cs b/Engine/Source/Programs/AutomationTool/AutomationUtils/MacHostPlatform.cs
--- a/Engine/Source/Programs/AutomationTool/AutomationUtils/MacHostPlatform.cs
+++ b/Engine/Source/Programs/AutomationTool/AutomationUtils/MacHostPlatform.cs
@@ -103,7 +103,7 @@ namespace AutomationTool
 
 		public override Process CreateProcess(string AppName)
 		{
-			var NewProcess = new Process();
+			Process NewProcess = new Process();
 			return NewProcess;
 		}
 
diff --git a/Engine/Source/Programs/AutomationTool/AutomationUtils/P4Utils.cs b/Engine/Source/Programs/AutomationTool/AutomationUtils/P4Utils.cs
--- a/Engine/Source/Programs/AutomationTool/AutomationUtils/P4Utils.cs
+++ b/Engine/Source/Programs/AutomationTool/AutomationUtils/P4Utils.cs
@@ -1917,7 +1917,7 @@ namespace AutomationTool
                             throw new AutomationException("Was expecting changelists to be reported back in the same order we asked for them (CL {0} != {1})", ChangeString, CurrentChangelist.ToString());
                         }
 
-						var DescribeRecord = new DescribeRecord();
+						DescribeRecord DescribeRecord = new DescribeRecord();
 						DescribeRecords.Add( DescribeRecord );
 
 						DescribeRecord.CL = CurrentChangelist;
@@ -2006,7 +2006,7 @@ namespace AutomationTool
 								// Remove the "... " prefix
 								var FilesLine = Line.Substring( 4 );
 
-								var DescribeFile = new DescribeRecord.DescribeFile();
+								DescribeRecord.DescribeFile DescribeFile = new DescribeRecord.DescribeFile();
 								DescribeRecord.Files.Add( DescribeFile );
  							
 								// Find the revision #
@@ -2793,7 +2793,7 @@ namespace AutomationTool
 		/// <returns>List of the files contained in the changelist.</returns>
 		public List<string> ChangeFiles(int CL, out bool Pending, bool AllowSpew = true)
 		{
-			var Result = new List<string>();
+			List<string> Result = new List<string>();
 
 			if (ChangeExists(CL, out Pending, AllowSpew))
 			{
@@ -3093,7 +3093,7 @@ namespace AutomationTool
 		/// <returns></returns>
 		public P4Label[] GetLabels(string Filter, bool bCaseSensitive = true)
 		{
-			var LabelList = new List<P4Label>();
+			List<P4Label> LabelList = new List<P4Label>();
 
 			string Output;
 			if (P4Output(out Output, "", "labels -t " + (bCaseSensitive ? "-e" : "-E") + Filter, null, false))
@@ -3420,7 +3420,7 @@ namespace AutomationTool
 		/// <returns>Parsed output.</returns>
 		public Dictionary<string, string> ParseTaggedP4Output(string Output)
 		{
-			var Tags = new Dictionary<string, string>(StringComparer.InvariantCultureIgnoreCase);
+			Dictionary<string, string> Tags = new Dictionary<string, string>(StringComparer.InvariantCultureIgnoreCase);
 			var Lines = Output.Split(new string[] { Environment.NewLine }, StringSplitOptions.RemoveEmptyEntries);
             string DelayKey = "";
             int DelayIndex = 0;
@@ -3601,7 +3601,7 @@ namespace AutomationTool
                         int ClientStartIndex = Mapping.IndexOf(ClientMappingRoot, StringComparison.InvariantCultureIgnoreCase);
                         if (ClientStartIndex > 0)
                         {
-                            var ViewPair = new KeyValuePair<string, string>(
+							KeyValuePair<string, string> ViewPair = new KeyValuePair<string, string>(
                                 Mapping.Substring(0, ClientStartIndex - 1),
                                 Mapping.Substring(ClientStartIndex + ClientMappingRoot.Length));
                             Info.View.Add(ViewPair);
@@ -3625,7 +3625,7 @@ namespace AutomationTool
 		/// <returns>List of clients owned by the user.</returns>
 		public P4ClientInfo[] GetClientsForUser(string UserName, string PathUnderClientRoot = null, string AllowedStream = null)
 		{
-			var ClientList = new List<P4ClientInfo>();
+			List<P4ClientInfo> ClientList = new List<P4ClientInfo>();
 
 			// Get all clients for this user
 			string P4Command = String.Format("clients -u {0}", UserName);
@@ -3750,7 +3750,7 @@ namespace AutomationTool
 			{
 				throw new AutomationException("{0} failed.", DirsCmdLine);
 			}
-			var Result = new List<string>();
+			List<string> Result = new List<string>();
 			var Lines = P4Result.Output.Split(new string[] { Environment.NewLine }, StringSplitOptions.RemoveEmptyEntries);
 			foreach (string Line in Lines)
 			{
@@ -4472,7 +4472,7 @@ namespace AutomationTool
 			string Text = "";
 			foreach (var Attr in AllAttributes)
 			{
-				var AttrValue = (P4FileAttributes)Attr.Key;
+				P4FileAttributes AttrValue = (P4FileAttributes)Attr.Key;
 				if ((Attributes & AttrValue) == AttrValue)
 				{
 					Text += Attr.Value;
diff --git a/Engine/Source/Programs/AutomationTool/AutomationUtils/Platform.cs b/Engine/Source/Programs/AutomationTool/AutomationUtils/Platform.cs
--- a/Engine/Source/Programs/AutomationTool/AutomationUtils/Platform.cs
+++ b/Engine/Source/Programs/AutomationTool/AutomationUtils/Platform.cs
@@ -121,7 +121,7 @@ namespace AutomationTool
 			// Create dummy platforms for platforms we don't support
 			foreach (UnrealTargetPlatform PlatformType in UnrealTargetPlatform.GetValidPlatforms())
 			{
-				var TargetDesc = new TargetPlatformDescriptor(PlatformType);
+				TargetPlatformDescriptor TargetDesc = new TargetPlatformDescriptor(PlatformType);
 				Platform ExistingInstance;
 				if (AllPlatforms.TryGetValue(TargetDesc, out ExistingInstance) == false)
 				{
@@ -144,7 +144,7 @@ namespace AutomationTool
 				Logger.LogError("Failed to get assembly types for {Location}", ScriptAssembly.Location);
 				if (Ex is ReflectionTypeLoadException)
 				{
-					var TypeLoadException = (ReflectionTypeLoadException)Ex;
+					ReflectionTypeLoadException TypeLoadException = (ReflectionTypeLoadException)Ex;
 					if (!IsNullOrEmpty(TypeLoadException.LoaderExceptions))
 					{
 						Logger.LogError("Loader Exceptions:");
@@ -166,7 +166,7 @@ namespace AutomationTool
 				if (PotentialPlatformType != typeof(Platform) && typeof(Platform).IsAssignableFrom(PotentialPlatformType) && !PotentialPlatformType.IsAbstract)
 				{
 					Logger.LogDebug("Creating platform {Platform} from {Location}.", PotentialPlatformType.Name, ScriptAssembly.Location);
-					var PlatformInstance = Activator.CreateInstance(PotentialPlatformType) as Platform;
+					Platform PlatformInstance = Activator.CreateInstance(PotentialPlatformType) as Platform;
 					var PlatformDesc = PlatformInstance.GetTargetPlatformDescriptor();
 
 					Platform ExistingInstance;
diff --git a/Engine/Source/Programs/AutomationTool/AutomationUtils/ProcessUtils.cs b/Engine/Source/Programs/AutomationTool/AutomationUtils/ProcessUtils.cs
--- a/Engine/Source/Programs/AutomationTool/AutomationUtils/ProcessUtils.cs
+++ b/Engine/Source/Programs/AutomationTool/AutomationUtils/ProcessUtils.cs
@@ -208,7 +208,7 @@ namespace AutomationTool
 				NewProcess.StartInfo.WorkingDirectory = WorkingDir;
 			}
 
-			var Result = new ProcessResult(AppName, NewProcess, bAllowSpew, bCaptureSpew, SpewVerbosity: SpewVerbosity, InSpewFilterCallback: SpewFilterCallback);
+			ProcessResult Result = new ProcessResult(AppName, NewProcess, bAllowSpew, bCaptureSpew, SpewVerbosity: SpewVerbosity, InSpewFilterCallback: SpewFilterCallback);
 			AddProcess(Result);
 			return Result;
 		}
@@ -808,11 +808,11 @@ namespace AutomationTool
 				bKilledAChild = false;
 				// For some reason Process.GetProcesses() sometimes returns the same process twice
 				// So keep track of all processes we already tried to kill
-				var KilledPids = new HashSet<int>();
+				HashSet<int> KilledPids = new HashSet<int>();
 				var AllProcs = Process.GetProcesses();
 				foreach (Process KillCandidate in AllProcs)
 				{
-					var VisitedPids = new HashSet<int>();
+					HashSet<int> VisitedPids = new HashSet<int>();
 					if (ProcessManager.CanBeKilled(KillCandidate.ProcessName) && 
 						!KilledPids.Contains(KillCandidate.Id) && 
 						IsOurDescendant(ProcessToKill, KillCandidate.Id, VisitedPids))
diff --git a/Engine/Source/Programs/AutomationTool/AutomationUtils/ProjectParams.cs b/Engine/Source/Programs/AutomationTool/AutomationUtils/ProjectParams.cs
--- a/Engine/Source/Programs/AutomationTool/AutomationUtils/ProjectParams.cs
+++ b/Engine/Source/Programs/AutomationTool/AutomationUtils/ProjectParams.cs
@@ -1171,7 +1171,7 @@ namespace AutomationTool
                     var PortString = Command.ParseParamValue("port");
                     if (String.IsNullOrEmpty(PortString) == false)
                     {
-                        var Ports = new ParamList<string>(PortString.Split(new char[] { '+' }, StringSplitOptions.RemoveEmptyEntries));
+						ParamList<string> Ports = new ParamList<string>(PortString.Split(new char[] { '+' }, StringSplitOptions.RemoveEmptyEntries));
                         foreach (var P in Ports)
                         {
                             this.Port.Add(P);
@@ -1194,7 +1194,7 @@ namespace AutomationTool
                     var MapsString = Command.ParseParamValue("MapsToCook");
                     if (String.IsNullOrEmpty(MapsString) == false)
                     {
-                        var MapNames = new ParamList<string>(MapsString.Split(new char[] { '+' }, StringSplitOptions.RemoveEmptyEntries));
+						ParamList<string> MapNames = new ParamList<string>(MapsString.Split(new char[] { '+' }, StringSplitOptions.RemoveEmptyEntries));
                         foreach ( var M in MapNames ) 
                         {
                             this.MapsToCook.Add( M );
@@ -1216,7 +1216,7 @@ namespace AutomationTool
 					var MapsString = Command.ParseParamValue("MapIniSectionsToCook");
 					if (String.IsNullOrEmpty(MapsString) == false)
 					{
-						var MapNames = new ParamList<string>(MapsString.Split(new char[] { '+' }, StringSplitOptions.RemoveEmptyEntries));
+						ParamList<string> MapNames = new ParamList<string>(MapsString.Split(new char[] { '+' }, StringSplitOptions.RemoveEmptyEntries));
 						foreach (var M in MapNames)
 						{
 							this.MapIniSectionsToCook.Add(M);
@@ -1253,7 +1253,7 @@ namespace AutomationTool
 					var MapsString = Command.ParseParamValue("MapsToRebuildLightMaps");
 					if (String.IsNullOrEmpty(MapsString) == false)
 					{
-						var MapNames = new ParamList<string>(MapsString.Split(new char[] { '+' }, StringSplitOptions.RemoveEmptyEntries));
+						ParamList<string> MapNames = new ParamList<string>(MapsString.Split(new char[] { '+' }, StringSplitOptions.RemoveEmptyEntries));
 						foreach (var M in MapNames)
 						{
 							this.MapsToRebuildLightMaps.Add(M);
@@ -1275,7 +1275,7 @@ namespace AutomationTool
                     var MapsString = Command.ParseParamValue("MapsToRebuildHLODMaps");
                     if (String.IsNullOrEmpty(MapsString) == false)
                     {
-                        var MapNames = new ParamList<string>(MapsString.Split(new char[] { '+' }, StringSplitOptions.RemoveEmptyEntries));
+						ParamList<string> MapNames = new ParamList<string>(MapsString.Split(new char[] { '+' }, StringSplitOptions.RemoveEmptyEntries));
                         foreach (var M in MapNames)
                         {
                             this.MapsToRebuildHLODMaps.Add(M);
@@ -1297,7 +1297,7 @@ namespace AutomationTool
 					var TitleString = Command.ParseParamValue("TitleID");
 					if (String.IsNullOrEmpty(TitleString) == false)
 					{
-						var TitleIDs = new ParamList<string>(TitleString.Split(new char[] { '+' }, StringSplitOptions.RemoveEmptyEntries));
+						ParamList<string> TitleIDs = new ParamList<string>(TitleString.Split(new char[] { '+' }, StringSplitOptions.RemoveEmptyEntries));
 						foreach (var T in TitleIDs)
 						{
 							this.TitleID.Add(T);
diff --git a/Engine/Source/Programs/AutomationTool/AutomationUtils/UnrealBuild.cs b/Engine/Source/Programs/AutomationTool/AutomationUtils/UnrealBuild.cs
--- a/Engine/Source/Programs/AutomationTool/AutomationUtils/UnrealBuild.cs
+++ b/Engine/Source/Programs/AutomationTool/AutomationUtils/UnrealBuild.cs
@@ -560,7 +560,7 @@ namespace AutomationTool
 			//////////////////////////////////////
 
 			// make a set of unique platforms involved
-			var UniquePlatforms = new List<UnrealTargetPlatform>();
+			List<UnrealTargetPlatform> UniquePlatforms = new List<UnrealTargetPlatform>();
 			foreach (var Target in InAgenda.Targets)
 			{
 				if (!UniquePlatforms.Contains(Target.Platform))
diff --git a/Engine/Source/Programs/AutomationTool/AutomationUtils/Utils.cs b/Engine/Source/Programs/AutomationTool/AutomationUtils/Utils.cs
--- a/Engine/Source/Programs/AutomationTool/AutomationUtils/Utils.cs
+++ b/Engine/Source/Programs/AutomationTool/AutomationUtils/Utils.cs
@@ -1032,7 +1032,7 @@ namespace AutomationTool
 				yield break;
 			}
 
-			using (var Reader = new StringReader(Str))
+			using (StringReader Reader = new StringReader(Str))
 			{
 				string line = Reader.ReadLine();
 				while (line != null)
@@ -1095,8 +1095,8 @@ namespace AutomationTool
         /// <returns>Version that puts the Major/Minor/Patch fields in the Major/Minor/Build fields, respectively.</returns>
         public static Version ReadVersionFromFile(string Filename)
         {
-            var regex = new Regex(@"#define.+_(?<Type>MAJOR|MINOR|PATCH)_VERSION\s+(?<Value>.+)", RegexOptions.ExplicitCapture | RegexOptions.IgnoreCase);
-            var foundElements = new Dictionary<string, int>(3);
+			Regex regex = new Regex(@"#define.+_(?<Type>MAJOR|MINOR|PATCH)_VERSION\s+(?<Value>.+)", RegexOptions.ExplicitCapture | RegexOptions.IgnoreCase);
+			Dictionary<string, int> foundElements = new Dictionary<string, int>(3);
             foreach (var line in File.ReadLines(Filename))
             {
                 try
@@ -1168,7 +1168,7 @@ namespace AutomationTool
 
 		public override bool Equals(object obj)
 		{
-			var that = obj as FEngineVersionSupport;
+			FEngineVersionSupport that = obj as FEngineVersionSupport;
 			if (that == null)
 				return false;
 			return this.ToString() == that.ToString();
diff --git a/Engine/Source/Programs/AutomationTool/AutomationUtils/WindowsHostPlatform.cs b/Engine/Source/Programs/AutomationTool/AutomationUtils/WindowsHostPlatform.cs
--- a/Engine/Source/Programs/AutomationTool/AutomationUtils/WindowsHostPlatform.cs
+++ b/Engine/Source/Programs/AutomationTool/AutomationUtils/WindowsHostPlatform.cs
@@ -78,7 +78,7 @@ namespace AutomationTool
 
 		public override Process CreateProcess(string AppName)
 		{
-			var NewProcess = new Process();
+			Process NewProcess = new Process();
 			return NewProcess;
 		}
 
diff --git a/Engine/Source/Programs/AutomationTool/AutomationUtils/ZenUtils.cs b/Engine/Source/Programs/AutomationTool/AutomationUtils/ZenUtils.cs
--- a/Engine/Source/Programs/AutomationTool/AutomationUtils/ZenUtils.cs
+++ b/Engine/Source/Programs/AutomationTool/AutomationUtils/ZenUtils.cs
@@ -206,11 +206,11 @@ namespace AutomationUtils
 		/// <returns>A new HttpClient, with SocketsHttpHandler set.</returns>
 		public static HttpClient CreateUnixSocketHttpClient(string UnixSocketPath)
 		{
-			var SocketsHandler = new SocketsHttpHandler
+			SocketsHttpHandler SocketsHandler = new SocketsHttpHandler
 			{
 				ConnectCallback = async (Context, Token) => {
-					var Socket = new Socket(AddressFamily.Unix, SocketType.Stream, ProtocolType.IP);
-					var EndPoint = new UnixDomainSocketEndPoint(UnixSocketPath);
+					Socket Socket = new Socket(AddressFamily.Unix, SocketType.Stream, ProtocolType.IP);
+					UnixDomainSocketEndPoint EndPoint = new UnixDomainSocketEndPoint(UnixSocketPath);
 					await Socket.ConnectAsync(EndPoint);
 					return new NetworkStream(Socket, ownsSocket: true);
 				}
@@ -307,7 +307,7 @@ namespace AutomationUtils
 			}
 
 			string Uri = string.Format("http://{0}/health/ready", GetHttpHostNameAndPortString());
-			using var Request = new HttpRequestMessage(HttpMethod.Get, Uri);
+			using HttpRequestMessage Request = new HttpRequestMessage(HttpMethod.Get, Uri);
 
 			HttpResponseMessage HttpGetResult = null;
 			HttpClient Client = CreateHttpClientInternal();
@@ -447,7 +447,7 @@ namespace AutomationUtils
 			HttpResponseMessage HttpGetResult = null;
 			HttpClient Client = CreateHttpClientInternal();
 			string Uri = string.Format("http://{0}/ws", GetHttpHostNameAndPortString());
-			using var Request = new HttpRequestMessage(HttpMethod.Get, Uri);
+			using HttpRequestMessage Request = new HttpRequestMessage(HttpMethod.Get, Uri);
 			Request.Headers.Add("Accept", "application/json");
 
 			try
diff --git a/Engine/Source/Programs/AutomationTool/Localization/LocalizationProvider.cs b/Engine/Source/Programs/AutomationTool/Localization/LocalizationProvider.cs
--- a/Engine/Source/Programs/AutomationTool/Localization/LocalizationProvider.cs
+++ b/Engine/Source/Programs/AutomationTool/Localization/LocalizationProvider.cs
@@ -56,10 +56,10 @@ namespace EpicGames.Localization
 		public void CalculateSplitPlatformNames(string RootWorkingDirectory)
 		{
 			// Is this isn't a single culture import/export, then also check for split platform sub-folders
-			var NewSplitPlatformNames = new List<string>();
+			List<string> NewSplitPlatformNames = new List<string>();
 			if (bUseCultureDirectory)
 			{
-				var PlatformSourceDirectory = new DirectoryReference(CommandUtils.CombinePaths(RootWorkingDirectory, DestinationPath, PlatformLocalizationFolderName));
+				DirectoryReference PlatformSourceDirectory = new DirectoryReference(CommandUtils.CombinePaths(RootWorkingDirectory, DestinationPath, PlatformLocalizationFolderName));
 				if (DirectoryReference.Exists(PlatformSourceDirectory))
 				{
 					foreach (DirectoryReference FoundDirectory in DirectoryReference.EnumerateDirectories(PlatformSourceDirectory, "*", SearchOption.TopDirectoryOnly))
diff --git a/Engine/Source/Programs/AutomationTool/Program.cs b/Engine/Source/Programs/AutomationTool/Program.cs
--- a/Engine/Source/Programs/AutomationTool/Program.cs
+++ b/Engine/Source/Programs/AutomationTool/Program.cs
@@ -260,8 +260,8 @@ namespace AutomationToolDriver
 				{
 					// find if the command has been specified
 					var text = File.ReadAllText(Profile);
-					var RawObject = fastJSON.JSON.Instance.Parse(text) as Dictionary<string, object>;
-					var Params = RawObject["scripts"] as List<object>;
+					Dictionary<string, object> RawObject = fastJSON.JSON.Instance.Parse(text) as Dictionary<string, object>;
+					List<object> Params = RawObject["scripts"] as List<object>;
 					foreach (var Script in Params)
 					{
 						string ScriptName = (Script as Dictionary<string, object>)["script"] as string;
diff --git a/Engine/Source/Programs/OidcToken/TokenService.cs b/Engine/Source/Programs/OidcToken/TokenService.cs
--- a/Engine/Source/Programs/OidcToken/TokenService.cs
+++ b/Engine/Source/Programs/OidcToken/TokenService.cs
@@ -183,7 +183,7 @@ namespace OidcToken
 							client.DefaultRequestHeaders.Accept.Add(new MediaTypeWithQualityHeaderValue("application/json"));
 
 							string zenUrl = $"{Settings.CurrentValue.ZenUrl}/auth/oidc/refreshtoken";
-							var request = new ZenAddOidcTokenRequest { ProviderName = providerName, RefreshToken = tokenInfo.RefreshToken };
+							ZenAddOidcTokenRequest request = new ZenAddOidcTokenRequest { ProviderName = providerName, RefreshToken = tokenInfo.RefreshToken };
 							HttpContent content = new StringContent(JsonSerializer.Serialize<ZenAddOidcTokenRequest>(request, TokenServiceStateContext.Default.ZenAddOidcTokenRequest), Encoding.UTF8, "application/json");
 							HttpResponseMessage result = await client.PostAsync(zenUrl, content);
 								
