More replace var with type

#rnx
#Synced-CL 51479624

diff --git a/Engine/Source/Programs/AutomationTool/AutomationUtils/Automation.cs b/Engine/Source/Programs/AutomationTool/AutomationUtils/Automation.cs
--- a/Engine/Source/Programs/AutomationTool/AutomationUtils/Automation.cs
+++ b/Engine/Source/Programs/AutomationTool/AutomationUtils/Automation.cs
@@ -293,7 +293,7 @@ namespace AutomationTool
 				Logger.LogInformation("Executing commands...");
 				for (int CommandIndex = 0; CommandIndex < CommandsToExecute.Count; ++CommandIndex)
 				{
-					var CommandInfo = CommandsToExecute[CommandIndex];
+					CommandInfo CommandInfo = CommandsToExecute[CommandIndex];
 					Logger.LogDebug("Attempting to execute {CommandInfo}", CommandInfo.ToString());
 					Type CommandType;
 					if (!Commands.TryGetValue(CommandInfo.CommandName, out CommandType))
@@ -353,7 +353,7 @@ namespace AutomationTool
 		{
 			for (int CommandIndex = 0; CommandIndex < CommandsToExecute.Count; ++CommandIndex)
 			{
-				var CommandInfo = CommandsToExecute[CommandIndex];
+				CommandInfo CommandInfo = CommandsToExecute[CommandIndex];
 				Type CommandType;
 				if (Commands.TryGetValue(CommandInfo.CommandName, out CommandType) == false)
 				{
@@ -388,7 +388,7 @@ AutomationTool.exe [-verbose] [-Help|-List|-CompileOnly] [-p4] Command0 [-Arg0 -
 			string Message = Environment.NewLine;
 			Message += "Available commands:" + Environment.NewLine;
 			string AssemblyName = "";
-			foreach (var AvailableCommand in 
+			foreach (KeyValuePair<string, Type> AvailableCommand in 
 				Commands.OrderBy(Command => Command.Value.Assembly.GetName().Name).ThenBy(Command => Command.Key))
 			{
 				string NewAssemblyName = AvailableCommand.Value.Assembly.GetName().Name;
diff --git a/Engine/Source/Programs/AutomationTool/AutomationUtils/BuildCommand.cs b/Engine/Source/Programs/AutomationTool/AutomationUtils/BuildCommand.cs
--- a/Engine/Source/Programs/AutomationTool/AutomationUtils/BuildCommand.cs
+++ b/Engine/Source/Programs/AutomationTool/AutomationUtils/BuildCommand.cs
@@ -300,7 +300,7 @@ namespace AutomationTool
 			else
 			{
 				BranchInfo Branch = new BranchInfo();
-				var GameProj = Branch.FindGame(originalProjectName);
+				BranchInfo.BranchUProject GameProj = Branch.FindGame(originalProjectName);
 				if (GameProj != null)
 				{
 					projectFullPath = GameProj.FilePath;
diff --git a/Engine/Source/Programs/AutomationTool/AutomationUtils/BuildUtils.cs b/Engine/Source/Programs/AutomationTool/AutomationUtils/BuildUtils.cs
--- a/Engine/Source/Programs/AutomationTool/AutomationUtils/BuildUtils.cs
+++ b/Engine/Source/Programs/AutomationTool/AutomationUtils/BuildUtils.cs
@@ -131,11 +131,11 @@ namespace AutomationTool
 		{
 			string Permissions = IsProbablyAMacOrIOSExe(Filename) || IsProbablyALinuxExe(Filename) ? "0755" : "0644";
 			string CommandString = string.Format("-c 'chmod {0} \"{1}\"'", Permissions, Filename.Replace("'", "'\"'\"'"));
-			var Result = CommandUtils.Run("sh", CommandString, Options:ERunOptions.SpewIsVerbose);
+			IProcessResult Result = CommandUtils.Run("sh", CommandString, Options:ERunOptions.SpewIsVerbose);
 			if (Result.ExitCode != 0)
 			{
 				// Try simply invoking chmod directly.
-				var DirectResult = CommandUtils.Run("chmod", string.Format("{0} \"{1}\"", Permissions, Filename.Replace("'", "'\"'\"'")));
+				IProcessResult DirectResult = CommandUtils.Run("chmod", string.Format("{0} \"{1}\"", Permissions, Filename.Replace("'", "'\"'\"'")));
 				if (DirectResult.ExitCode != 0)
 				{
 					throw new AutomationException(String.Format("Command 'sh {0}' failed to chmod \"{1}\" (Exit Code {2}): {3}", CommandString, Filename, Result.ExitCode, Result.Output));
diff --git a/Engine/Source/Programs/AutomationTool/AutomationUtils/BundleUtils.cs b/Engine/Source/Programs/AutomationTool/AutomationUtils/BundleUtils.cs
--- a/Engine/Source/Programs/AutomationTool/AutomationUtils/BundleUtils.cs
+++ b/Engine/Source/Programs/AutomationTool/AutomationUtils/BundleUtils.cs
@@ -325,7 +325,7 @@ namespace AutomationUtils.Automation
 			string FileName, IReadOnlyDictionary<string, TPlatformBundleSettings> InstallBundles) where TPlatformBundleSettings : BundleSettings
 		{
 			// Try to find a matching chunk regex or exact filename
-			foreach (var Bundle in InstallBundles.Values)
+			foreach (TPlatformBundleSettings Bundle in InstallBundles.Values)
 			{
 				foreach (string RegexString in Bundle.FileRegex)
 				{
@@ -374,7 +374,7 @@ namespace AutomationUtils.Automation
 		{
 			// Gather all bundle clouda app name overrides
 			HashSet<string> CloudAppNames = new HashSet<string>();
-			foreach (var Bundle in InstallBundles.Values)
+			foreach (TPlatformBundleSettings Bundle in InstallBundles.Values)
 			{
 				CloudAppNames.Add(Bundle.CloudAppNameOverride);
 			}
@@ -392,7 +392,7 @@ namespace AutomationUtils.Automation
 			while (CurrentDependencies.Count > 0)
 			{
 				List<TPlatformBundleSettings> NextDependencies = new List<TPlatformBundleSettings>();
-				foreach (var Dep in CurrentDependencies)
+				foreach (TPlatformBundleSettings Dep in CurrentDependencies)
 				{
 					if (FoundDependencies.ContainsKey(Dep.Name))
 						continue;
diff --git a/Engine/Source/Programs/AutomationTool/AutomationUtils/CommandUtils.cs b/Engine/Source/Programs/AutomationTool/AutomationUtils/CommandUtils.cs
--- a/Engine/Source/Programs/AutomationTool/AutomationUtils/CommandUtils.cs
+++ b/Engine/Source/Programs/AutomationTool/AutomationUtils/CommandUtils.cs
@@ -971,7 +971,7 @@ namespace AutomationTool
 		/// <param name="UtcDate">The accessed and modified date to set.  If not specified, defaults to the current date and time.</param>
 		public static void TouchFile(string Filename, DateTime? UtcDate = null)
 		{
-			var Date = UtcDate ?? DateTime.UtcNow;
+			DateTime Date = UtcDate ?? DateTime.UtcNow;
 			Filename = ConvertSeparators(PathSeparator.Slash, Filename);
 			if (!File.Exists(Filename))
 			{
@@ -2651,7 +2651,7 @@ namespace AutomationTool
 
 					// We've failed to take the lock. Throw an exception if the timeout has elapsed.
 					// Otherwise print a log message and retry.
-					var CurrentTime = DateTime.Now;
+					DateTime CurrentTime = DateTime.Now;
 					if (CurrentTime >= Deadline)
 					{
 						throw new AutomationException("Couldn't create lock file '{0}' after {1} seconds.", LockFilePath, CurrentTime.Subtract(StartTime).TotalSeconds);
@@ -3221,10 +3221,10 @@ namespace AutomationTool
 			{
 				// Find and load the first helper class
 				Assembly[] LoadedAssemblies = AppDomain.CurrentDomain.GetAssemblies();
-				foreach (var Dll in LoadedAssemblies)
+				foreach (Assembly Dll in LoadedAssemblies)
 				{
 					Type[] AllTypes = Dll.SafeGetLoadedTypes();
-					foreach (var PotentialConfigType in AllTypes.Where(x => !x.IsAbstract && x != typeof(ICodeSignHelper) && typeof(ICodeSignHelper).IsAssignableFrom(x)))
+					foreach (Type PotentialConfigType in AllTypes.Where(x => !x.IsAbstract && x != typeof(ICodeSignHelper) && typeof(ICodeSignHelper).IsAssignableFrom(x)))
 					{
 						try
 						{
@@ -3238,7 +3238,7 @@ namespace AutomationTool
 						}
 						catch (Exception Ex)
 						{
-							var Inner = Ex.InnerException;
+							Exception Inner = Ex.InnerException;
 							while (null != Inner)
 							{
 								Logger.LogWarning("Exception encountered creating HelperClass [{Name}] with error: {Message}", PotentialConfigType.Name, Inner.Message);
diff --git a/Engine/Source/Programs/AutomationTool/AutomationUtils/CommandletUtils.cs b/Engine/Source/Programs/AutomationTool/AutomationUtils/CommandletUtils.cs
--- a/Engine/Source/Programs/AutomationTool/AutomationUtils/CommandletUtils.cs
+++ b/Engine/Source/Programs/AutomationTool/AutomationUtils/CommandletUtils.cs
@@ -517,7 +517,7 @@ namespace AutomationTool
 			string CombinedParams = String.Empty;
 			if (ParamValues != null)
 			{
-				var Iter = ParamValues.GetEnumerator();
+				IEnumerator<string> Iter = ParamValues.GetEnumerator();
 				while (Iter.MoveNext())
 				{
 					if (!String.IsNullOrEmpty(Iter.Current))
@@ -564,7 +564,7 @@ namespace AutomationTool
 			else
 			{
 				BranchInfo Branch = new BranchInfo();
-				var GameProj = Branch.FindGame(OriginalProjectName);
+				BranchInfo.BranchUProject GameProj = Branch.FindGame(OriginalProjectName);
 				if (GameProj != null)
 				{
 					ProjectFullPath = GameProj.FilePath;
diff --git a/Engine/Source/Programs/AutomationTool/AutomationUtils/DeploymentContext.cs b/Engine/Source/Programs/AutomationTool/AutomationUtils/DeploymentContext.cs
--- a/Engine/Source/Programs/AutomationTool/AutomationUtils/DeploymentContext.cs
+++ b/Engine/Source/Programs/AutomationTool/AutomationUtils/DeploymentContext.cs
@@ -527,8 +527,8 @@ public class DeploymentContext //: ProjectParams
 		// Parse the custom config dir out of the receipts
 		foreach (StageTarget Target in StageTargets)
 		{
-			var Results = Target.Receipt.AdditionalProperties.Where(x => x.Name == "CustomConfig");
-			foreach (var Property in Results)
+			IEnumerable<ReceiptProperty> Results = Target.Receipt.AdditionalProperties.Where(x => x.Name == "CustomConfig");
+			foreach (ReceiptProperty Property in Results)
 			{
 				string FoundCustomConfig = Property.Value;
 				if (String.IsNullOrEmpty(FoundCustomConfig))
diff --git a/Engine/Source/Programs/AutomationTool/AutomationUtils/DeviceReservation.cs b/Engine/Source/Programs/AutomationTool/AutomationUtils/DeviceReservation.cs
--- a/Engine/Source/Programs/AutomationTool/AutomationUtils/DeviceReservation.cs
+++ b/Engine/Source/Programs/AutomationTool/AutomationUtils/DeviceReservation.cs
@@ -202,7 +202,7 @@ namespace AutomationTool.DeviceReservation
 			{
 				Request.ContentType = "application/json";
 
-				using (var RequestStream = Request.GetRequestStream())
+				using (Stream RequestStream = Request.GetRequestStream())
 				{
 					var JsonString = JsonSerializer.Serialize(ObjectToSerialize, SerializeOptions);
 					StreamWriter Writer = new StreamWriter(RequestStream);
@@ -213,7 +213,7 @@ namespace AutomationTool.DeviceReservation
 			}
 
 			using (HttpWebResponse Response = (HttpWebResponse)Request.GetResponse())
-			using (var ResponseStream = Response.GetResponseStream())
+			using (Stream ResponseStream = Response.GetResponseStream())
 			{
 				MemoryStream MemoryStream = new MemoryStream();
 				ResponseStream.CopyTo(MemoryStream);
@@ -232,7 +232,7 @@ namespace AutomationTool.DeviceReservation
 			{
 				Request.ContentType = "application/json";
 
-				using (var RequestStream = Request.GetRequestStream())
+				using (Stream RequestStream = Request.GetRequestStream())
 				{
 					var JsonString = JsonSerializer.Serialize(ObjectToSerialize, new JsonSerializerOptions());
 					StreamWriter Writer = new StreamWriter(RequestStream);
@@ -247,7 +247,7 @@ namespace AutomationTool.DeviceReservation
 			}
 
 			using (HttpWebResponse Response = (HttpWebResponse)Request.GetResponse())
-			using (var ResponseStream = Response.GetResponseStream())
+			using (Stream ResponseStream = Response.GetResponseStream())
 			{
 				// Nothing to do here. Error codes throw exceptions.
 			}
diff --git a/Engine/Source/Programs/AutomationTool/AutomationUtils/FileRegions.cs b/Engine/Source/Programs/AutomationTool/AutomationUtils/FileRegions.cs
--- a/Engine/Source/Programs/AutomationTool/AutomationUtils/FileRegions.cs
+++ b/Engine/Source/Programs/AutomationTool/AutomationUtils/FileRegions.cs
@@ -66,7 +66,7 @@ namespace AutomationTool
 			Logger.LogInformation(" +-------------------+----------------------------------+---------------+");
 			Logger.LogInformation(" |        Type       |           Total Length           |  Num Regions  |");
 			Logger.LogInformation(" +-------------------+----------------------------------+---------------+");
-			foreach (var Pair in RegionsByType)
+			foreach (KeyValuePair<FileRegionType, List<FileRegion>> Pair in RegionsByType)
 			{
 				// Find the total number of bytes covered by regions with this pattern
 				long TotalSizeInBytes = Pair.Value.Sum(r => (long)r.Length);
diff --git a/Engine/Source/Programs/AutomationTool/AutomationUtils/GoogleDriveHelper.cs b/Engine/Source/Programs/AutomationTool/AutomationUtils/GoogleDriveHelper.cs
--- a/Engine/Source/Programs/AutomationTool/AutomationUtils/GoogleDriveHelper.cs
+++ b/Engine/Source/Programs/AutomationTool/AutomationUtils/GoogleDriveHelper.cs
@@ -288,10 +288,10 @@ namespace DriveHelper
 				return;
 			}
 
-			var Request = Service.Files.Get(FileID);
+			FilesResource.GetRequest Request = Service.Files.Get(FileID);
 			Request.SupportsTeamDrives = true;
 			Request.Fields = "size";
-			var Response = Request.Execute();
+			Google.Apis.Drive.v3.Data.File Response = Request.Execute();
 			long FileSize = (Response.Size != null) ? (long)Response.Size : 0;
 			using (FileStream Stream = System.IO.File.Create(DestinationFilePathAndName + ".tmp"))
 			{
@@ -315,10 +315,10 @@ namespace DriveHelper
 				return;
 			}
 
-			var Request = Service.Files.Get(FileID);
+			FilesResource.GetRequest Request = Service.Files.Get(FileID);
 			Request.SupportsTeamDrives = true;
 			Request.Fields = "size";
-			var Response = Request.Execute();
+			Google.Apis.Drive.v3.Data.File Response = Request.Execute();
 			long FileSize = (Response.Size != null) ? (long)Response.Size : 0;
 			using (FileStream Stream = System.IO.File.Create(DestinationFilePathAndName + ".tmp"))
 			{
diff --git a/Engine/Source/Programs/AutomationTool/AutomationUtils/GoogleSheetsHelper.cs b/Engine/Source/Programs/AutomationTool/AutomationUtils/GoogleSheetsHelper.cs
--- a/Engine/Source/Programs/AutomationTool/AutomationUtils/GoogleSheetsHelper.cs
+++ b/Engine/Source/Programs/AutomationTool/AutomationUtils/GoogleSheetsHelper.cs
@@ -162,7 +162,7 @@ namespace SheetsHelper
 		{
 			Refresh();
 
-			var Sheet = SourceSpreadsheet.Sheets.FirstOrDefault(S => S.Properties.Title == Title);
+			Sheet Sheet = SourceSpreadsheet.Sheets.FirstOrDefault(S => S.Properties.Title == Title);
 
 			if (Sheet == null)
 			{
@@ -177,7 +177,7 @@ namespace SheetsHelper
 		{
 			Refresh();
 
-			var Sheet = SourceSpreadsheet.Sheets.FirstOrDefault(S => S.Properties.Title == Title);
+			Sheet Sheet = SourceSpreadsheet.Sheets.FirstOrDefault(S => S.Properties.Title == Title);
 			return Sheet != null;
 		}
 
@@ -435,7 +435,7 @@ namespace SheetsHelper
 		public void Refresh()
 		{
 			SheetsUtils.TryCheckQuotaSync();
-			var Response = Service.Spreadsheets.Get(SpreadsheetID).Execute();
+			Spreadsheet Response = Service.Spreadsheets.Get(SpreadsheetID).Execute();
 			SourceSheet = Response.Sheets.FirstOrDefault(S => S.Properties.SheetId == SheetID);
 		}
 
diff --git a/Engine/Source/Programs/AutomationTool/AutomationUtils/HashCollection.cs b/Engine/Source/Programs/AutomationTool/AutomationUtils/HashCollection.cs
--- a/Engine/Source/Programs/AutomationTool/AutomationUtils/HashCollection.cs
+++ b/Engine/Source/Programs/AutomationTool/AutomationUtils/HashCollection.cs
@@ -116,7 +116,7 @@ namespace AutomationTool
 		/// <param name="HashType"></param>
 		public bool AddFiles(IEnumerable<FileReference> Files, HashCollection.HashType HashType)
 		{
-			foreach (var File in Files)
+			foreach (FileReference File in Files)
 			{
 				if (!AddFile(File, HashType))
 				{
@@ -171,7 +171,7 @@ namespace AutomationTool
 		public void LogDifferences(HashCollection RHS)
 		{
 			ILogger Logger = Log.Logger;
-			foreach (var Entry in Hashes)
+			foreach (HashEntry Entry in Hashes)
 			{
 				if (!RHS.HasEntryForName(Entry.Name))
 				{
@@ -205,7 +205,7 @@ namespace AutomationTool
 		{
 			try
 			{
-				using (var stream = System.IO.File.OpenRead(InPath))
+				using (FileStream stream = System.IO.File.OpenRead(InPath))
 				{
 					XmlSerializer Serializer = new XmlSerializer(typeof(HashCollection));
 					return Serializer.Deserialize(stream) as HashCollection;
diff --git a/Engine/Source/Programs/AutomationTool/AutomationUtils/HelpUtils.cs b/Engine/Source/Programs/AutomationTool/AutomationUtils/HelpUtils.cs
--- a/Engine/Source/Programs/AutomationTool/AutomationUtils/HelpUtils.cs
+++ b/Engine/Source/Programs/AutomationTool/AutomationUtils/HelpUtils.cs
@@ -373,8 +373,8 @@ namespace AutomationTool
 				}
 			}
 
-			var AllMembers = GetAllMemebers(ObjType);
-			foreach (var Member in AllMembers)
+			List<MemberInfo> AllMembers = GetAllMemebers(ObjType);
+			foreach (MemberInfo Member in AllMembers)
 			{
 				var AllMemeberAtributes = Member.GetCustomAttributes(false);
 				foreach (var CustomAttribute in AllMemeberAtributes)
@@ -393,7 +393,7 @@ namespace AutomationTool
 			List<MemberInfo> Members = new List<MemberInfo>();
 			for (; ObjType != null; ObjType = ObjType.BaseType)
 			{
-				var ObjMembers = ObjType.GetMembers(BindingFlags.Public | BindingFlags.NonPublic | BindingFlags.Instance | BindingFlags.Static);
+				MemberInfo[] ObjMembers = ObjType.GetMembers(BindingFlags.Public | BindingFlags.NonPublic | BindingFlags.Instance | BindingFlags.Static);
 				Members.AddRange(ObjMembers);
 			}
 			return Members;
diff --git a/Engine/Source/Programs/AutomationTool/AutomationUtils/HttpClientSingleton.cs b/Engine/Source/Programs/AutomationTool/AutomationUtils/HttpClientSingleton.cs
--- a/Engine/Source/Programs/AutomationTool/AutomationUtils/HttpClientSingleton.cs
+++ b/Engine/Source/Programs/AutomationTool/AutomationUtils/HttpClientSingleton.cs
@@ -2,6 +2,7 @@
 using Microsoft.Extensions.Http;
 using Polly;
 using Polly.Extensions.Http;
+using Polly.Retry;
 using System;
 using System.Net.Http;
 
@@ -14,7 +15,7 @@ namespace AutomationUtils
 		private static readonly Lazy<HttpClient> internalInstance =
 			new Lazy<HttpClient>(() =>
 			{
-				var retryPolicy = HttpPolicyExtensions
+				AsyncRetryPolicy<HttpResponseMessage> retryPolicy = HttpPolicyExtensions
 				// 408, 5XX responses 
 				.HandleTransientHttpError()
 				// 404, 504 should have been covered by HandleTransientHttpError but adding here to be safe
diff --git a/Engine/Source/Programs/AutomationTool/AutomationUtils/LogUtils.cs b/Engine/Source/Programs/AutomationTool/AutomationUtils/LogUtils.cs
--- a/Engine/Source/Programs/AutomationTool/AutomationUtils/LogUtils.cs
+++ b/Engine/Source/Programs/AutomationTool/AutomationUtils/LogUtils.cs
@@ -87,7 +87,7 @@ namespace AutomationTool
 		{
 			const int MaxAttempts = 1000;
 
-			var Now     = DateTime.Now;
+			DateTime Now     = DateTime.Now;
 			int Attempt = 0;
 
 			string LogFilename;
diff --git a/Engine/Source/Programs/AutomationTool/AutomationUtils/MCPPublic.cs b/Engine/Source/Programs/AutomationTool/AutomationUtils/MCPPublic.cs
--- a/Engine/Source/Programs/AutomationTool/AutomationUtils/MCPPublic.cs
+++ b/Engine/Source/Programs/AutomationTool/AutomationUtils/MCPPublic.cs
@@ -1731,10 +1731,10 @@ namespace EpicGames.MCP.Automation
 			if (Handler == null)
 			{
 				Assembly[] LoadedAssemblies = AppDomain.CurrentDomain.GetAssemblies();
-				foreach (var Dll in LoadedAssemblies)
+				foreach (Assembly Dll in LoadedAssemblies)
 				{
 					Type[] AllTypes = Dll.SafeGetLoadedTypes();
-					foreach (var PotentialConfigType in AllTypes)
+					foreach (Type PotentialConfigType in AllTypes)
 					{
 						if (PotentialConfigType != typeof(BuildPatchToolBase) && typeof(BuildPatchToolBase).IsAssignableFrom(PotentialConfigType))
 						{
@@ -1887,10 +1887,10 @@ namespace EpicGames.MCP.Automation
 			if (Handler == null)
 			{
 				Assembly[] LoadedAssemblies = AppDomain.CurrentDomain.GetAssemblies();
-				foreach (var Dll in LoadedAssemblies)
+				foreach (Assembly Dll in LoadedAssemblies)
 				{
 					Type[] AllTypes = Dll.SafeGetLoadedTypes();
-					foreach (var PotentialConfigType in AllTypes)
+					foreach (Type PotentialConfigType in AllTypes)
 					{
 						if (PotentialConfigType != typeof(BuildInfoPublisherBase) && typeof(BuildInfoPublisherBase).IsAssignableFrom(PotentialConfigType))
 						{
@@ -2046,10 +2046,10 @@ namespace EpicGames.MCP.Automation
             if (Handler == null)
             {
                 Assembly[] LoadedAssemblies = AppDomain.CurrentDomain.GetAssemblies();
-                foreach (var Dll in LoadedAssemblies)
+                foreach (Assembly Dll in LoadedAssemblies)
                 {
                     Type[] AllTypes = Dll.SafeGetLoadedTypes();
-                    foreach (var PotentialConfigType in AllTypes)
+                    foreach (Type PotentialConfigType in AllTypes)
                     {
                         if (PotentialConfigType != typeof(McpAccountServiceBase) && typeof(McpAccountServiceBase).IsAssignableFrom(PotentialConfigType))
                         {
@@ -2195,10 +2195,10 @@ namespace EpicGames.MCP.Automation
 					else
 					{
 						Assembly[] LoadedAssemblies = AppDomain.CurrentDomain.GetAssemblies();
-						foreach (var Dll in LoadedAssemblies)
+						foreach (Assembly Dll in LoadedAssemblies)
 						{
 							Type[] AllTypes = Dll.SafeGetLoadedTypes();
-							foreach (var PotentialConfigType in AllTypes)
+							foreach (Type PotentialConfigType in AllTypes)
 							{
 								if (PotentialConfigType != typeof(CloudStorageBase) && typeof(CloudStorageBase).IsAssignableFrom(PotentialConfigType))
 								{
@@ -2264,7 +2264,7 @@ namespace EpicGames.MCP.Automation
 		/// <param name="bGzip">If true, the data to be uploaded will be gzipped, and the appropraite Content-Encoding value set on the uploaded file in S3.</param>
 		/// <param name="bOverwrite">If true, will overwrite an existing file.  If false, will throw an exception if the file exists.</param>
 		/// <param name="bMakePublic">Specifies whether the file should be made publicly readable.</param>
-		/// <param name="bQuiet">If set to true, all log output for the operation is supressed.</param>
+		/// <param name="bQuiet">If set to true, all log output for the operation is suppressed.</param>
 		/// <param name="Metadata">If not null, key-value pairs of metadata to be applied to the object.</param>
 		/// <returns>A PostFileResult indicating whether the call was successful, and the URL to the uploaded file.</returns>
 		abstract public PostFileResult PostFile(string Container, string Identifier, byte[] Contents, string ContentType = null, bool bGzip = false, bool bOverwrite = true, bool bMakePublic = false, bool bQuiet = false, IDictionary<string, object> Metadata = null);
@@ -2279,7 +2279,7 @@ namespace EpicGames.MCP.Automation
 		/// <param name="bGzip">If true, the data to be uploaded will be gzipped, and the appropraite Content-Encoding value set on the uploaded file in S3.</param>
 		/// <param name="bOverwrite">If true, will overwrite an existing file.  If false, will throw an exception if the file exists.</param>
 		/// <param name="bMakePublic">Specifies whether the file should be made publicly readable.</param>
-        /// <param name="bQuiet">If set to true, all log output for the operation is supressed.</param>
+        /// <param name="bQuiet">If set to true, all log output for the operation is suppressed.</param>
 		/// <param name="Metadata">If not null, key-value pairs of metadata to be applied to the object.</param>
 		/// <returns>A PostFileResult indicating whether the call was successful, and the URL to the uploaded file.</returns>
 		abstract public Task<PostFileResult> PostFileAsync(string Container, string Identifier, byte[] Contents, string ContentType = null, bool bGzip = false, bool bOverwrite = true, bool bMakePublic = false, bool bQuiet = false, IDictionary<string, object> Metadata = null);
@@ -2294,7 +2294,7 @@ namespace EpicGames.MCP.Automation
 		/// <param name="bGzip">If true, the data to be uploaded will be gzipped, and the appropraite Content-Encoding value set on the uploaded file in S3.</param>
 		/// <param name="bOverwrite">If true, will overwrite an existing file.  If false, will throw an exception if the file exists.</param>
 		/// <param name="bMakePublic">Specifies whether the file should be made publicly readable.</param>
-		/// <param name="bQuiet">If set to true, all log output for the operation is supressed.</param>
+		/// <param name="bQuiet">If set to true, all log output for the operation is suppressed.</param>
 		/// <param name="Metadata">If not null, key-value pairs of metadata to be applied to the object.</param>
 		/// <returns>A PostFileResult indicating whether the call was successful, and the URL to the uploaded file.</returns>
 		abstract public PostFileResult PostFile(string Container, string Identifier, string SourceFilePath, string ContentType = null, bool bGzip = false, bool bOverwrite = true, bool bMakePublic = false, bool bQuiet = false, IDictionary<string, object> Metadata = null);
@@ -2309,7 +2309,7 @@ namespace EpicGames.MCP.Automation
 		/// <param name="bGzip">If true, the data to be uploaded will be gzipped, and the appropraite Content-Encoding value set on the uploaded file in S3.</param>
 		/// <param name="bOverwrite">If true, will overwrite an existing file.  If false, will throw an exception if the file exists.</param>
 		/// <param name="bMakePublic">Specifies whether the file should be made publicly readable.</param>
-        /// <param name="bQuiet">If set to true, all log output for the operation is supressed.</param>
+        /// <param name="bQuiet">If set to true, all log output for the operation is suppressed.</param>
 		/// <param name="Metadata">If not null, key-value pairs of metadata to be applied to the object.</param>
 		/// <returns>A PostFileResult indicating whether the call was successful, and the URL to the uploaded file.</returns>
 		abstract public Task<PostFileResult> PostFileAsync(string Container, string Identifier, string SourceFilePath, string ContentType = null, bool bGzip = false, bool bOverwrite = true, bool bMakePublic = false, bool bQuiet = false, IDictionary<string, object> Metadata = null);
@@ -2326,7 +2326,7 @@ namespace EpicGames.MCP.Automation
 		/// <param name="bGzip">If true, the data to be uploaded will be gzipped, and the appropraite Content-Encoding value set on the uploaded file in S3.</param>
 		/// <param name="bOverwrite">If true, will overwrite an existing file. If false, will throw an exception if the file exists.</param>
 		/// <param name="bMakePublic">Specifies whether the file should be made publicly readable.</param>
-		/// <param name="bQuiet">If set to true, all log output for the operation is supressed.</param>
+		/// <param name="bQuiet">If set to true, all log output for the operation is suppressed.</param>
 		/// <param name="Metadata">If not null, key-value pairs of metadata to be applied to the object.</param>
 		/// <returns>A PostFileResult indicating whether the call was successful, and the URL to the uploaded file.</returns>
 		public PostFileResult PostMultipartFile(string Container, string Identifier, string SourceFilePath, int NumConcurrentConnections, decimal PartSizeMegabytes = 5.0m, string ContentType = null, bool bGzip = false, bool bOverwrite = true, bool bMakePublic = false, bool bQuiet = false, IDictionary<string, object> Metadata = null)
@@ -2346,7 +2346,7 @@ namespace EpicGames.MCP.Automation
 		/// <param name="bGzip">If true, the data to be uploaded will be gzipped, and the appropraite Content-Encoding value set on the uploaded file in S3.</param>
 		/// <param name="bOverwrite">If true, will overwrite an existing file. If false, will throw an exception if the file exists.</param>
 		/// <param name="bMakePublic">Specifies whether the file should be made publicly readable.</param>
-		/// <param name="bQuiet">If set to true, all log output for the operation is supressed.</param>
+		/// <param name="bQuiet">If set to true, all log output for the operation is suppressed.</param>
 		/// <param name="Metadata">If not null, key-value pairs of metadata to be applied to the object.</param>
 		/// <returns>A PostFileResult indicating whether the call was successful, and the URL to the uploaded file.</returns>
 		abstract public Task<PostFileResult> PostMultipartFileAsync(string Container, string Identifier, string SourceFilePath, int NumConcurrentConnections, decimal PartSizeMegabytes = 5.0m, string ContentType = null, bool bGzip = false, bool bOverwrite = true, bool bMakePublic = false, bool bQuiet = false, IDictionary<string, object> Metadata = null);
@@ -2399,7 +2399,7 @@ namespace EpicGames.MCP.Automation
 		/// <param name="Container">The name of the folder or container from which to list files.</param>
 		/// <param name="Prefix">A string with which the identifier or filename should start. Typically used to specify a relative directory within the container to list all of its files recursively. Specify null to return all files.</param>
 		/// <param name="bRecursive">Indicates whether the list of files returned should traverse subdirectories</param>
-		/// <param name="bQuiet">If set to true, all log output for the operation is supressed.</param>
+		/// <param name="bQuiet">If set to true, all log output for the operation is suppressed.</param>
 		/// <returns>An array of paths to the files in the specified location and matching the prefix constraint.</returns>
 		public string[] ListFiles(string Container, string Prefix = null, bool bRecursive = true, bool bQuiet = false)
 		{
@@ -2585,10 +2585,10 @@ namespace EpicGames.MCP.Automation
 			if (Handler == null)
 			{
 				Assembly[] LoadedAssemblies = AppDomain.CurrentDomain.GetAssemblies();
-				foreach (var Dll in LoadedAssemblies)
+				foreach (Assembly Dll in LoadedAssemblies)
 				{
 					Type[] AllTypes = Dll.SafeGetLoadedTypes();
-					foreach (var PotentialConfigType in AllTypes)
+					foreach (Type PotentialConfigType in AllTypes)
 					{
 						if (PotentialConfigType != typeof(CatalogServiceBase) && typeof(CatalogServiceBase).IsAssignableFrom(PotentialConfigType))
 						{
@@ -2652,10 +2652,10 @@ namespace EpicGames.MCP.Config
                 // Note that we're using the default constructor on the secret config types.
                 Configs = new Dictionary<string, McpConfigData>();
                 Assembly[] LoadedAssemblies = AppDomain.CurrentDomain.GetAssemblies();
-                foreach (var Dll in LoadedAssemblies)
+                foreach (Assembly Dll in LoadedAssemblies)
                 {
                     Type[] AllTypes = Dll.SafeGetLoadedTypes();
-                    foreach (var PotentialConfigType in AllTypes)
+                    foreach (Type PotentialConfigType in AllTypes)
                     {
 						// do not attempt to instantiate abstract (base) configs
 						if(PotentialConfigType.IsAbstract)
@@ -2676,7 +2676,7 @@ namespace EpicGames.MCP.Config
                             }
                             catch (Exception Ex)
                             {
-                               	var Inner = Ex.InnerException;
+								Exception Inner = Ex.InnerException;
 								while(null != Inner)
 								{
 									Logger.LogWarning("Exception encountered creating McpConfig [{Name}] with error: {Message}", PotentialConfigType.Name, Inner.Message);
diff --git a/Engine/Source/Programs/AutomationTool/AutomationUtils/P4Environment.cs b/Engine/Source/Programs/AutomationTool/AutomationUtils/P4Environment.cs
--- a/Engine/Source/Programs/AutomationTool/AutomationUtils/P4Environment.cs
+++ b/Engine/Source/Programs/AutomationTool/AutomationUtils/P4Environment.cs
@@ -309,14 +309,14 @@ namespace AutomationTool
 		public static string DetectUserName(P4Connection Connection)
 		{
 			var UserName = String.Empty;
-			var P4Result = Connection.P4("info", AllowSpew: false);
+			IProcessResult P4Result = Connection.P4("info", AllowSpew: false);
 			if (P4Result.ExitCode != 0)
 			{
 				throw new AutomationException("Perforce command failed: {0}. Please make sure your P4PORT or {1} is set properly.", P4Result.Output, EnvVarNames.P4Port);
 			}
 
 			// Retrieve the P4 user name			
-			var Tags = Connection.ParseTaggedP4Output(P4Result.Output);
+			Dictionary<string, string> Tags = Connection.ParseTaggedP4Output(P4Result.Output);
 			if(!Tags.TryGetValue("User name", out UserName) || String.IsNullOrEmpty(UserName))
 			{
 				if (!String.IsNullOrEmpty(UserName))
@@ -413,8 +413,8 @@ namespace AutomationTool
 		{
 			Logger.LogDebug("Detecting the most recent client.");
 			P4ClientInfo MostRecentClient = null;
-			var MostRecentAccessTime = DateTime.MinValue;
-			foreach (var ClientInfo in Clients)
+			DateTime MostRecentAccessTime = DateTime.MinValue;
+			foreach (P4ClientInfo ClientInfo in Clients)
 			{
 				if (ClientInfo.Access.Ticks > MostRecentAccessTime.Ticks)
 				{
diff --git a/Engine/Source/Programs/AutomationTool/AutomationUtils/P4Utils.cs b/Engine/Source/Programs/AutomationTool/AutomationUtils/P4Utils.cs
--- a/Engine/Source/Programs/AutomationTool/AutomationUtils/P4Utils.cs
+++ b/Engine/Source/Programs/AutomationTool/AutomationUtils/P4Utils.cs
@@ -1143,7 +1143,7 @@ namespace AutomationTool
 			bRequireP4 = false;
 			bRequireCL = false;
 
-			foreach (var CommandInfo in CommandsToExecute)
+			foreach (CommandInfo CommandInfo in CommandsToExecute)
 			{
 				Type Command;
 				if (Commands.TryGetValue(CommandInfo.CommandName, out Command))
@@ -1281,7 +1281,7 @@ namespace AutomationTool
 				RunOptions |= CommandUtils.ERunOptions.SpewIsVerbose;
 			}
 
-			var SpewDelegate = AllowSpew ? null : new ProcessResult.SpewFilterCallbackType(NoSpewFilter);
+			ProcessResult.SpewFilterCallbackType SpewDelegate = AllowSpew ? null : new ProcessResult.SpewFilterCallbackType(NoSpewFilter);
 
 			IProcessResult Result;
 			// if there's a star anywhere in the commandline, p4.exe, when parsing command line on Windows, will internally perform a find-files to expand the *,
@@ -1334,7 +1334,7 @@ namespace AutomationTool
 		{
 			Output = "";
 
-            var Result = P4(ExtraGlobalOptions, CommandLine, Input, AllowSpew, WithClient);
+			IProcessResult Result = P4(ExtraGlobalOptions, CommandLine, Input, AllowSpew, WithClient);
 
 			Output = Result.Output;
 			return Result.ExitCode == 0;
@@ -1420,7 +1420,7 @@ namespace AutomationTool
 				return false;
 			}
 
-            var Result = P4(ExtraGlobalOptions, CommandLine, Input, AllowSpew, WithClient, SpewIsVerbose:SpewIsVerbose);
+			IProcessResult Result = P4(ExtraGlobalOptions, CommandLine, Input, AllowSpew, WithClient, SpewIsVerbose:SpewIsVerbose);
 
 			CommandUtils.WriteToFile(LogPath, CommandLine + "\n");
 			CommandUtils.WriteToFile(LogPath, Result.Output);
@@ -1623,10 +1623,10 @@ namespace AutomationTool
             string Result = "";
             try
             {
-                var P4Result = P4(String.Format("user -o {0}", User), AllowSpew: false);
+				IProcessResult P4Result = P4(String.Format("user -o {0}", User), AllowSpew: false);
 			    if (P4Result.ExitCode == 0)
 			    {
-				    var Tags = ParseTaggedP4Output(P4Result.Output);
+					Dictionary<string, string> Tags = ParseTaggedP4Output(P4Result.Output);
                     Tags.TryGetValue("Email", out Result);
                 }
             }
@@ -3072,7 +3072,7 @@ namespace AutomationTool
 				throw new InvalidOperationException("The depot should have at least one submitted changelist. Brand new depot?");
 			}
 
-			var Match = ChangesListOutputPattern.Match(Output);
+			Match Match = ChangesListOutputPattern.Match(Output);
 
 			if (!Match.Success)
 			{
@@ -3398,7 +3398,7 @@ namespace AutomationTool
 		{
 			Logger.LogDebug("ChangeFileType({Filename}, {Attributes}, {Arg2})", Filename, Attributes, String.IsNullOrEmpty(Changelist) ? "null" : Changelist);
 
-			var Stat = FStat(Filename);
+			P4FileStat Stat = FStat(Filename);
 			if (String.IsNullOrEmpty(Changelist))
 			{
 				Changelist = (Stat.Action != P4Action.None) ? Stat.Change : "default";
@@ -3508,7 +3508,7 @@ namespace AutomationTool
 				Logger.LogDebug("Checking if client {ClientName} exists", ClientName);
 			}
 
-            var P4Result = P4(String.Format("-c {0}", ClientName), "where //...", Input: null, AllowSpew: false, WithClient: false);
+			IProcessResult P4Result = P4(String.Format("-c {0}", ClientName), "where //...", Input: null, AllowSpew: false, WithClient: false);
             return P4Result.Output.IndexOf("unknown - use 'client' command", StringComparison.InvariantCultureIgnoreCase) < 0 && P4Result.Output.IndexOf("doesn't exist", StringComparison.InvariantCultureIgnoreCase) < 0;
 		}
 
@@ -3551,10 +3551,10 @@ namespace AutomationTool
 		public P4ClientInfo GetClientInfoInternal(string ClientName)
 		{
 			P4ClientInfo Info = new P4ClientInfo();
-            var P4Result = P4(String.Format("client -o {0}", ClientName), AllowSpew: false, WithClient: false);
+			IProcessResult P4Result = P4(String.Format("client -o {0}", ClientName), AllowSpew: false, WithClient: false);
 			if (P4Result.ExitCode == 0)
 			{
-				var Tags = ParseTaggedP4Output(P4Result.Output);
+				Dictionary<string, string> Tags = ParseTaggedP4Output(P4Result.Output);
                 Info.Name = ClientName;
 				Tags.TryGetValue("Host", out Info.Host);
 				Tags.TryGetValue("Root", out Info.RootPath);
@@ -3593,7 +3593,7 @@ namespace AutomationTool
                     Info.SubmitOptions = (P4SubmitOption)ParseEnumValues(SubmitOptions, typeof(P4SubmitOption));
                 }
                 string ClientMappingRoot = "//" + ClientName;
-                foreach (var Pair in Tags)
+                foreach (KeyValuePair<string, string> Pair in Tags)
                 {
                     if (Pair.Key.StartsWith("View", StringComparison.InvariantCultureIgnoreCase))
                     {
@@ -3726,7 +3726,7 @@ namespace AutomationTool
 			else
 			{
 				SpecInput += "View:" + Environment.NewLine;
-				foreach (var Mapping in ClientSpec.View)
+				foreach (KeyValuePair<string, string> Mapping in ClientSpec.View)
 				{
 					SpecInput += "\t" + Mapping.Key + " //" + ClientSpec.Name + Mapping.Value + Environment.NewLine;
 				}
@@ -3745,7 +3745,7 @@ namespace AutomationTool
 		public List<string> Dirs(string CommandLine)
 		{
 			var DirsCmdLine = String.Format("dirs {0}", CommandLine);
-			var P4Result = P4(DirsCmdLine, AllowSpew: false);
+			IProcessResult P4Result = P4(DirsCmdLine, AllowSpew: false);
 			if (P4Result.ExitCode != 0)
 			{
 				throw new AutomationException("{0} failed.", DirsCmdLine);
@@ -3805,7 +3805,7 @@ namespace AutomationTool
 			List<string> Result = new List<string>();
 			Dictionary<string, string> FileActions = MapFileOutputToActions(P4Result.Output);
 
-			foreach (var KV in FileActions)
+			foreach (KeyValuePair<string, string> KV in FileActions)
 			{
 				if (!DeleteActions.Contains(KV.Value))
 				{
@@ -3833,7 +3833,7 @@ namespace AutomationTool
 			List<string> Result = new List<string>();
 			Dictionary<string, string> FileActions = MapFileOutputToActions(P4Result.Output);
 
-			foreach (var KV in FileActions)
+			foreach (KeyValuePair<string, string> KV in FileActions)
 			{
 				if (!DeleteActions.Contains(KV.Value))
 				{
@@ -4369,15 +4369,15 @@ namespace AutomationTool
 
 		private static void ParseFileType(string Filetype, ref P4FileStat Stat)
 		{
-			var AllFileTypes = GetEnumValuesAndKeywords(typeof(P4FileType));
-			var AllAttributes = GetEnumValuesAndKeywords(typeof(P4FileAttributes));
+			KeyValuePair<object, string>[] AllFileTypes = GetEnumValuesAndKeywords(typeof(P4FileType));
+			KeyValuePair<object, string>[] AllAttributes = GetEnumValuesAndKeywords(typeof(P4FileAttributes));
 
 			Stat.Type = P4FileType.Unknown;
 			Stat.Attributes = P4FileAttributes.None;
 
 			// Parse file flags
-			var OldFileFlags = GetEnumValuesAndKeywords(typeof(P4FileAttributes), OldStyleBinaryFlags);
-			foreach (var FileTypeFlag in OldFileFlags)
+			KeyValuePair<object, string>[] OldFileFlags = GetEnumValuesAndKeywords(typeof(P4FileAttributes), OldStyleBinaryFlags);
+			foreach (KeyValuePair<object, string> FileTypeFlag in OldFileFlags)
 			{
 				if ((!String.IsNullOrEmpty(FileTypeFlag.Value) && Char.ToLowerInvariant(FileTypeFlag.Value[0]) == Char.ToLowerInvariant(Filetype[0]))
 					// @todo: This is a nasty hack to get .ipa files to work - RobM plz fix?
@@ -4394,7 +4394,7 @@ namespace AutomationTool
 			}
 			// Parse file type
 			var TypeAndAttributes = Filetype.Split('+');
-			foreach (var FileType in AllFileTypes)
+			foreach (KeyValuePair<object, string> FileType in AllFileTypes)
 			{
 				if (FileType.Value == TypeAndAttributes[0])
 				{
@@ -4409,7 +4409,7 @@ namespace AutomationTool
 				for (int AttributeIndex = 0; AttributeIndex < FileAttributes.Length; ++AttributeIndex)
 				{
 					char Attr = FileAttributes[AttributeIndex];
-					foreach (var FileAttribute in AllAttributes)
+					foreach (KeyValuePair<object, string> FileAttribute in AllAttributes)
 					{
 						if (!String.IsNullOrEmpty(FileAttribute.Value) && FileAttribute.Value[0] == Attr)
 						{
@@ -4424,8 +4424,8 @@ namespace AutomationTool
 		static P4Action ParseAction(string Action)
 		{
 			P4Action Result = P4Action.Unknown;
-			var AllActions = GetEnumValuesAndKeywords(typeof(P4Action));
-			foreach (var ActionKeyword in AllActions)
+			KeyValuePair<object, string>[] AllActions = GetEnumValuesAndKeywords(typeof(P4Action));
+			foreach (KeyValuePair<object, string> ActionKeyword in AllActions)
 			{
 				if (ActionKeyword.Value == Action)
 				{
@@ -4438,7 +4438,7 @@ namespace AutomationTool
 
 		private static KeyValuePair<object, string>[] GetEnumValuesAndKeywords(Type EnumType)
 		{
-			var Values = Enum.GetValues(EnumType);
+			Array Values = Enum.GetValues(EnumType);
 			KeyValuePair<object, string>[] ValuesAndKeywords = new KeyValuePair<object, string>[Values.Length];
 			int ValueIndex = 0;
 			foreach (var Value in Values)
@@ -4461,16 +4461,16 @@ namespace AutomationTool
 
 		private static string GetEnumDescription(Type EnumType, object Value)
 		{
-			var MemberInfo = EnumType.GetMember(Value.ToString());
+			MemberInfo[] MemberInfo = EnumType.GetMember(Value.ToString());
 			var Atributes = MemberInfo[0].GetCustomAttributes(typeof(DescriptionAttribute), false);
 			return ((DescriptionAttribute)Atributes[0]).Description;
 		}
 
 		private static string FileAttributesToString(P4FileAttributes Attributes)
 		{
-			var AllAttributes = GetEnumValuesAndKeywords(typeof(P4FileAttributes));
+			KeyValuePair<object, string>[] AllAttributes = GetEnumValuesAndKeywords(typeof(P4FileAttributes));
 			string Text = "";
-			foreach (var Attr in AllAttributes)
+			foreach (KeyValuePair<object, string> Attr in AllAttributes)
 			{
 				P4FileAttributes AttrValue = (P4FileAttributes)Attr.Key;
 				if ((Attributes & AttrValue) == AttrValue)
diff --git a/Engine/Source/Programs/AutomationTool/AutomationUtils/PackageUtils.cs b/Engine/Source/Programs/AutomationTool/AutomationUtils/PackageUtils.cs
--- a/Engine/Source/Programs/AutomationTool/AutomationUtils/PackageUtils.cs
+++ b/Engine/Source/Programs/AutomationTool/AutomationUtils/PackageUtils.cs
@@ -76,7 +76,7 @@ namespace AutomationTool
 
 		private static void GetFileListByPatchLayers(FileInfo[] PakFiles, out List<string>[] PatchLayers)
 		{
-			var SortedLayers = SortFilesByPatchLayers(PakFiles);
+			List<FileInfo>[] SortedLayers = SortFilesByPatchLayers(PakFiles);
 
 			PatchLayers = SortedLayers.Select(patchLayer => patchLayer.Select(fileInfo => fileInfo.Name).ToList()).ToArray();
 		}
@@ -94,7 +94,7 @@ namespace AutomationTool
 		/// <returns></returns>
 		public static void ExtractPakFiles(DirectoryInfo SourceDirectory, string TargetDirectory, string CryptoKeysFilename, string AdditionalArgs, bool bExtractByLayers, FileReference ProjectFile=null)
 		{
-			var PakFiles = SourceDirectory.GetFiles("*.pak", SearchOption.TopDirectoryOnly);
+			FileInfo[] PakFiles = SourceDirectory.GetFiles("*.pak", SearchOption.TopDirectoryOnly);
 
 			if (bExtractByLayers)
 			{
diff --git a/Engine/Source/Programs/AutomationTool/AutomationUtils/Platform.cs b/Engine/Source/Programs/AutomationTool/AutomationUtils/Platform.cs
--- a/Engine/Source/Programs/AutomationTool/AutomationUtils/Platform.cs
+++ b/Engine/Source/Programs/AutomationTool/AutomationUtils/Platform.cs
@@ -114,7 +114,7 @@ namespace AutomationTool
 			Logger.LogDebug("Creating platforms.");
 
 			// Create all available platforms.
-			foreach (var ScriptAssembly in AssembliesWithPlatforms)
+			foreach (Assembly ScriptAssembly in AssembliesWithPlatforms)
 			{
 				CreatePlatformsFromAssembly(ScriptAssembly);
 			}
@@ -148,7 +148,7 @@ namespace AutomationTool
 					if (!IsNullOrEmpty(TypeLoadException.LoaderExceptions))
 					{
 						Logger.LogError("Loader Exceptions:");
-						foreach (var LoaderException in TypeLoadException.LoaderExceptions)
+						foreach (Exception LoaderException in TypeLoadException.LoaderExceptions)
 						{
 							Logger.LogError(LoaderException, "{Text}", LogUtils.FormatException(LoaderException));
 						}
@@ -161,13 +161,13 @@ namespace AutomationTool
 				// Re-throw, this is still a critical error!
 				throw;
 			}
-			foreach (var PotentialPlatformType in AllTypes)
+			foreach (Type PotentialPlatformType in AllTypes)
 			{
 				if (PotentialPlatformType != typeof(Platform) && typeof(Platform).IsAssignableFrom(PotentialPlatformType) && !PotentialPlatformType.IsAbstract)
 				{
 					Logger.LogDebug("Creating platform {Platform} from {Location}.", PotentialPlatformType.Name, ScriptAssembly.Location);
 					Platform PlatformInstance = Activator.CreateInstance(PotentialPlatformType) as Platform;
-					var PlatformDesc = PlatformInstance.GetTargetPlatformDescriptor();
+					TargetPlatformDescriptor PlatformDesc = PlatformInstance.GetTargetPlatformDescriptor();
 
 					Platform ExistingInstance;
 					if (!AllPlatforms.TryGetValue(PlatformDesc, out ExistingInstance))
diff --git a/Engine/Source/Programs/AutomationTool/AutomationUtils/ProcessUtils.cs b/Engine/Source/Programs/AutomationTool/AutomationUtils/ProcessUtils.cs
--- a/Engine/Source/Programs/AutomationTool/AutomationUtils/ProcessUtils.cs
+++ b/Engine/Source/Programs/AutomationTool/AutomationUtils/ProcessUtils.cs
@@ -188,10 +188,10 @@ namespace AutomationTool
 		/// <returns>New Process objects</returns>
 		public static IProcessResult CreateProcess(string AppName, bool bAllowSpew, bool bCaptureSpew, Dictionary<string, string> Env = null, LogEventType SpewVerbosity = LogEventType.Console, ProcessResult.SpewFilterCallbackType SpewFilterCallback = null, string WorkingDir = null)
 		{
-			var NewProcess = HostPlatform.Current.CreateProcess(AppName);
+			Process NewProcess = HostPlatform.Current.CreateProcess(AppName);
 			if (Env != null)
 			{
-				foreach (var EnvPair in Env)
+				foreach (KeyValuePair<string, string> EnvPair in Env)
 				{
 					if (NewProcess.StartInfo.Environment.ContainsKey(EnvPair.Key))
 					{
@@ -275,7 +275,7 @@ namespace AutomationTool
 			if(ProcessesToKill.Count > 0)
 			{
 				Logger.LogDebug("Trying to kill {Arg0} spawned processes.", ProcessesToKill.Count);
-				foreach (var Proc in ProcessesToKill)
+				foreach (IProcess Proc in ProcessesToKill)
 				{
 					try
 					{
@@ -292,7 +292,7 @@ namespace AutomationTool
 					for (int Cnt = 0; Cnt < 9; Cnt++)
 					{
 						bool AllDone = true;
-						foreach (var Proc in ProcessesToKill)
+						foreach (IProcess Proc in ProcessesToKill)
 						{
 							try
 							{
@@ -329,7 +329,7 @@ namespace AutomationTool
 						Thread.Sleep(10000);
 					}
 				}
-				foreach (var Proc in ProcessesToKill)
+				foreach (IProcess Proc in ProcessesToKill)
 				{
 					var ProcName = string.Empty;
 					try
@@ -809,7 +809,7 @@ namespace AutomationTool
 				// For some reason Process.GetProcesses() sometimes returns the same process twice
 				// So keep track of all processes we already tried to kill
 				HashSet<int> KilledPids = new HashSet<int>();
-				var AllProcs = Process.GetProcesses();
+				Process[] AllProcs = Process.GetProcesses();
 				foreach (Process KillCandidate in AllProcs)
 				{
 					HashSet<int> VisitedPids = new HashSet<int>();
@@ -933,7 +933,7 @@ namespace AutomationTool
 		{
 			lock(ExeToTimeInMs)
 			{
-				foreach (var Item in ExeToTimeInMs)
+				foreach (KeyValuePair<string, int> Item in ExeToTimeInMs)
 				{
 					Logger.LogDebug("Total {Time}s to run {Exe}", Item.Value / 1000, Item.Key);
 				}
@@ -1114,7 +1114,7 @@ namespace AutomationTool
 
 				App = ResolvedPath;
 			}
-			var StartTime = DateTime.UtcNow;
+			DateTime StartTime = DateTime.UtcNow;
 
 			LogEventType SpewVerbosity = Options.HasFlag(ERunOptions.SpewIsVerbose) ? LogEventType.Verbose : LogEventType.Console;
 			if (!Options.HasFlag(ERunOptions.NoLoggingOfRunCommand))
diff --git a/Engine/Source/Programs/AutomationTool/AutomationUtils/ProjectParams.cs b/Engine/Source/Programs/AutomationTool/AutomationUtils/ProjectParams.cs
--- a/Engine/Source/Programs/AutomationTool/AutomationUtils/ProjectParams.cs
+++ b/Engine/Source/Programs/AutomationTool/AutomationUtils/ProjectParams.cs
@@ -2621,7 +2621,7 @@ namespace AutomationTool
 			List<UnrealTargetPlatform> ClientTargetPlatformTypes = ClientTargetPlatforms.ConvertAll(x => x.Type).Distinct().ToList();
 			// @todo (wip) - Removing Blueprint nativization as a feature.
 			bool bRunAssetNativization = false;// this.RunAssetNativization;
-			var Properties = ProjectUtils.GetProjectProperties(RawProjectPath, ClientTargetPlatformTypes, ClientConfigsToBuild, bRunAssetNativization);
+			ProjectProperties Properties = ProjectUtils.GetProjectProperties(RawProjectPath, ClientTargetPlatformTypes, ClientConfigsToBuild, bRunAssetNativization);
 
 			bIsCodeBasedProject = Properties.bIsCodeBasedProject;
 			DetectedTargets = Properties.Targets;
@@ -2633,7 +2633,7 @@ namespace AutomationTool
 			var ServerTarget = String.Empty;
 			var ProgramTarget = String.Empty;
 
-			var ProjectType = TargetType.Game;
+			TargetType ProjectType = TargetType.Game;
 
 			if (!bIsCodeBasedProject)
 			{
@@ -2850,7 +2850,7 @@ namespace AutomationTool
 				ProjectPlatformBinariesPaths = new Dictionary<UnrealTargetPlatform, DirectoryReference>();
 				ProjectExePaths = new Dictionary<UnrealTargetPlatform, FileReference>();
 
-				var ProjectClientBinariesPath = ProjectUtils.GetClientProjectBinariesRootPath(RawProjectPath, ProjectType, Properties.bIsCodeBasedProject);
+				DirectoryReference ProjectClientBinariesPath = ProjectUtils.GetClientProjectBinariesRootPath(RawProjectPath, ProjectType, Properties.bIsCodeBasedProject);
 
 				foreach (TargetPlatformDescriptor TargetPlatform in ClientTargetPlatforms)
 				{
@@ -3163,7 +3163,7 @@ namespace AutomationTool
 			get
 			{
 				List<Platform> ClientPlatformInstances = new List<Platform>();
-				foreach ( var ClientPlatform in ClientTargetPlatforms )
+				foreach (TargetPlatformDescriptor ClientPlatform in ClientTargetPlatforms )
 				{
 					ClientPlatformInstances.Add(Platform.Platforms[ClientPlatform]);
 				}
@@ -3185,7 +3185,7 @@ namespace AutomationTool
 			get
 			{
 				List<Platform> ServerPlatformInstances = new List<Platform>();
-				foreach (var ServerPlatform in ServerTargetPlatforms)
+				foreach (TargetPlatformDescriptor ServerPlatform in ServerTargetPlatforms)
 				{
 					ServerPlatformInstances.Add(Platform.Platforms[ServerPlatform]);
 				}
diff --git a/Engine/Source/Programs/AutomationTool/AutomationUtils/SourceIndexingUtils.cs b/Engine/Source/Programs/AutomationTool/AutomationUtils/SourceIndexingUtils.cs
--- a/Engine/Source/Programs/AutomationTool/AutomationUtils/SourceIndexingUtils.cs
+++ b/Engine/Source/Programs/AutomationTool/AutomationUtils/SourceIndexingUtils.cs
@@ -69,7 +69,7 @@ namespace AutomationUtils
 			{
 				CommandUtils.Logger.LogInformation("Applying the following Virtual File System path transformations to source indexing information:");
 
-				foreach (var VFSMappingEntry in VFSMapping)
+				foreach ((DirectoryReference Local, DirectoryReference Virtualized) VFSMappingEntry in VFSMapping)
 				{
 					CommandUtils.Logger.LogInformation("\t{Local} -> {Virtualized}", VFSMappingEntry.Local.FullName, VFSMappingEntry.Virtualized.FullName);
 				}
diff --git a/Engine/Source/Programs/AutomationTool/AutomationUtils/UnrealBuild.cs b/Engine/Source/Programs/AutomationTool/AutomationUtils/UnrealBuild.cs
--- a/Engine/Source/Programs/AutomationTool/AutomationUtils/UnrealBuild.cs
+++ b/Engine/Source/Programs/AutomationTool/AutomationUtils/UnrealBuild.cs
@@ -561,7 +561,7 @@ namespace AutomationTool
 
 			// make a set of unique platforms involved
 			List<UnrealTargetPlatform> UniquePlatforms = new List<UnrealTargetPlatform>();
-			foreach (var Target in InAgenda.Targets)
+			foreach (BuildTarget Target in InAgenda.Targets)
 			{
 				if (!UniquePlatforms.Contains(Target.Platform))
 				{
@@ -664,7 +664,7 @@ namespace AutomationTool
 				CommandUtils.P4.ReconcileNoDeletes(WorkingCL, CommandUtils.MakePathSafeToUseWithCommandLine(File));
 
 				// Change file type on binary files to be always writeable.
-				var FileStats = CommandUtils.P4.FStat(File);
+				P4FileStat FileStats = CommandUtils.P4.FStat(File);
 
                 if (CommandUtils.IsProbablyAMacOrIOSExe(File))
                 {
diff --git a/Engine/Source/Programs/AutomationTool/AutomationUtils/Utils.cs b/Engine/Source/Programs/AutomationTool/AutomationUtils/Utils.cs
--- a/Engine/Source/Programs/AutomationTool/AutomationUtils/Utils.cs
+++ b/Engine/Source/Programs/AutomationTool/AutomationUtils/Utils.cs
@@ -176,7 +176,7 @@ namespace AutomationTool
 				Logger.LogDebug("RecursivelyDeleteDirectory {Path}", Path);
 			}
 			// Delete all files. This will also delete read-only files.
-			var FilesInDirectory = Directory.EnumerateFiles(Path);
+			IEnumerable<string> FilesInDirectory = Directory.EnumerateFiles(Path);
 			foreach (string Filename in FilesInDirectory)
 			{
 				if (SafeDeleteFile(Filename, bQuiet) == false)
@@ -186,7 +186,7 @@ namespace AutomationTool
 			}
 
 			// Recursively delete all files from sub-directories.
-			var FoldersInDirectory = Directory.EnumerateDirectories(Path);
+			IEnumerable<string> FoldersInDirectory = Directory.EnumerateDirectories(Path);
 			foreach (string Folder in FoldersInDirectory)
 			{
 				if (RecursivelyDeleteDirectory(Folder, bQuiet) == false)
@@ -1101,7 +1101,7 @@ namespace AutomationTool
             {
                 try
                 {
-                    var match = regex.Match(line);
+					Match match = regex.Match(line);
                     if (match.Success)
                     {
                         foundElements.Add(match.Groups["Type"].Value, int.Parse(match.Groups["Value"].Value));
diff --git a/Engine/Source/Programs/AutomationTool/AutomationUtils/ZenUtils.cs b/Engine/Source/Programs/AutomationTool/AutomationUtils/ZenUtils.cs
--- a/Engine/Source/Programs/AutomationTool/AutomationUtils/ZenUtils.cs
+++ b/Engine/Source/Programs/AutomationTool/AutomationUtils/ZenUtils.cs
@@ -485,7 +485,7 @@ namespace AutomationUtils
 		/// <returns>ZenWorkspace on success. null if there's no workspace for in this directory tree</returns>
 		public ZenWorkspace FindWorkspaceForDir(DirectoryReference BaseDir)
 		{
-			var Workspaces = GetWorkspaceList();
+			List<ZenWorkspace> Workspaces = GetWorkspaceList();
 			if (Workspaces == null)
 			{
 				return null;
diff --git a/Engine/Source/Programs/AutomationTool/Localization/LocalizationConfigFile.cs b/Engine/Source/Programs/AutomationTool/Localization/LocalizationConfigFile.cs
--- a/Engine/Source/Programs/AutomationTool/Localization/LocalizationConfigFile.cs
+++ b/Engine/Source/Programs/AutomationTool/Localization/LocalizationConfigFile.cs
@@ -106,7 +106,7 @@ namespace EpicGames.Localization
 					{
 						if (!String.IsNullOrEmpty(configLine.Value))
 						{
-							foreach (var pair in replacements)
+							foreach (KeyValuePair<string, string> pair in replacements)
 							{
 								// @TODOLocalization: For now, we only replace placeholders found in the values of the config files, replacing keys is not supported 
 								configLine.Value = configLine.Value.Replace(pair.Key, pair.Value);
diff --git a/Engine/Source/Programs/AutomationTool/Localization/LocalizationProvider.cs b/Engine/Source/Programs/AutomationTool/Localization/LocalizationProvider.cs
--- a/Engine/Source/Programs/AutomationTool/Localization/LocalizationProvider.cs
+++ b/Engine/Source/Programs/AutomationTool/Localization/LocalizationProvider.cs
@@ -101,7 +101,7 @@ namespace EpicGames.Localization
 		public List<string> GetConfigFilesToRun(List<string> RequestedLocalizationStepNames)
 		{
 			List<string> ConfigFiles = new();
-			foreach (var LocalizationStep in LocalizationSteps)
+			foreach (ProjectStepInfo LocalizationStep in LocalizationSteps)
 			{
 				if (RequestedLocalizationStepNames.Contains(LocalizationStep.Name))
 				{
@@ -179,13 +179,13 @@ namespace EpicGames.Localization
 				CachedLocalizationProviderTypes = new Dictionary<string, Type>();
 				foreach (Assembly Dll in ScriptManager.AllScriptAssemblies)
 				{
-					var AllTypes = Dll.GetTypes();
-					foreach (var PotentialLocalizationNodeType in AllTypes)
+					Type[] AllTypes = Dll.GetTypes();
+					foreach (Type PotentialLocalizationNodeType in AllTypes)
 					{
 						if (PotentialLocalizationNodeType != typeof(LocalizationProvider) && !PotentialLocalizationNodeType.IsAbstract && !PotentialLocalizationNodeType.IsInterface && typeof(LocalizationProvider).IsAssignableFrom(PotentialLocalizationNodeType))
 						{
 							// Types should implement a static StaticGetLocalizationProviderId method
-							var Method = PotentialLocalizationNodeType.GetMethod("StaticGetLocalizationProviderId");
+							MethodInfo Method = PotentialLocalizationNodeType.GetMethod("StaticGetLocalizationProviderId");
 							if (Method != null)
 							{
 								try
diff --git a/Engine/Source/Programs/AutomationTool/Localization/LocalizationUtilities.cs b/Engine/Source/Programs/AutomationTool/Localization/LocalizationUtilities.cs
--- a/Engine/Source/Programs/AutomationTool/Localization/LocalizationUtilities.cs
+++ b/Engine/Source/Programs/AutomationTool/Localization/LocalizationUtilities.cs
@@ -28,7 +28,7 @@ namespace EpicGames.Localization
 			{
 				string RelativePath = PluginDirectoryReference.MakeRelativeTo(PluginRootDirectoryReference);
 				IReadOnlyList<PluginInfo> PluginsUnderDirectory = Plugins.ReadPluginsFromDirectory(PluginRootDirectoryReference, RelativePath, Type);
-				foreach (var Plugin in PluginsUnderDirectory)
+				foreach (PluginInfo Plugin in PluginsUnderDirectory)
 				{
 					ReturnList.Add(Plugin.Name);
 				}
diff --git a/Engine/Source/Programs/AutomationTool/Program.cs b/Engine/Source/Programs/AutomationTool/Program.cs
--- a/Engine/Source/Programs/AutomationTool/Program.cs
+++ b/Engine/Source/Programs/AutomationTool/Program.cs
@@ -199,7 +199,7 @@ namespace AutomationToolDriver
 
 		private static void ParseDictionary(Dictionary<string, object> Value, List<string> Arguments)
 		{
-			foreach (var Pair in Value)
+			foreach (KeyValuePair<string, object> Pair in Value)
 			{
 				if ((Pair.Value as string) != null && !string.IsNullOrEmpty(Pair.Value as string))
 				{
diff --git a/Engine/Source/Programs/OidcToken/Program.cs b/Engine/Source/Programs/OidcToken/Program.cs
--- a/Engine/Source/Programs/OidcToken/Program.cs
+++ b/Engine/Source/Programs/OidcToken/Program.cs
@@ -180,7 +180,7 @@ namespace OidcToken
 			await Host.CreateDefaultBuilder(args)
 				.UseSerilog((context, configuration) =>
 				{
-					var section = context.Configuration.GetSection("Serilog");
+					IConfigurationSection section = context.Configuration.GetSection("Serilog");
 					if (!section.GetChildren().Any())
 					{
 						// no serilog provider configuration, adding some defaults
