Using previous nightly build for patch-in-place

#rnx
#Synced-CL 51328467

diff --git a/Engine/Source/Developer/IoStoreUtilities/Private/IoStoreLooseFiles.cpp b/Engine/Source/Developer/IoStoreUtilities/Private/IoStoreLooseFiles.cpp
--- a/Engine/Source/Developer/IoStoreUtilities/Private/IoStoreLooseFiles.cpp
+++ b/Engine/Source/Developer/IoStoreUtilities/Private/IoStoreLooseFiles.cpp
@@ -61,7 +61,7 @@ public:
 	virtual ~FLooseFilesIoStoreWriter()
 	{ }
 
-	virtual void SetReferenceChunkDatabase(TSharedPtr<IIoStoreWriterReferenceChunkDatabase> ReferenceChunkDatabase) override
+	virtual void SetReferenceChunkDatabase(TSharedPtr<IIoStoreWriterReferenceChunkDatabase> ReferenceChunkDatabase, TSharedPtr<IIoStoreWriterReferenceChunkDatabase> PatchInPlaceReferenceChunkDatabase) override
 	{ }
 
 	virtual void EnableDiskLayoutOrdering(const TArray<TUniquePtr<FIoStoreReader>>& PatchSourceReaders = TArray<TUniquePtr<FIoStoreReader>>())
diff --git a/Engine/Source/Developer/IoStoreUtilities/Private/IoStoreUtilities.cpp b/Engine/Source/Developer/IoStoreUtilities/Private/IoStoreUtilities.cpp
--- a/Engine/Source/Developer/IoStoreUtilities/Private/IoStoreUtilities.cpp
+++ b/Engine/Source/Developer/IoStoreUtilities/Private/IoStoreUtilities.cpp
@@ -4579,6 +4579,19 @@ static int32 CreateTarget(const FIoStoreArguments& Arguments, const FIoStoreWrit
 		}
 	}
 
+	TSharedPtr<IIoStoreWriterReferenceChunkDatabase> PatchInPlaceChunkDatabase;
+	if (Arguments.PatchInPlaceReferenceContainerPathOverride.Len())
+	{
+		PatchInPlaceChunkDatabase = MakeShared<FIoStoreChunkDatabase>();
+		FIoStoreChunkDatabase* ChunkDbPtr = (FIoStoreChunkDatabase*)PatchInPlaceChunkDatabase.Get();
+
+		if (ChunkDbPtr->Init(Arguments.PatchInPlaceReferenceContainerPathOverride, Arguments.PatchInPlaceReferenceAdditionalContainersPath, Arguments.PatchInPlaceReferenceChunkKeys) == false)
+		{
+			UE_LOG(LogIoStore, Warning, TEXT("Failed to initialize reference chunk store. Pak will continue."));
+			PatchInPlaceChunkDatabase.Reset();
+		}
+	}
+
 	TArray<FCookedPackage*> Packages;
 	FPackageNameMap PackageNameMap;
 	FPackageIdMap PackageIdMap;
@@ -4616,7 +4629,7 @@ static int32 CreateTarget(const FIoStoreArguments& Arguments, const FIoStoreWrit
 				GlobalContainerSettings.ContainerFlags |= EIoContainerFlags::Signed;
 			}
 			GlobalIoStoreWriter = IoStoreWriterContext->CreateContainer(*Arguments.GlobalContainerPath, GlobalContainerSettings);
-			GlobalIoStoreWriter->SetReferenceChunkDatabase(ChunkDatabase);
+			GlobalIoStoreWriter->SetReferenceChunkDatabase(ChunkDatabase, PatchInPlaceChunkDatabase);
 			IoStoreWriters.Add(GlobalIoStoreWriter);
 			IoStoreWriterInfos.Add({UE::Cook::EPluginSizeTypes::Installed, "global"});
 		}
@@ -4668,7 +4681,7 @@ static int32 CreateTarget(const FIoStoreArguments& Arguments, const FIoStoreWrit
 				ContainerSettings.bGenerateDiffPatch = ContainerTarget->bGenerateDiffPatch;
 				ContainerTarget->IoStoreWriter = IoStoreWriterContext->CreateContainer(*ContainerTarget->OutputPath, ContainerSettings);
 				ContainerTarget->IoStoreWriter->EnableDiskLayoutOrdering(ContainerTarget->PatchSourceReaders);
-				ContainerTarget->IoStoreWriter->SetReferenceChunkDatabase(ChunkDatabase);
+				ContainerTarget->IoStoreWriter->SetReferenceChunkDatabase(ChunkDatabase, PatchInPlaceChunkDatabase);
 				IoStoreWriters.Add(ContainerTarget->IoStoreWriter);
 
 				if (EnumHasAnyFlags(ContainerTarget->ContainerFlags, EIoContainerFlags::OnDemand))
@@ -4701,7 +4714,7 @@ static int32 CreateTarget(const FIoStoreArguments& Arguments, const FIoStoreWrit
 					ContainerTarget->OptionalSegmentIoStoreWriter = IoStoreWriterContext->CreateContainer(*ContainerTarget->OptionalSegmentOutputPath, ContainerSettings);
 					ContainerSettings.ContainerId = ContainerTarget->ContainerId;
 
-					ContainerTarget->OptionalSegmentIoStoreWriter->SetReferenceChunkDatabase(ChunkDatabase);
+					ContainerTarget->OptionalSegmentIoStoreWriter->SetReferenceChunkDatabase(ChunkDatabase, PatchInPlaceChunkDatabase);
 					IoStoreWriters.Add(ContainerTarget->OptionalSegmentIoStoreWriter);
 					IoStoreWriterInfos.Add({UE::Cook::EPluginSizeTypes::OptionalSegment, ContainerTarget->Name});
 				}
@@ -9161,6 +9174,18 @@ bool ParseContainerGenerationArguments(FIoStoreArguments& Arguments, FIoStoreWri
 		FParse::Value(FCommandLine::Get(), TEXT("-patchInPlaceMaxChainSize="), WriterSettings.PatchInPlaceMaxChainSize);
 		FParse::Value(FCommandLine::Get(), TEXT("-patchInPlaceMaxChainLength="), WriterSettings.PatchInPlaceMaxChainLength);
 		FParse::Value(FCommandLine::Get(), TEXT("-patchInPlaceParetoFrontHeuristic="), WriterSettings.PatchInPlaceParetoFrontHeuristic);
+
+		if (FParse::Value(FCommandLine::Get(), TEXT("-patchInPlaceReferenceContainerPathOverride="), Arguments.PatchInPlaceReferenceContainerPathOverride))
+		{
+			FString CryptoKeysCacheFilename;
+			if (FParse::Value(FCommandLine::Get(), TEXT("-patchInPlaceReferenceContainerCryptoOverride="), CryptoKeysCacheFilename))
+			{
+				UE_LOG(LogIoStore, Display, TEXT("Parsing reference container crypto keys from a crypto key cache file '%s'"), *CryptoKeysCacheFilename);
+				KeyChainUtilities::LoadKeyChainFromFile(CryptoKeysCacheFilename, Arguments.PatchInPlaceReferenceChunkKeys);
+			}
+
+			FParse::Value(FCommandLine::Get(), TEXT("-patchInPlaceReferenceAdditionalContainersPath="), Arguments.PatchInPlaceReferenceAdditionalContainersPath);
+		}
 	}
 
 	UE_LOG(LogIoStore, Display, TEXT("Using memory mapping alignment '%ld'"), WriterSettings.MemoryMappingAlignment);
diff --git a/Engine/Source/Developer/IoStoreUtilities/Private/IoStoreUtilitiesCore.h b/Engine/Source/Developer/IoStoreUtilities/Private/IoStoreUtilitiesCore.h
--- a/Engine/Source/Developer/IoStoreUtilities/Private/IoStoreUtilitiesCore.h
+++ b/Engine/Source/Developer/IoStoreUtilities/Private/IoStoreUtilitiesCore.h
@@ -148,8 +148,11 @@ struct FIoStoreArguments
 	FString ReferenceChunkGlobalContainerFileName;
 	FString ReferenceChunkChangesCSVFileName;
 	FString ReferenceChunkAdditionalContainersPath;
+	FString PatchInPlaceReferenceContainerPathOverride;
+	FString PatchInPlaceReferenceAdditionalContainersPath;
 	FString CsvPath;
 	FKeyChain ReferenceChunkKeys;
+	FKeyChain PatchInPlaceReferenceChunkKeys;
 	FReleasedPackages ReleasedPackages;
 	TUniquePtr<FCookedPackageStore> PackageStore;
 	TUniquePtr<FIoBuffer> ScriptObjects;
diff --git a/Engine/Source/Developer/IoStoreUtilities/Private/IoStoreWriter.cpp b/Engine/Source/Developer/IoStoreUtilities/Private/IoStoreWriter.cpp
--- a/Engine/Source/Developer/IoStoreUtilities/Private/IoStoreWriter.cpp
+++ b/Engine/Source/Developer/IoStoreUtilities/Private/IoStoreWriter.cpp
@@ -1995,8 +1995,19 @@ public:
 	{
 	}
 
-	void SetReferenceChunkDatabase(TSharedPtr<IIoStoreWriterReferenceChunkDatabase> InReferenceChunkDatabase)
+	void SetReferenceChunkDatabase(TSharedPtr<IIoStoreWriterReferenceChunkDatabase> InReferenceChunkDatabase, TSharedPtr<IIoStoreWriterReferenceChunkDatabase> InPatchInPlaceReferenceChunkDatabase)
 	{
+		if (InPatchInPlaceReferenceChunkDatabase.IsValid())
+		{
+			PatchInPlaceReferenceChunkDatabase = InPatchInPlaceReferenceChunkDatabase;
+			PatchInPlaceReferenceChunkDatabase->NotifyAddedToWriter(ContainerSettings.ContainerId, FPaths::GetBaseFilename(TocFilePath));
+		}
+		else
+		{
+			// fallback to using reference chunk database
+			PatchInPlaceReferenceChunkDatabase = InReferenceChunkDatabase;
+		}
+
 		if (InReferenceChunkDatabase.IsValid() == false)
 		{
 			ReferenceChunkDatabase = InReferenceChunkDatabase;
@@ -2009,8 +2020,9 @@ public:
 			UE_LOG(LogIoStore, Warning, TEXT("No chunks will match, so ignoring. ReferenceChunkDb: %d, IoStoreWriter: %d"), InReferenceChunkDatabase->GetCompressionBlockSize(), WriterContext->GetSettings().CompressionBlockSize);
 			return;
 		}
+
 		ReferenceChunkDatabase = InReferenceChunkDatabase;
-		
+
 		// Add ourselves to the reference chunk db's list of possibles
 		ReferenceChunkDatabase->NotifyAddedToWriter(ContainerSettings.ContainerId, FPaths::GetBaseFilename(TocFilePath));
 	}
@@ -2503,14 +2515,14 @@ public:
 		PatchInPlacePinnedSizeBeforeDefrag = 0;
 		PatchInPlacePinnedSizeAfterDefrag = 0;
 
-		if (!ReferenceChunkDatabase.IsValid())
+		if (!PatchInPlaceReferenceChunkDatabase.IsValid())
 		{
 			UE_LOG(LogIoStore, Warning, TEXT("Trying to use patch-in-place without valid reference database"));
 			return;
 		}
 
 		FString ReferenceContainerPath;
-		if (!ReferenceChunkDatabase->GetContainerFilePath(ContainerSettings.ContainerId, ReferenceContainerPath))
+		if (!PatchInPlaceReferenceChunkDatabase->GetContainerFilePath(ContainerSettings.ContainerId, ReferenceContainerPath))
 		{
 			UE_LOG(LogIoStore, Display, TEXT("Skipping patch-in-place for container '%s' because reference container was not found"), *ContainerPathAndBaseFileName);
 			return;
@@ -3519,6 +3531,7 @@ private:
 	bool						bHasFlushed = false;
 	bool						bHasResult = false;
 	TSharedPtr<IIoStoreWriterReferenceChunkDatabase> ReferenceChunkDatabase;
+	TSharedPtr<IIoStoreWriterReferenceChunkDatabase> PatchInPlaceReferenceChunkDatabase;
 
 
 	friend class FIoStoreWriterContextImpl;
diff --git a/Engine/Source/Developer/IoStoreUtilities/Private/IoStoreWriter.h b/Engine/Source/Developer/IoStoreUtilities/Private/IoStoreWriter.h
--- a/Engine/Source/Developer/IoStoreUtilities/Private/IoStoreWriter.h
+++ b/Engine/Source/Developer/IoStoreUtilities/Private/IoStoreWriter.h
@@ -201,7 +201,7 @@ public:
 	*	If a reference database is provided, the IoStoreWriter implementation may elect to reuse compressed blocks
 	*	from previous containers instead of recompressing input data. This must be set before any writes are appended.
 	*/
-	virtual void SetReferenceChunkDatabase(TSharedPtr<IIoStoreWriterReferenceChunkDatabase> ReferenceChunkDatabase) = 0;
+	virtual void SetReferenceChunkDatabase(TSharedPtr<IIoStoreWriterReferenceChunkDatabase> ReferenceChunkDatabase, TSharedPtr<IIoStoreWriterReferenceChunkDatabase> PatchInPlaceReferenceChunkDatabase) = 0;
 	virtual void EnableDiskLayoutOrdering(const TArray<TUniquePtr<FIoStoreReader>>& PatchSourceReaders = TArray<TUniquePtr<FIoStoreReader>>()) = 0;
 	virtual void Append(const FIoChunkId& ChunkId, FIoBuffer Chunk, const FIoWriteOptions& WriteOptions, uint64 OrderHint = MAX_uint64) = 0;
 	virtual void Append(const FIoChunkId& ChunkId, IIoStoreWriteRequest* Request, const FIoWriteOptions& WriteOptions) = 0;
diff --git a/Engine/Source/Programs/AutomationTool/Scripts/CopyBuildToStagingDirectory.Automation.cs b/Engine/Source/Programs/AutomationTool/Scripts/CopyBuildToStagingDirectory.Automation.cs
--- a/Engine/Source/Programs/AutomationTool/Scripts/CopyBuildToStagingDirectory.Automation.cs
+++ b/Engine/Source/Programs/AutomationTool/Scripts/CopyBuildToStagingDirectory.Automation.cs
@@ -4605,27 +4605,75 @@ namespace AutomationScripts
 					{
 						AdditionalArgs += " -ReferenceContainerAdditionalPath=\"" + Params.ReferenceContainerAdditionalPath + "\"";
 					}
+
 					if (Params.ReferenceContainerCryptoKeys.Length != 0)
 					{
 						AdditionalArgs += " -ReferenceContainerCryptoKeys=\"" + Params.ReferenceContainerCryptoKeys + "\"";
 					}
+
 					if (Params.ReferenceContainerChangesCSVFileName.Length != 0)
 					{
 						AdditionalArgs += " -ReferenceContainerChangesCSVFileName=\"" + Params.ReferenceContainerChangesCSVFileName + "\"";
 					}
-					bool bUsePatchInPlaceLayout = false;
-					if (PlatformGameConfig.GetBool("/Script/UnrealEd.ProjectPackagingSettings", "bUsePatchInPlaceLayout", out bUsePatchInPlaceLayout) && bUsePatchInPlaceLayout)
-					{
-						AdditionalArgs += " -UsePatchInPlaceLayout";
+				}
+
+				bool bUsePatchInPlaceLayout = false;
+				if (PlatformGameConfig.GetBool("/Script/UnrealEd.ProjectPackagingSettings", "bUsePatchInPlaceLayout", out bUsePatchInPlaceLayout) && bUsePatchInPlaceLayout)
+				{
+					AdditionalArgs += " -UsePatchInPlaceLayout";
 					
-						if (PlatformGameConfig.GetInt32("/Script/UnrealEd.ProjectPackagingSettings", "PatchInPlaceMaximumTotalGrowthMB", out int PatchInPlaceMaximumTotalGrowthMB))
+					if (PlatformGameConfig.GetInt32("/Script/UnrealEd.ProjectPackagingSettings", "PatchInPlaceMaximumTotalGrowthMB", out int PatchInPlaceMaximumTotalGrowthMB))
+					{
+						AdditionalArgs += $" -patchInPlaceMaximumTotalGrowth={PatchInPlaceMaximumTotalGrowthMB}MB";
+					}
+
+					if (PlatformGameConfig.GetInt32("/Script/UnrealEd.ProjectPackagingSettings", "PatchInPlaceMaximumPerContainerGrowthMB", out int PatchInPlaceMaximumPerContainerGrowthMB))
+					{
+						AdditionalArgs += $" -patchInPlaceMaximumPerContainerGrowth={PatchInPlaceMaximumPerContainerGrowthMB}MB";
+					}
+
+					if (PlatformGameConfig.GetInt32("/Script/UnrealEd.ProjectPackagingSettings", "PatchInPlaceParetoFrontCoarseSamples", out int PatchInPlaceParetoFrontCoarseSamples))
+					{
+						AdditionalArgs += $" -patchInPlaceParetoFrontCoarseSamples={PatchInPlaceParetoFrontCoarseSamples}";
+					}
+
+					if (PlatformGameConfig.GetInt32("/Script/UnrealEd.ProjectPackagingSettings", "PatchInPlaceParetoFrontFineSamples", out int PatchInPlaceParetoFrontFineSamples))
+					{
+						AdditionalArgs += $" -patchInPlaceParetoFrontFineSamples={PatchInPlaceParetoFrontFineSamples}";
+					}
+
+					if (PlatformGameConfig.GetInt32("/Script/UnrealEd.ProjectPackagingSettings", "PatchInPlaceReadGranularity", out int PatchInPlaceReadGranularity))
+					{
+						AdditionalArgs += $" -patchInPlaceReadGranularity={PatchInPlaceReadGranularity}";
+					}
+
+					if (PlatformGameConfig.GetInt32("/Script/UnrealEd.ProjectPackagingSettings", "PatchInPlaceMaxChainSize", out int PatchInPlaceMaxChainSize))
+					{
+						AdditionalArgs += $" -patchInPlaceMaxChainSize={PatchInPlaceMaxChainSize}";
+					}
+
+					if (PlatformGameConfig.GetInt32("/Script/UnrealEd.ProjectPackagingSettings", "PatchInPlaceMaxChainLength", out int PatchInPlaceMaxChainLength))
+					{
+						AdditionalArgs += $" -patchInPlaceMaxChainLength={PatchInPlaceMaxChainLength}";
+					}
+
+					if (PlatformGameConfig.GetInt32("/Script/UnrealEd.ProjectPackagingSettings", "PatchInPlaceParetoFrontHeuristic", out int PatchInPlaceParetoFrontHeuristic))
+					{
+						AdditionalArgs += $" -patchInPlaceParetoFrontHeuristic={PatchInPlaceParetoFrontHeuristic}";
+					}
+
+					if (PlatformGameConfig.GetString("/Script/UnrealEd.ProjectPackagingSettings", "PatchInPlaceReferenceContainerPathOverride", out string PatchInPlaceReferenceContainerPathOverride) && PatchInPlaceReferenceContainerPathOverride.Length != 0)
+					{
+						AdditionalArgs += $" -patchInPlaceReferenceContainerPathOverride=\"{PatchInPlaceReferenceContainerPathOverride}\"";
+
+						if (PlatformGameConfig.GetString("/Script/UnrealEd.ProjectPackagingSettings", "PatchInPlaceReferenceContainerCryptoOverride", out string PatchInPlaceReferenceContainerCryptoOverride) && PatchInPlaceReferenceContainerCryptoOverride.Length != 0)
 						{
-							AdditionalArgs += $" -patchInPlaceMaximumTotalGrowth={PatchInPlaceMaximumTotalGrowthMB}MB";
+							AdditionalArgs += $" -patchInPlaceReferenceContainerCryptoOverride=\"{PatchInPlaceReferenceContainerCryptoOverride}\"";
 						}
 
-						if (PlatformGameConfig.GetInt32("/Script/UnrealEd.ProjectPackagingSettings", "PatchInPlaceMaximumPerContainerGrowthMB", out int PatchInPlaceMaximumPerContainerGrowthMB))
+						if (PlatformGameConfig.GetString("/Script/UnrealEd.ProjectPackagingSettings", "PatchInPlaceReferenceAdditionalContainersPath", out string PatchInPlaceReferenceAdditionalContainersPath) && PatchInPlaceReferenceAdditionalContainersPath.Length != 0)
 						{
-							AdditionalArgs += $" -patchInPlaceMaximumPerContainerGrowth={PatchInPlaceMaximumPerContainerGrowthMB}MB";
+							AdditionalArgs += $" -patchInPlaceReferenceAdditionalContainersPath=\"{PatchInPlaceReferenceAdditionalContainersPath}\"";
 						}
 					}
 				}
