Remove unneeded Microsoft.CSharp PackageReference

#rnx
#Synced-CL 51506754

diff --git a/Engine/Plugins/VirtualProduction/ICVFXTesting/Build/Scripts/ICVFX.Automation.csproj b/Engine/Plugins/VirtualProduction/ICVFXTesting/Build/Scripts/ICVFX.Automation.csproj
--- a/Engine/Plugins/VirtualProduction/ICVFXTesting/Build/Scripts/ICVFX.Automation.csproj
+++ b/Engine/Plugins/VirtualProduction/ICVFXTesting/Build/Scripts/ICVFX.Automation.csproj
@@ -17,7 +17,6 @@
     </Compile>
   </ItemGroup>
   <ItemGroup>
-    <PackageReference Include="Microsoft.CSharp" Version="4.7.0" />
     <PackageReference Include="System.Data.DataSetExtensions" Version="4.5.0" />
   </ItemGroup>
   <ItemGroup>
diff --git a/Engine/Source/Programs/AutomationTool/Android/Android.Automation.csproj b/Engine/Source/Programs/AutomationTool/Android/Android.Automation.csproj
--- a/Engine/Source/Programs/AutomationTool/Android/Android.Automation.csproj
+++ b/Engine/Source/Programs/AutomationTool/Android/Android.Automation.csproj
@@ -22,7 +22,4 @@
     <ProjectReference Include="..\AutomationUtils\AutomationUtils.Automation.csproj" PrivateAssets="All"><Private>false</Private></ProjectReference>
     <ProjectReference Include="..\Scripts\AutomationScripts.Automation.csproj" PrivateAssets="All"><Private>false</Private></ProjectReference>
   </ItemGroup>
-  <ItemGroup>
-    <PackageReference Include="Microsoft.CSharp" Version="4.7.0" />
-  </ItemGroup>
 </Project>
\ No newline at end of file
diff --git a/Engine/Source/Programs/AutomationTool/Apple/Apple.Automation.csproj b/Engine/Source/Programs/AutomationTool/Apple/Apple.Automation.csproj
--- a/Engine/Source/Programs/AutomationTool/Apple/Apple.Automation.csproj
+++ b/Engine/Source/Programs/AutomationTool/Apple/Apple.Automation.csproj
@@ -29,7 +29,4 @@
     <ProjectReference Include="..\..\UnrealBuildTool\UnrealBuildTool.csproj" PrivateAssets="All"><Private>false</Private></ProjectReference>
     <ProjectReference Include="..\AutomationUtils\AutomationUtils.Automation.csproj" PrivateAssets="All"><Private>false</Private></ProjectReference>
   </ItemGroup>
-  <ItemGroup>
-    <PackageReference Include="Microsoft.CSharp" Version="4.7.0" PrivateAssets="all" />
-  </ItemGroup>
 </Project>
\ No newline at end of file
diff --git a/Engine/Source/Programs/AutomationTool/BuildGraph/BuildGraph.Automation.csproj b/Engine/Source/Programs/AutomationTool/BuildGraph/BuildGraph.Automation.csproj
--- a/Engine/Source/Programs/AutomationTool/BuildGraph/BuildGraph.Automation.csproj
+++ b/Engine/Source/Programs/AutomationTool/BuildGraph/BuildGraph.Automation.csproj
@@ -33,7 +33,4 @@
     <ProjectReference Include="..\..\UnrealBuildTool\UnrealBuildTool.csproj" PrivateAssets="All"><Private>false</Private></ProjectReference>
     <ProjectReference Include="..\AutomationUtils\AutomationUtils.Automation.csproj" PrivateAssets="All"><Private>false</Private></ProjectReference>
   </ItemGroup>
-  <ItemGroup>
-    <PackageReference Include="Microsoft.CSharp" Version="4.7.0" />
-  </ItemGroup>
 </Project>
\ No newline at end of file
diff --git a/Engine/Source/Programs/AutomationTool/CookedEditor/CookedEditor.Automation.csproj b/Engine/Source/Programs/AutomationTool/CookedEditor/CookedEditor.Automation.csproj
--- a/Engine/Source/Programs/AutomationTool/CookedEditor/CookedEditor.Automation.csproj
+++ b/Engine/Source/Programs/AutomationTool/CookedEditor/CookedEditor.Automation.csproj
@@ -25,7 +25,4 @@
     <ProjectReference Include="..\Scripts\AutomationScripts.Automation.csproj" />
     <ProjectReference Include="..\Win\Win.Automation.csproj" />
   </ItemGroup>
-  <ItemGroup>
-    <PackageReference Include="Microsoft.CSharp" Version="4.7.0" />
-  </ItemGroup>
 </Project>
\ No newline at end of file
diff --git a/Engine/Source/Programs/AutomationTool/CrowdinLocalization/CrowdinLocalization.Automation.csproj b/Engine/Source/Programs/AutomationTool/CrowdinLocalization/CrowdinLocalization.Automation.csproj
--- a/Engine/Source/Programs/AutomationTool/CrowdinLocalization/CrowdinLocalization.Automation.csproj
+++ b/Engine/Source/Programs/AutomationTool/CrowdinLocalization/CrowdinLocalization.Automation.csproj
@@ -21,7 +21,4 @@
     <ProjectReference Include="..\AutomationUtils\AutomationUtils.Automation.csproj" PrivateAssets="All"><Private>false</Private></ProjectReference>
     <ProjectReference Include="..\Localization\Localization.Automation.csproj" PrivateAssets="All"><Private>false</Private></ProjectReference>
   </ItemGroup>
-  <ItemGroup>
-    <PackageReference Include="Microsoft.CSharp" Version="4.7.0" />
-  </ItemGroup>
 </Project>
\ No newline at end of file
diff --git a/Engine/Source/Programs/AutomationTool/Gauntlet/Gauntlet.Automation.csproj b/Engine/Source/Programs/AutomationTool/Gauntlet/Gauntlet.Automation.csproj
--- a/Engine/Source/Programs/AutomationTool/Gauntlet/Gauntlet.Automation.csproj
+++ b/Engine/Source/Programs/AutomationTool/Gauntlet/Gauntlet.Automation.csproj
@@ -77,7 +77,6 @@
   <ItemGroup>
     <PackageReference Include="Appium.WebDriver" Version="8.0.0" />
     <PackageReference Include="Magick.NET-Q16-HDRI-AnyCPU" Version="14.9.1" PrivateAssets="all" />
-    <PackageReference Include="Microsoft.CSharp" Version="4.7.0" />
     <PackageReference Include="MySql.Data" Version="6.10.9" PrivateAssets="all" />
     <PackageReference Include="Newtonsoft.Json" Version="13.0.4" PrivateAssets="all" />
   </ItemGroup>
diff --git a/Engine/Source/Programs/AutomationTool/IOS/IOS.Automation.csproj b/Engine/Source/Programs/AutomationTool/IOS/IOS.Automation.csproj
--- a/Engine/Source/Programs/AutomationTool/IOS/IOS.Automation.csproj
+++ b/Engine/Source/Programs/AutomationTool/IOS/IOS.Automation.csproj
@@ -35,7 +35,6 @@
     <ProjectReference Include="..\Apple\Apple.Automation.csproj" />
   </ItemGroup>
   <ItemGroup>
-    <PackageReference Include="Microsoft.CSharp" Version="4.7.0" PrivateAssets="all" />
     <PackageReference Include="System.Resources.Extensions" Version="4.7.1" PrivateAssets="all" />
     <PackageReference Include="Portable.BouncyCastle" Version="1.8.10" PrivateAssets="all" />
   </ItemGroup>
diff --git a/Engine/Source/Programs/AutomationTool/Linux/Linux.Automation.csproj b/Engine/Source/Programs/AutomationTool/Linux/Linux.Automation.csproj
--- a/Engine/Source/Programs/AutomationTool/Linux/Linux.Automation.csproj
+++ b/Engine/Source/Programs/AutomationTool/Linux/Linux.Automation.csproj
@@ -23,7 +23,4 @@
     <ProjectReference Include="..\AutomationUtils\AutomationUtils.Automation.csproj" PrivateAssets="All"><Private>false</Private></ProjectReference>
     <ProjectReference Include="..\SteamDeck\SteamDeck.Automation.csproj" />
   </ItemGroup>
-  <ItemGroup>
-    <PackageReference Include="Microsoft.CSharp" Version="4.7.0" />
-  </ItemGroup>
 </Project>
\ No newline at end of file
diff --git a/Engine/Source/Programs/AutomationTool/LowLevelTests/LowLevelTests.Automation.csproj b/Engine/Source/Programs/AutomationTool/LowLevelTests/LowLevelTests.Automation.csproj
--- a/Engine/Source/Programs/AutomationTool/LowLevelTests/LowLevelTests.Automation.csproj
+++ b/Engine/Source/Programs/AutomationTool/LowLevelTests/LowLevelTests.Automation.csproj
@@ -34,7 +34,6 @@
     </Compile>
   </ItemGroup>
   <ItemGroup>
-    <PackageReference Include="Microsoft.CSharp" Version="4.7.0" />
     <PackageReference Include="MySql.Data" Version="6.10.9" PrivateAssets="all" />
   </ItemGroup>
   <ItemGroup>
diff --git a/Engine/Source/Programs/AutomationTool/Mac/Mac.Automation.csproj b/Engine/Source/Programs/AutomationTool/Mac/Mac.Automation.csproj
--- a/Engine/Source/Programs/AutomationTool/Mac/Mac.Automation.csproj
+++ b/Engine/Source/Programs/AutomationTool/Mac/Mac.Automation.csproj
@@ -22,7 +22,4 @@
     <ProjectReference Include="..\AutomationUtils\AutomationUtils.Automation.csproj" PrivateAssets="All"><Private>false</Private></ProjectReference>
     <ProjectReference Include="..\Apple\Apple.Automation.csproj" />
   </ItemGroup>
-  <ItemGroup>
-    <PackageReference Include="Microsoft.CSharp" Version="4.7.0" />
-  </ItemGroup>
 </Project>
\ No newline at end of file
diff --git a/Engine/Source/Programs/AutomationTool/Microsoft/Microsoft.Automation.csproj b/Engine/Source/Programs/AutomationTool/Microsoft/Microsoft.Automation.csproj
--- a/Engine/Source/Programs/AutomationTool/Microsoft/Microsoft.Automation.csproj
+++ b/Engine/Source/Programs/AutomationTool/Microsoft/Microsoft.Automation.csproj
@@ -21,7 +21,4 @@
     <ProjectReference Include="..\..\UnrealBuildTool\UnrealBuildTool.csproj" PrivateAssets="All"><Private>false</Private></ProjectReference>
     <ProjectReference Include="..\AutomationUtils\AutomationUtils.Automation.csproj" PrivateAssets="All"><Private>false</Private></ProjectReference>
   </ItemGroup>
-  <ItemGroup>
-    <PackageReference Include="Microsoft.CSharp" Version="4.7.0" />
-  </ItemGroup>
 </Project>
\ No newline at end of file
diff --git a/Engine/Source/Programs/AutomationTool/OneSkyLocalization/OneSkyLocalization.Automation.csproj b/Engine/Source/Programs/AutomationTool/OneSkyLocalization/OneSkyLocalization.Automation.csproj
--- a/Engine/Source/Programs/AutomationTool/OneSkyLocalization/OneSkyLocalization.Automation.csproj
+++ b/Engine/Source/Programs/AutomationTool/OneSkyLocalization/OneSkyLocalization.Automation.csproj
@@ -30,7 +30,4 @@
     <ProjectReference Include="..\AutomationUtils\AutomationUtils.Automation.csproj" PrivateAssets="All"><Private>false</Private></ProjectReference>
     <ProjectReference Include="..\Localization\Localization.Automation.csproj" PrivateAssets="All"><Private>false</Private></ProjectReference>
   </ItemGroup>
-  <ItemGroup>
-    <PackageReference Include="Microsoft.CSharp" Version="4.7.0" />
-  </ItemGroup>
 </Project>
\ No newline at end of file
diff --git a/Engine/Source/Programs/AutomationTool/Scripts/AutomationScripts.Automation.csproj b/Engine/Source/Programs/AutomationTool/Scripts/AutomationScripts.Automation.csproj
--- a/Engine/Source/Programs/AutomationTool/Scripts/AutomationScripts.Automation.csproj
+++ b/Engine/Source/Programs/AutomationTool/Scripts/AutomationScripts.Automation.csproj
@@ -33,7 +33,6 @@
     <ProjectReference Include="..\Localization\Localization.Automation.csproj" PrivateAssets="All"><Private>false</Private></ProjectReference>
   </ItemGroup>
   <ItemGroup>
-    <PackageReference Include="Microsoft.CSharp" Version="4.7.0" />
     <PackageReference Include="Microsoft.VisualStudio.OLE.Interop" Version="17.14.40260" />
     <PackageReference Include="AWSSDK.S3" Version="3.3.113.2" PrivateAssets="all" />
   </ItemGroup>
diff --git a/Engine/Source/Programs/AutomationTool/SmartlingLocalization/SmartlingLocalization.Automation.csproj b/Engine/Source/Programs/AutomationTool/SmartlingLocalization/SmartlingLocalization.Automation.csproj
--- a/Engine/Source/Programs/AutomationTool/SmartlingLocalization/SmartlingLocalization.Automation.csproj
+++ b/Engine/Source/Programs/AutomationTool/SmartlingLocalization/SmartlingLocalization.Automation.csproj
@@ -20,7 +20,4 @@
     <ProjectReference Include="..\AutomationUtils\AutomationUtils.Automation.csproj" PrivateAssets="All"><Private>false</Private></ProjectReference>
     <ProjectReference Include="..\Localization\Localization.Automation.csproj" PrivateAssets="All"><Private>false</Private></ProjectReference>
   </ItemGroup>
-  <ItemGroup>
-    <PackageReference Include="Microsoft.CSharp" Version="4.7.0" />
-  </ItemGroup>
 </Project>
\ No newline at end of file
diff --git a/Engine/Source/Programs/AutomationTool/SteamDeck/SteamDeck.Automation.csproj b/Engine/Source/Programs/AutomationTool/SteamDeck/SteamDeck.Automation.csproj
--- a/Engine/Source/Programs/AutomationTool/SteamDeck/SteamDeck.Automation.csproj
+++ b/Engine/Source/Programs/AutomationTool/SteamDeck/SteamDeck.Automation.csproj
@@ -22,7 +22,4 @@
     <ProjectReference Include="..\..\UnrealBuildTool\UnrealBuildTool.csproj" PrivateAssets="All"><Private>false</Private></ProjectReference>
     <ProjectReference Include="..\AutomationUtils\AutomationUtils.Automation.csproj" PrivateAssets="All"><Private>false</Private></ProjectReference>
   </ItemGroup>
-  <ItemGroup>
-    <PackageReference Include="Microsoft.CSharp" Version="4.7.0" />
-  </ItemGroup>
 </Project>
\ No newline at end of file
diff --git a/Engine/Source/Programs/AutomationTool/TVOS/TVOS.Automation.csproj b/Engine/Source/Programs/AutomationTool/TVOS/TVOS.Automation.csproj
--- a/Engine/Source/Programs/AutomationTool/TVOS/TVOS.Automation.csproj
+++ b/Engine/Source/Programs/AutomationTool/TVOS/TVOS.Automation.csproj
@@ -23,7 +23,4 @@
     <ProjectReference Include="..\AutomationUtils\AutomationUtils.Automation.csproj" PrivateAssets="All"><Private>false</Private></ProjectReference>
     <ProjectReference Include="..\IOS\IOS.Automation.csproj" PrivateAssets="All"><Private>false</Private></ProjectReference>
   </ItemGroup>
-  <ItemGroup>
-    <PackageReference Include="Microsoft.CSharp" Version="4.7.0" />
-  </ItemGroup>
 </Project>
\ No newline at end of file
diff --git a/Engine/Source/Programs/AutomationTool/Turnkey/Turnkey.Automation.csproj b/Engine/Source/Programs/AutomationTool/Turnkey/Turnkey.Automation.csproj
--- a/Engine/Source/Programs/AutomationTool/Turnkey/Turnkey.Automation.csproj
+++ b/Engine/Source/Programs/AutomationTool/Turnkey/Turnkey.Automation.csproj
@@ -34,7 +34,6 @@
   </ItemGroup>
   <ItemGroup>
     <PackageReference Include="Google.Apis.Drive.v3" Version="1.33.1.1232" PrivateAssets="all" />
-    <PackageReference Include="Microsoft.CSharp" Version="4.7.0" />
     <!-- Unused, but required to resolve version conflicts from transitive dependencies. -->
     <PackageReference Include="Newtonsoft.Json" Version="13.0.4" PrivateAssets="all" />
     <PackageReference Include="System.IO.FileSystem.AccessControl" Version="4.7.0" />
diff --git a/Engine/Source/Programs/AutomationTool/Win/Win.Automation.csproj b/Engine/Source/Programs/AutomationTool/Win/Win.Automation.csproj
--- a/Engine/Source/Programs/AutomationTool/Win/Win.Automation.csproj
+++ b/Engine/Source/Programs/AutomationTool/Win/Win.Automation.csproj
@@ -25,7 +25,4 @@
     <ProjectReference Include="..\SteamDeck\SteamDeck.Automation.csproj" />
     <ProjectReference Include="..\Microsoft\Microsoft.Automation.csproj" PrivateAssets="All"><Private>false</Private></ProjectReference>
   </ItemGroup>
-  <ItemGroup>
-    <PackageReference Include="Microsoft.CSharp" Version="4.7.0" />
-  </ItemGroup>
 </Project>
\ No newline at end of file
diff --git a/Engine/Source/Programs/AutomationTool/XLocLocalization/XLocLocalization.Automation.csproj b/Engine/Source/Programs/AutomationTool/XLocLocalization/XLocLocalization.Automation.csproj
--- a/Engine/Source/Programs/AutomationTool/XLocLocalization/XLocLocalization.Automation.csproj
+++ b/Engine/Source/Programs/AutomationTool/XLocLocalization/XLocLocalization.Automation.csproj
@@ -21,7 +21,6 @@
     <ProjectReference Include="..\Localization\Localization.Automation.csproj" PrivateAssets="All"><Private>false</Private></ProjectReference>
   </ItemGroup>
   <ItemGroup>
-    <PackageReference Include="Microsoft.CSharp" Version="4.7.0" PrivateAssets="all" />
     <PackageReference Include="System.ServiceModel.Http" Version="4.10.3" PrivateAssets="all" />
   </ItemGroup>
 </Project>
\ No newline at end of file
diff --git a/Engine/Source/Programs/Shared/EpicGames.UBA/EpicGames.UBA.csproj b/Engine/Source/Programs/Shared/EpicGames.UBA/EpicGames.UBA.csproj
--- a/Engine/Source/Programs/Shared/EpicGames.UBA/EpicGames.UBA.csproj
+++ b/Engine/Source/Programs/Shared/EpicGames.UBA/EpicGames.UBA.csproj
@@ -20,7 +20,6 @@
   <Import Project="Library.props" />
 
   <ItemGroup>
-    <PackageReference Include="Microsoft.CSharp" Version="4.7.0" />
     <PackageReference Include="Microsoft.Extensions.Logging" Version="8.0.1" />
   </ItemGroup>
 
diff --git a/Engine/Source/Programs/Shared/EpicGames.UHT/EpicGames.UHT.csproj b/Engine/Source/Programs/Shared/EpicGames.UHT/EpicGames.UHT.csproj
--- a/Engine/Source/Programs/Shared/EpicGames.UHT/EpicGames.UHT.csproj
+++ b/Engine/Source/Programs/Shared/EpicGames.UHT/EpicGames.UHT.csproj
@@ -9,7 +9,4 @@
     <ProjectReference Include="..\EpicGames.Core\EpicGames.Core.csproj" />
   </ItemGroup>
 
-  <ItemGroup>
-    <PackageReference Include="Microsoft.CSharp" Version="4.7.0" />
-  </ItemGroup>
 </Project>
