[iOS] Fix Remote Mac Build failure during packaging

Remote Mac Builds need to use Ionic.Zip library but it was not being copied to the expected paths.

#Synced-CL 51453194

diff --git a/Engine/Source/Programs/AutomationTool/IOS/IOS.Automation.csproj b/Engine/Source/Programs/AutomationTool/IOS/IOS.Automation.csproj
--- a/Engine/Source/Programs/AutomationTool/IOS/IOS.Automation.csproj
+++ b/Engine/Source/Programs/AutomationTool/IOS/IOS.Automation.csproj
@@ -19,7 +19,6 @@
     </Reference>
     <Reference Include="Ionic.Zip.Reduced">
       <HintPath>..\..\..\..\Binaries\DotNET\Ionic.Zip.Reduced.dll</HintPath>
-      <Private>False</Private>
     </Reference>
   </ItemGroup>
 
