Adjusting LogEntry to capture optional timestamps

#rnx

#Synced-CL 51459300

diff --git a/Engine/Source/Programs/AutomationTool/Gauntlet/Unreal/Utils/Gauntlet.UnrealLogParser.cs b/Engine/Source/Programs/AutomationTool/Gauntlet/Unreal/Utils/Gauntlet.UnrealLogParser.cs
--- a/Engine/Source/Programs/AutomationTool/Gauntlet/Unreal/Utils/Gauntlet.UnrealLogParser.cs
+++ b/Engine/Source/Programs/AutomationTool/Gauntlet/Unreal/Utils/Gauntlet.UnrealLogParser.cs
@@ -1,15 +1,16 @@
 // Copyright Epic Games, Inc. All Rights Reserved.
 
+using AutomationTool;
+using EpicGames.Core;
+using EpicGames.LogMatchers;
+using Microsoft.Extensions.Logging;
 using System;
 using System.Collections.Generic;
-using System.Text.RegularExpressions;
+using System.Globalization;
 using System.Linq;
 using System.Text;
-using EpicGames.Core;
-using EpicGames.LogMatchers;
-using Microsoft.Extensions.Logging;
+using System.Text.RegularExpressions;
 using Logging = Microsoft.Extensions.Logging;
-using AutomationTool;
 
 namespace Gauntlet
 {
@@ -66,7 +67,12 @@ namespace Gauntlet
 			public string Message { get; private set; }
 
 			/// <summary>
-			/// Format the entry as it would have appeared in the log. 
+			/// Optional datetime of the entry
+			/// </summary>
+			public DateTime? Timestamp { get; private set; }
+
+			/// <summary>
+			/// Format the entry as it would have appeared in the log (without the timestamp). 
 			/// </summary>
 			/// <returns></returns>
 			public override string ToString()
@@ -83,6 +89,43 @@ namespace Gauntlet
 				}
 			}
 
+			/// <summary>
+			/// Format the entry as it would have appeared in the log (with the UTC timestamp).
+			/// </summary>
+			/// <returns></returns>
+			public string ToStringWithUtcTimestamp()
+			{
+				if (Timestamp.HasValue)
+				{
+					return string.Format("[{0}]{1}", Timestamp.Value.ToString("yyyy.MM.dd-HH.mm.ss:fff"), ToString());
+				}
+				else
+				{
+					return ToString();
+				}
+			}
+
+			/// <summary>
+			/// Constructor that requires all info
+			/// </summary>
+			/// <param name="InPrefix"></param>
+			/// <param name="InCategory"></param>
+			/// <param name="InLevel"></param>
+			/// <param name="InMessage"></param>
+			/// <param name="InTimestamp"></param>
+			public LogEntry(string InPrefix, string InCategory, LogLevel InLevel, string InMessage, string InTimestamp)
+			{
+				Prefix = InPrefix;
+				Category = InCategory;
+				Level = InLevel;
+				Message = InMessage;
+
+				if (DateTime.TryParseExact(InTimestamp, "yyyy.MM.dd-HH.mm.ss:fff", CultureInfo.InvariantCulture, DateTimeStyles.AssumeUniversal, out DateTime ParsedTimestamp))
+				{
+					Timestamp = ParsedTimestamp;
+				}
+			}
+
 			/// <summary>
 			/// Constructor that requires all info
 			/// </summary>
@@ -90,12 +133,14 @@ namespace Gauntlet
 			/// <param name="InCategory"></param>
 			/// <param name="InLevel"></param>
 			/// <param name="InMessage"></param>
-			public LogEntry(string InPrefix, string InCategory, LogLevel InLevel, string InMessage)
+			/// <param name="InTimestamp"></param>
+			public LogEntry(string InPrefix, string InCategory, LogLevel InLevel, string InMessage, DateTime? InTimestamp)
 			{
 				Prefix = InPrefix;
 				Category = InCategory;
 				Level = InLevel;
 				Message = InMessage;
+				Timestamp = InTimestamp;
 			}
 		}
 
@@ -377,7 +422,8 @@ namespace Gauntlet
 				Clear();
 			}
 
-			Regex UELogLinePattern = new Regex(@"(?<channel>[A-Za-z][\w\d]+):\s(?:(?<level>Display|Verbose|VeryVerbose|Warning|Error|Fatal):\s)?");
+			// Capture of the UE log where the timestamp and frame are optional groups as they are not always available per line
+			Regex UELogLinePattern = new Regex(@"(?:\[(?<timestamp>[0-9.\-:]+)\])?(?:\[(?<frame>[\s\d]+)\])?(?<channel>[A-Za-z][\w\d]+):\s(?:(?<level>Display|Verbose|VeryVerbose|Warning|Error|Fatal):\s)?");
 
 			if (LineOffset >= 0)
 			{
@@ -392,6 +438,7 @@ namespace Gauntlet
 				{
 					string Channel = MatchLine.Groups["channel"].Value;
 					string Message = Line.Substring(MatchLine.Index + MatchLine.Length);
+					string Timestamp = MatchLine.Groups["timestamp"].Value;
 					ReadOnlySpan<char> LevelSpan = MatchLine.Groups["level"].ValueSpan;
 
 					UnrealLog.LogLevel Level = UnrealLog.LogLevel.Log;
@@ -415,12 +462,12 @@ namespace Gauntlet
 						Prefix = Channel.Substring(0, 3);
 						Channel = Channel.Substring(3);
 					}
-					Entry = new(Prefix, Channel, Level, Message);
+					Entry = new(Prefix, Channel, Level, Message, Timestamp);
 				}
 				else
 				{
 					// Not an Unreal Engine log line
-					Entry = new(string.Empty, string.Empty, UnrealLog.LogLevel.Log, Line);
+					Entry = new(string.Empty, string.Empty, UnrealLog.LogLevel.Log, Line, string.Empty);
 				}
 
 				LogEvents.Add(Entry);
@@ -621,7 +668,8 @@ namespace Gauntlet
 		{
 			// Search for LogFoo: <Display|Error|etc>: Message
 			// Also need to handle 'Log' not always being present, and the category being empty for a level of 'Log'
-			Regex Pattern = new Regex(@"(?<prefix>Log)?(?<category>[A-Za-z][\w\d]+):\s(?<level>Display|Verbose|VeryVerbose|Warning|Error|Fatal)?(?::\s)?");
+			// Also handling the optional timestamp and frame groups
+			Regex Pattern = new Regex(@"(?:\[(?<timestamp>[0-9.\-:]+)\])?(?:\[(?<frame>[\s\d]+)\])?(?<prefix>Log)?(?<category>[A-Za-z][\w\d]+):\s(?<level>Display|Verbose|VeryVerbose|Warning|Error|Fatal)?(?::\s)?");
 
 			List<UnrealLog.LogEntry> ParsedEntries = new List<UnrealLog.LogEntry>();
 
@@ -634,6 +682,7 @@ namespace Gauntlet
 				string Prefix = M.Groups["prefix"].Value;
 				string Category = M.Groups["category"].Value;
 				string Message = Line.Substring(M.Index + M.Length);
+				string Timestamp = M.Groups["timestamp"].Value;
 
 				ReadOnlySpan<char> LevelSpan = M.Groups["level"].ValueSpan;
 				UnrealLog.LogLevel Level = UnrealLog.LogLevel.Log;
@@ -652,7 +701,7 @@ namespace Gauntlet
 					}
 				}
 
-				ParsedEntries.Add(new UnrealLog.LogEntry(Prefix, Category, Level, Message));
+				ParsedEntries.Add(new UnrealLog.LogEntry(Prefix, Category, Level, Message, Timestamp));
 			}
 
 			return ParsedEntries;
