[Gauntlet] Modifies UnrealTestConfiguration.NullRhi read accesibility to public

#Synced-CL 51364815

diff --git a/Engine/Source/Programs/AutomationTool/Gauntlet/Unreal/Base/Gauntlet.UnrealTestConfiguration.cs b/Engine/Source/Programs/AutomationTool/Gauntlet/Unreal/Base/Gauntlet.UnrealTestConfiguration.cs
--- a/Engine/Source/Programs/AutomationTool/Gauntlet/Unreal/Base/Gauntlet.UnrealTestConfiguration.cs
+++ b/Engine/Source/Programs/AutomationTool/Gauntlet/Unreal/Base/Gauntlet.UnrealTestConfiguration.cs
@@ -667,7 +667,7 @@ namespace Gauntlet
 		/// </summary>
 		/// 
 		[AutoParam(false)]
-		protected bool Nullrhi { get; set; }
+		public bool NullRHI { get; protected set; }
 
 		// Public options that tests can configure
 
@@ -1043,7 +1043,7 @@ namespace Gauntlet
 		/// <returns></returns>
 		public virtual void ApplyToConfig(UnrealAppConfig AppConfig, UnrealSessionRole ConfigRole, IEnumerable<UnrealSessionRole> OtherRoles)
 		{
-			if (Nullrhi)
+			if (NullRHI)
 			{
 				AppConfig.CommandLine += " -nullrhi";
 			}
@@ -1082,7 +1082,7 @@ namespace Gauntlet
 					}
 				}
 
-				if (ScreenshotPeriod > 0 && Nullrhi == false)
+				if (ScreenshotPeriod > 0 && NullRHI == false)
 				{
 					AppConfig.CommandLine += string.Format(" -gauntlet.screenshotperiod={0}", ScreenshotPeriod);
 				}
