[Gauntlet] Add overload to UnrealSession.LaunchDeferred role to return the newly created role instance

#Synced-CL 51364815
#Synced-CL 51364815

diff --git a/Engine/Source/Programs/AutomationTool/Gauntlet/Unreal/Base/Gauntlet.UnrealSession.cs b/Engine/Source/Programs/AutomationTool/Gauntlet/Unreal/Base/Gauntlet.UnrealSession.cs
--- a/Engine/Source/Programs/AutomationTool/Gauntlet/Unreal/Base/Gauntlet.UnrealSession.cs
+++ b/Engine/Source/Programs/AutomationTool/Gauntlet/Unreal/Base/Gauntlet.UnrealSession.cs
@@ -468,13 +468,25 @@ namespace Gauntlet
 		/// <returns></returns>
 		public bool LaunchDeferredRole( UnrealSessionRole Role )
 		{
+			return LaunchDeferredRole(Role, out _);
+		}
+
+		/// <summary>
+		/// Launches a role that was previously flagged as 'DeferredLaunch'. Note that this does not handle device failure, marking problem devices etc.
+		/// </summary>
+		/// <param name="Role"></param>
+		/// <param name="RoleInstance"></param>
+		/// <returns></returns>
+		public bool LaunchDeferredRole(UnrealSessionRole Role, out RoleInstance RoleInstance)
+		{
+			RoleInstance = null;
 			if (Role.DeferredLaunch == false)
 			{
 				Log.Error("Cannot start deferred role {0} because it is not marked 'DeferredLaunch'", Role.ToString());
 				return false;
 			}
 
-			int CurrentRoleInstanceIndex = AllRoles.FindIndex( X => X.Role == Role );
+			int CurrentRoleInstanceIndex = AllRoles.FindIndex(X => X.Role == Role);
 			if (CurrentRoleInstanceIndex < 0)
 			{
 				Log.Error("Cannot start dereferred role {0} because it cannot be found among the 'all roles' list", Role.ToString());
@@ -495,7 +507,8 @@ namespace Gauntlet
 				{
 					// remove deferred role and update to a running role instance
 					DeferredRoleToAppInstall.Remove(AllRoles[CurrentRoleInstanceIndex]);
-					AllRoles[CurrentRoleInstanceIndex] = new RoleInstance(Role, Instance);
+					RoleInstance = new RoleInstance(Role, Instance);
+					AllRoles[CurrentRoleInstanceIndex] = RoleInstance;
 				}
 
 				Success = true;
