[Automation] Ignore file write times for overlay executables when running on a buildmachine.
This ensures build installation order of operations doesn't matter

#Synced-CL 51364815

diff --git a/Engine/Source/Programs/AutomationTool/Gauntlet/Framework/Base/Gauntlet.OverlayExecutable.cs b/Engine/Source/Programs/AutomationTool/Gauntlet/Framework/Base/Gauntlet.OverlayExecutable.cs
--- a/Engine/Source/Programs/AutomationTool/Gauntlet/Framework/Base/Gauntlet.OverlayExecutable.cs
+++ b/Engine/Source/Programs/AutomationTool/Gauntlet/Framework/Base/Gauntlet.OverlayExecutable.cs
@@ -73,9 +73,10 @@ namespace Gauntlet
 			string LocalExecutable = Path.Combine(PlatformBinariesDirectory, ExecutableName + ExecutableExtension);
 			string LocalPreflightExecutable = Path.Combine(PlatformBinariesDirectory, PreflightExecutableName + ExecutableExtension);
 
+			bool bIsNewer = File.GetLastWriteTime(LocalPreflightExecutable) > File.GetLastWriteTime(BaseExecutable);
+
 			// Prioritize using -Preflight executables, if they exist
-			if ((File.Exists(LocalPreflightExecutable) || Directory.Exists(LocalPreflightExecutable))
-				&& File.GetLastWriteTime(LocalPreflightExecutable) > File.GetLastWriteTime(BaseExecutable))
+			if ((File.Exists(LocalPreflightExecutable) || Directory.Exists(LocalPreflightExecutable)) && (bIsNewer || CommandUtils.IsBuildMachine))
 			{
 				Log.Info("Applying newer preflight executable as overlay {PreflightExecutable}", LocalPreflightExecutable);
 				OverlayExecutable = LocalPreflightExecutable;
@@ -88,7 +89,9 @@ namespace Gauntlet
 				return false;
 			}
 
-			if (File.GetLastWriteTime(BaseExecutable) > File.GetLastWriteTime(LocalExecutable))
+			bIsNewer = File.GetLastWriteTime(LocalExecutable) > File.GetLastWriteTime(BaseExecutable);
+
+			if (!bIsNewer && !CommandUtils.IsBuildMachine)
 			{
 				Log.Verbose("Local executable for {Platform} is not newer than base executable. Skipping overlay for this role", Platform);
 				return false;
