Fix for discrepancy between chunkdb sizes and reported deltas.

#rnx

#Synced-CL 51354374

#Synced-CL 51354374

#Synced-CL 51354374

#Synced-CL 51354374

diff --git a/Engine/Source/Programs/AutomationTool/AutomationUtils/BundleUtils.cs b/Engine/Source/Programs/AutomationTool/AutomationUtils/BundleUtils.cs
--- a/Engine/Source/Programs/AutomationTool/AutomationUtils/BundleUtils.cs
+++ b/Engine/Source/Programs/AutomationTool/AutomationUtils/BundleUtils.cs
@@ -464,5 +464,35 @@ namespace AutomationUtils.Automation
 			}
 			return VirtualBundleTags;
 		}
+
+		// Loads the metadata for a given manifest and returns any NON virtual bundle tags inferred from its contents.
+		public static List<string> InferNonVirtualBundleTagsFromManifest(string ManifestPath, BuildPatchToolBase.ToolVersion BPTVersion = BuildPatchToolBase.ToolVersion.Live)
+		{
+			if (!string.IsNullOrEmpty(ManifestPath))
+			{
+				BuildPatchToolBase.ManifestMetadataOptions MetadataOptions = new BuildPatchToolBase.ManifestMetadataOptions
+				{
+					SourceManifestFilename = ManifestPath
+				};
+				BuildPatchToolBase.ManifestMetadataOutput ManifestMetadataOutput = null;
+				BuildPatchToolBase.Get().Execute(MetadataOptions, out ManifestMetadataOutput, BPTVersion);
+				if (ManifestMetadataOutput != null)
+				{
+					List<string> VirtualBundleTags = InferVirtualBundleTagsFromManifest(ManifestMetadataOutput);
+					
+					// We expect very few virtual bundles (like 1) so doing this n^2 isn't terrible.
+					List<string> NonVirtualTags = new List<string>();
+					foreach (string Tag in ManifestMetadataOutput.FileTagList)
+					{
+						if (!VirtualBundleTags.Contains(Tag))
+						{
+							NonVirtualTags.Add(Tag);
+						}
+					}
+					return NonVirtualTags;
+				}
+			}
+			return new List<string>();
+		}
 	}
 }
diff --git a/Engine/Source/Programs/AutomationTool/AutomationUtils/MCPPublic.cs b/Engine/Source/Programs/AutomationTool/AutomationUtils/MCPPublic.cs
--- a/Engine/Source/Programs/AutomationTool/AutomationUtils/MCPPublic.cs
+++ b/Engine/Source/Programs/AutomationTool/AutomationUtils/MCPPublic.cs
@@ -1693,6 +1693,10 @@ namespace EpicGames.MCP.Automation
 			
 			// If generating with a custom optimized delta, specify it here.
 			public string DeltaFilenameTrailer;
+
+			// Minimum compression benefit to require when writing chunks to accept compression, to avoid spending client decompression time
+			// on minimal benefits. (10 is good)
+			public int MinCompressionPct = 0;
 		}
 
 		public class ManifestMetadataOptions
