[Gauntlet] Removes the global lock acquisition in desktop platform's ITargetDevice.Run functions as it is not needed and creates deadlocks if a deferred role is executed within a Test's Tick

#Synced-CL 51568933

diff --git a/Engine/Source/Programs/AutomationTool/Gauntlet/Platform/Linux/Gauntlet.TargetDeviceLinux.cs b/Engine/Source/Programs/AutomationTool/Gauntlet/Platform/Linux/Gauntlet.TargetDeviceLinux.cs
--- a/Engine/Source/Programs/AutomationTool/Gauntlet/Platform/Linux/Gauntlet.TargetDeviceLinux.cs
+++ b/Engine/Source/Programs/AutomationTool/Gauntlet/Platform/Linux/Gauntlet.TargetDeviceLinux.cs
@@ -55,48 +55,35 @@ namespace Gauntlet
 				throw new AutomationException("Specified path {0} not found!", LinuxApp.ExecutablePath);
 			}
 
-			ILongProcessResult Result = null;
+			string CmdLine = LinuxApp.CommandArguments;
+			Log.Info("Launching {0} on {1}", App.Name, ToString());
+			Log.Verbose("\t{0}", CmdLine);
 
-			lock (Globals.MainLock)
-			{
-				string ExePath = Path.GetDirectoryName(LinuxApp.ExecutablePath);
-				string NewWorkingDir = string.IsNullOrEmpty(LinuxApp.WorkingDirectory) ? ExePath : LinuxApp.WorkingDirectory;
-				string OldWD = Environment.CurrentDirectory;
-				Environment.CurrentDirectory = NewWorkingDir;
-
-				Log.Info("Launching {0} on {1}", App.Name, ToString());
-
-				string CmdLine = LinuxApp.CommandArguments;
-
-				Log.Verbose("\t{0}", CmdLine);
-
-				bool bAppContainerized = LinuxApp is IContainerized;
+			bool bAppContainerized = LinuxApp is IContainerized;
 
-				// Forward command to Docker container if running containerized
-				if (bAppContainerized)
-				{
-					ContainerInfo Container = ((IContainerized)LinuxApp).ContainerInfo;
-					string ContainerApp = Container.WorkingDir + "/" + Path.GetRelativePath(Globals.UnrealRootDir, LinuxApp.ExecutablePath).Replace('\\', '/');
-					CmdLine = $"run --name {Container.ContainerName} {Container.ImageName} {Container.RunCommandPrepend} {ContainerApp} {CmdLine}";
-				}
-
-				Result = new LongProcessResult(
-					bAppContainerized ? "docker" : LinuxApp.ExecutablePath,
-					CmdLine,
-					LinuxApp.RunOptions,
-					OutputCallback: LinuxApp.FilterLoggingDelegate,
-					WorkingDir: LinuxApp.WorkingDirectory,
-					LocalCache: LinuxApp.Device.LocalCachePath
-				);
+			// Forward command to Docker container if running containerized
+			if (bAppContainerized)
+			{
+				ContainerInfo Container = ((IContainerized)LinuxApp).ContainerInfo;
+				string ContainerApp = Container.WorkingDir + "/" + Path.GetRelativePath(Globals.UnrealRootDir, LinuxApp.ExecutablePath).Replace('\\', '/');
+				CmdLine = $"run --name {Container.ContainerName} {Container.ImageName} {Container.RunCommandPrepend} {ContainerApp} {CmdLine}";
+			}
 
-				if (Result.HasExited && Result.ExitCode != 0)
-				{
-					throw new AutomationException("Failed to launch {0}. Error {1}", LinuxApp.ExecutablePath, Result.ExitCode);
-				}
+			ILongProcessResult Result = new LongProcessResult(
+				bAppContainerized ? "docker" : LinuxApp.ExecutablePath,
+				CmdLine,
+				LinuxApp.RunOptions,
+				OutputCallback: LinuxApp.FilterLoggingDelegate,
+				WorkingDir: LinuxApp.WorkingDirectory,
+				LocalCache: LinuxApp.Device.LocalCachePath
+			);
 
-				Environment.CurrentDirectory = OldWD;
+			if (Result.HasExited && Result.ExitCode != 0)
+			{
+				throw new AutomationException("Failed to launch {0}. Error {1}", LinuxApp.ExecutablePath, Result.ExitCode);
 			}
 
+
 			return new LinuxAppInstance(LinuxApp, Result, LinuxApp.LogFile);
 		}
 
@@ -125,11 +112,14 @@ namespace Gauntlet
 
 			PopulateDirectoryMappings(Path.Combine(BuildDir, AppConfig.ProjectName));
 
+			string ExecutablePath = Path.IsPathRooted(Build.ExecutablePath)
+					? Build.ExecutablePath
+					: Path.Combine(BuildDir, Build.ExecutablePath);
+
 			LinuxAppInstall LinuxApp = new LinuxAppInstall(AppConfig.Name, AppConfig.ProjectName, this)
 			{
-				ExecutablePath = Path.IsPathRooted(Build.ExecutablePath)
-					? Build.ExecutablePath
-					: Path.Combine(BuildDir, Build.ExecutablePath)
+				ExecutablePath = ExecutablePath,
+				WorkingDirectory = Path.GetDirectoryName(ExecutablePath)
 			};
 			LinuxApp.SetDefaultCommandLineArguments(AppConfig, RunOptions, BuildDir);
 
@@ -237,6 +227,8 @@ namespace Gauntlet
 				LinuxApp.ExecutablePath = BinaryPath;
 			}
 
+			LinuxApp.WorkingDirectory = Path.GetDirectoryName(LinuxApp.ExecutablePath);
+
 			return LinuxApp;
 		}
 
diff --git a/Engine/Source/Programs/AutomationTool/Gauntlet/Platform/Mac/Gauntlet.TargetDeviceMac.cs b/Engine/Source/Programs/AutomationTool/Gauntlet/Platform/Mac/Gauntlet.TargetDeviceMac.cs
--- a/Engine/Source/Programs/AutomationTool/Gauntlet/Platform/Mac/Gauntlet.TargetDeviceMac.cs
+++ b/Engine/Source/Programs/AutomationTool/Gauntlet/Platform/Mac/Gauntlet.TargetDeviceMac.cs
@@ -104,31 +104,21 @@ namespace Gauntlet
 				throw new AutomationException("Invalid install type!");
 			}
 
-			ILongProcessResult Result = null;
-
-			lock (Globals.MainLock)
+			Log.Info("Launching {0} on {1}", App.Name, ToString());
+			Log.Verbose("\t{0}", MacInstall.CommandArguments);
+
+			ILongProcessResult Result = new LongProcessResult(
+				GetExecutableIfBundle(MacInstall.ExecutablePath),
+				MacInstall.CommandArguments,
+				MacInstall.RunOptions,
+				OutputCallback: MacInstall.FilterLoggingDelegate,
+				WorkingDir: MacInstall.WorkingDirectory,
+				LocalCache: MacInstall.Device.LocalCachePath
+			);
+
+			if (Result.HasExited && Result.ExitCode != 0)
 			{
-				string OldWD = Environment.CurrentDirectory;
-				Environment.CurrentDirectory = MacInstall.WorkingDirectory;
-
-				Log.Info("Launching {0} on {1}", App.Name, ToString());
-				Log.Verbose("\t{0}", MacInstall.CommandArguments);
-
-				Result = new LongProcessResult(
-					GetExecutableIfBundle(MacInstall.ExecutablePath),
-					MacInstall.CommandArguments,
-					MacInstall.RunOptions,
-					OutputCallback: MacInstall.FilterLoggingDelegate,
-					WorkingDir: MacInstall.WorkingDirectory,
-					LocalCache: MacInstall.Device.LocalCachePath
-				);
-
-				if (Result.HasExited && Result.ExitCode != 0)
-				{
-					throw new AutomationException("Failed to launch {0}. Error {1}", MacInstall.ExecutablePath, Result.ExitCode);
-				}
-
-				Environment.CurrentDirectory = OldWD;
+				throw new AutomationException("Failed to launch {0}. Error {1}", MacInstall.ExecutablePath, Result.ExitCode);
 			}
 
 			return new MacAppInstance(MacInstall, Result, MacInstall.LogFile);
diff --git a/Engine/Source/Programs/AutomationTool/Gauntlet/Platform/Windows/Gauntlet.TargetDeviceWindows.cs b/Engine/Source/Programs/AutomationTool/Gauntlet/Platform/Windows/Gauntlet.TargetDeviceWindows.cs
--- a/Engine/Source/Programs/AutomationTool/Gauntlet/Platform/Windows/Gauntlet.TargetDeviceWindows.cs
+++ b/Engine/Source/Programs/AutomationTool/Gauntlet/Platform/Windows/Gauntlet.TargetDeviceWindows.cs
@@ -85,36 +85,22 @@ namespace Gauntlet
 				throw new AutomationException("Specified path {0} not found!", WinApp.ExecutablePath);
 			}
 
-			ILongProcessResult Result = null;
-
-			lock (Globals.MainLock)
+			string CmdLine = WinApp.CommandArguments;
+			Log.Info("Launching {0} on {1}", App.Name, ToString());
+			Log.Verbose("\t{0}", CmdLine);
+
+			ILongProcessResult Result = new LongProcessResult(
+				WinApp.ExecutablePath,
+				CmdLine,
+				WinApp.RunOptions,
+				OutputCallback: WinApp.FilterLoggingDelegate,
+				WorkingDir: WinApp.WorkingDirectory,
+				LocalCache: WinApp.Device.LocalCachePath
+			);
+
+			if (Result.HasExited && Result.ExitCode != 0)
 			{
-				string ExePath = Path.GetDirectoryName(WinApp.ExecutablePath);
-				string NewWorkingDir = string.IsNullOrEmpty(WinApp.WorkingDirectory) ? ExePath : WinApp.WorkingDirectory;
-				string OldWD = Environment.CurrentDirectory;
-				Environment.CurrentDirectory = NewWorkingDir;
-
-				Log.Info("Launching {0} on {1}", App.Name, ToString());
-
-				string CmdLine = WinApp.CommandArguments;
-
-				Log.Verbose("\t{0}", CmdLine);
-
-				Result = new LongProcessResult(
-					WinApp.ExecutablePath,
-					CmdLine,
-					WinApp.RunOptions,
-					OutputCallback: WinApp.FilterLoggingDelegate,
-					WorkingDir: WinApp.WorkingDirectory,
-					LocalCache: WinApp.Device.LocalCachePath
-				);
-
-				if (Result.HasExited && Result.ExitCode != 0)
-				{
-					throw new AutomationException("Failed to launch {0}. Error {1}", WinApp.ExecutablePath, Result.ExitCode);
-				}
-
-				Environment.CurrentDirectory = OldWD;
+				throw new AutomationException("Failed to launch {0}. Error {1}", WinApp.ExecutablePath, Result.ExitCode);
 			}
 
 			return new WindowsAppInstance(WinApp, Result, WinApp.LogFile);
@@ -140,11 +126,14 @@ namespace Gauntlet
 
 			PopulateDirectoryMappings(Path.Combine(BuildDir, AppConfig.ProjectName));
 
+			string ExecutablePath = Path.IsPathRooted(Build.ExecutablePath)
+					? Build.ExecutablePath
+					: Path.Combine(BuildDir, Build.ExecutablePath);
+
 			WindowsAppInstall WinApp = new WindowsAppInstall(AppConfig.Name, AppConfig.ProjectName, this)
 			{
-				ExecutablePath = Path.IsPathRooted(Build.ExecutablePath)
-					? Build.ExecutablePath
-					: Path.Combine(BuildDir, Build.ExecutablePath)
+				ExecutablePath = ExecutablePath,
+				WorkingDirectory = Path.GetDirectoryName(ExecutablePath)
 			};
 			WinApp.SetDefaultCommandLineArguments(AppConfig, RunOptions, BuildDir);
 
@@ -295,6 +284,8 @@ namespace Gauntlet
 				WinApp.ExecutablePath = BinaryPath;
 			}
 
+			WinApp.WorkingDirectory = Path.GetDirectoryName(WinApp.ExecutablePath);
+
 			BootstrapInstall(WinApp);
 
 			return WinApp;
