Enable -statsPort in DstmDemo to support automated testing by controlling cvars through an HTTP endpoint.

#Synced-CL 51341769
#Synced-CL 51341769

diff --git a/Engine/Source/Programs/AutomationTool/Scripts/LaunchMultiServer.Automation.cs b/Engine/Source/Programs/AutomationTool/Scripts/LaunchMultiServer.Automation.cs
--- a/Engine/Source/Programs/AutomationTool/Scripts/LaunchMultiServer.Automation.cs
+++ b/Engine/Source/Programs/AutomationTool/Scripts/LaunchMultiServer.Automation.cs
@@ -242,6 +242,9 @@ Example running 4 clients locally connecting to single proxy server w/ a non-loc
 		public int BaseGameListenPort = 9000;
 		public int ProxyServerBasePort = -1;
 
+		// Stats port counter: each process gets a unique -statsPort starting at this value
+		private int NextStatsPort = 10100;
+
 		[CommandLine, Help("CustomConfig=<Section>", "Read a different section of the Project's GameConfig for default values")]
 		public string CustomConfig = "";
 
@@ -682,6 +685,8 @@ Example running 4 clients locally connecting to single proxy server w/ a non-loc
 
 			Args += String.Format(" -ConsoleTitle=\"{0} {1}\" -SessionName=\"{0} {1}\" ", Params.ShortProjectName, AppTitle);
 			Args += AdditionalArguments;
+			Args += $" -statsPort={NextStatsPort}";
+			NextStatsPort++;
 			Args = Args.Trim();
 
 			PushDir(Path.GetDirectoryName(App));
