Added the -enableStatsPort arugment to LaunchMultiServer to explicitly provide the -statsPort argument to the processes it spawns (disabled by default).

#Synced-CL 51341769

diff --git a/Engine/Source/Programs/AutomationTool/Scripts/LaunchMultiServer.Automation.cs b/Engine/Source/Programs/AutomationTool/Scripts/LaunchMultiServer.Automation.cs
--- a/Engine/Source/Programs/AutomationTool/Scripts/LaunchMultiServer.Automation.cs
+++ b/Engine/Source/Programs/AutomationTool/Scripts/LaunchMultiServer.Automation.cs
@@ -299,6 +299,9 @@ Example running 4 clients locally connecting to single proxy server w/ a non-loc
 		[CommandLine, Help("CommonArgs", "Specify arguments that are run by all processes")]
 		public string CommonArgs = "";
 
+		[CommandLine, Help("EnableStatsPort", "Enable the -statsPort argument on all launched processes")]
+		public bool EnableStatsPort = false;
+
 		[CommandLine, Help("Matchmake", "Project will use a matchmaking service to connect servers instead of passing the -MultiServerPeers and -ProxyGameServers options to servers.")]
 		public bool Matchmake = false;
 
@@ -685,8 +688,11 @@ Example running 4 clients locally connecting to single proxy server w/ a non-loc
 
 			Args += String.Format(" -ConsoleTitle=\"{0} {1}\" -SessionName=\"{0} {1}\" ", Params.ShortProjectName, AppTitle);
 			Args += AdditionalArguments;
-			Args += $" -statsPort={NextStatsPort}";
-			NextStatsPort++;
+			if (EnableStatsPort)
+			{
+				Args += $" -statsPort={NextStatsPort}";
+				NextStatsPort++;
+			}
 			Args = Args.Trim();
 
 			PushDir(Path.GetDirectoryName(App));
