IncrementalCook: AutomationTool and ProjectSettings:
1) Change text in ProjectSettings to indicate that the incremental settings chose whether the current cook tries to incrementally cook regardless of whether modern incrementalcook or legacyIterativedependencies are used.
2) Merge -cookincremental and -iterativecooking in AutomationTool into a single ProjectParam setting: CookIncremental. This matches the changes already made in ProjectSettings, ProjectLauncher, and CookCommandlet.
3) Fix automationtool to enforce that incrementalcooks do not support unversioned when legacyiterative is used.
#rnx

#Synced-CL 51558650

diff --git a/Engine/Source/Developer/DeveloperToolSettings/Classes/CookerSettings.h b/Engine/Source/Developer/DeveloperToolSettings/Classes/CookerSettings.h
--- a/Engine/Source/Developer/DeveloperToolSettings/Classes/CookerSettings.h
+++ b/Engine/Source/Developer/DeveloperToolSettings/Classes/CookerSettings.h
@@ -114,12 +114,12 @@ public:
 	UPROPERTY(GlobalConfig, EditAnywhere, Category = Cooker, meta = (DisplayName = "Generate DDC data in background for desired launch on platform (speeds up launch on)"))
 	bool bEnableBuildDDCInBackground;
 
-	/** Enable -iterate for launch on */
-	UPROPERTY(GlobalConfig, EditAnywhere, Category = Cooker, meta = (DisplayName = "LegacyIterative cooking for builds launched from the editor (launch on)"))
+	/** Enable IncrementalCooks for launch on */
+	UPROPERTY(GlobalConfig, EditAnywhere, Category = Cooker, meta = (DisplayName = "Use Incremental Cook for builds launched from the editor (launch on)"))
 	bool bIterativeCookingForLaunchOn;
 
-	/** Enable -iterate when triggering from Platforms dropdown menu */
-	UPROPERTY(GlobalConfig, EditAnywhere, Category = Cooker, meta = (DisplayName = "LegacyIterative cooking for the Platforms->Cook Content menu item"))
+	/** Enable IncrementalCooks when triggering from Platforms dropdown menu */
+	UPROPERTY(GlobalConfig, EditAnywhere, Category = Cooker, meta = (DisplayName = "Use Incremental Cook for the Platforms->Cook Content menu item"))
 	bool bIterativeCookingForFileCookContent;
 
 	/** Enable -cookonthefly for launch on */
diff --git a/Engine/Source/Editor/UnrealEd/Classes/CookOnTheSide/CookOnTheFlyServer.h b/Engine/Source/Editor/UnrealEd/Classes/CookOnTheSide/CookOnTheFlyServer.h
--- a/Engine/Source/Editor/UnrealEd/Classes/CookOnTheSide/CookOnTheFlyServer.h
+++ b/Engine/Source/Editor/UnrealEd/Classes/CookOnTheSide/CookOnTheFlyServer.h
@@ -1624,6 +1624,7 @@ private:
 	 * incrementally in the current mode; see PlatformManager.GetPlatformData().bFullCook for that property.
 	 */
 	bool bLegacyBuildDependencies = false;
+	bool bLegacyBuildDependenciesDueToNoZenStore = false;
 	/**
 	 * Automated DependencyTracking is only used in CookIncremental; legacyiterative does not use its data.
 	 * Additionally, it can be turned off to improve performance in CookIncremental, at the cost of not being able to
diff --git a/Engine/Source/Editor/UnrealEd/Private/CookOnTheFlyServer.cpp b/Engine/Source/Editor/UnrealEd/Private/CookOnTheFlyServer.cpp
--- a/Engine/Source/Editor/UnrealEd/Private/CookOnTheFlyServer.cpp
+++ b/Engine/Source/Editor/UnrealEd/Private/CookOnTheFlyServer.cpp
@@ -10238,6 +10238,7 @@ struct FDependencyTrackingConfigParser
 
 	// Output environment settings
 	bool bOutLegacyBuildDependencies = false;
+	bool bOutLegacyBuildDependenciesUsedDueToNoZenStore = false;
 	bool bOutDependencyTrackingEnabled = true;
 };
 
@@ -10263,6 +10264,7 @@ void FDependencyTrackingConfigParser::Parse()
 			TEXT("and we do not yet support storage of Incremental Cook dependency data with loose cooked packages. ")
 			TEXT("Falling back to legacy iterative cook mode."));
 		bOutLegacyBuildDependencies = true;
+		bOutLegacyBuildDependenciesUsedDueToNoZenStore = true;
 	}
 	else if (!bOutLegacyBuildDependencies)
 	{
@@ -10299,6 +10301,7 @@ void FBeginCookConfigSettings::LoadLocal(FBeginCookContext& BeginContext)
 	Parser.bInUsingZenStore.Emplace(BeginContext.COTFS.IsUsingZenStore());
 	Parser.Parse();
 	bLegacyBuildDependencies = Parser.bOutLegacyBuildDependencies;
+	bLegacyBuildDependenciesDueToNoZenStore = Parser.bOutLegacyBuildDependenciesUsedDueToNoZenStore;
 	bDependencyTrackingEnabled = Parser.bOutDependencyTrackingEnabled;
 
 	const TCHAR* CommandLine = FCommandLine::Get();
@@ -10326,6 +10329,7 @@ void FBeginCookConfigSettings::LoadLocal(FBeginCookContext& BeginContext)
 void UCookOnTheFlyServer::SetBeginCookConfigSettings(FBeginCookContext& BeginContext, UE::Cook::FBeginCookConfigSettings&& Settings)
 {
 	bLegacyBuildDependencies = Settings.bLegacyBuildDependencies;
+	bLegacyBuildDependenciesDueToNoZenStore = Settings.bLegacyBuildDependenciesDueToNoZenStore;
 	bDependencyTrackingEnabled = Settings.bDependencyTrackingEnabled;
 	bCookIncrementalAllowAllClasses = Settings.bCookIncrementalAllowAllClasses;
 	PackageDatas->SetBeginCookConfigSettings(Settings.CookShowInstigator);
@@ -10442,6 +10446,16 @@ void UCookOnTheFlyServer::LoadBeginCookIncrementalFlagsLocal(FBeginCookContext&
 	const bool bIsSharedLegacyIterativeCook = bLegacyBuildDependencies &&
 		IsCookFlagSet(ECookInitializationFlags::LegacyIterativeSharedBuild);
 
+	if (bIncrementalOrLegacyIterative && bLegacyBuildDependenciesDueToNoZenStore
+		&& EnumHasAnyFlags(CookFlags, ECookInitializationFlags::Unversioned))
+	{
+		UE_LOG(LogCook, Warning, TEXT("CookIncremental and UnversionedCookedContent were requested, but LegacyIterative is required due to ZenStore not being used. ")
+			TEXT("LegacyIterative incremental cooks do not robustly support unversioned cooked content due to failure to track dependencies on c++ serialization changes. ")
+			TEXT("The cooker will cook with these settings, but may fail to recook packages that require it those packages will cause runtime crashes when loaded. ")
+			TEXT("For robust incremental cooks, set bUseZenStore=true and bUseIoStore=true in ProjectPackagingSettings, or turn off UnversionedCookedContent.")
+		);
+	}
+
 	for (FBeginCookContextPlatform& PlatformContext : BeginContext.PlatformContexts)
 	{
 		const ITargetPlatform* TargetPlatform = PlatformContext.TargetPlatform;
diff --git a/Engine/Source/Editor/UnrealEd/Private/Cooker/CookTypes.cpp b/Engine/Source/Editor/UnrealEd/Private/Cooker/CookTypes.cpp
--- a/Engine/Source/Editor/UnrealEd/Private/Cooker/CookTypes.cpp
+++ b/Engine/Source/Editor/UnrealEd/Private/Cooker/CookTypes.cpp
@@ -878,6 +878,7 @@ void FInitializeConfigSettings::MoveToLocal(UCookOnTheFlyServer& COTFS)
 void FBeginCookConfigSettings::CopyFromLocal(const UCookOnTheFlyServer& COTFS)
 {
 	bLegacyBuildDependencies = COTFS.bLegacyBuildDependencies;
+	bLegacyBuildDependenciesDueToNoZenStore = COTFS.bLegacyBuildDependenciesDueToNoZenStore;
 	bDependencyTrackingEnabled = COTFS.bDependencyTrackingEnabled;
 	bCookIncrementalAllowAllClasses = COTFS.bCookIncrementalAllowAllClasses;
 	FParse::Value(FCommandLine::Get(), TEXT("-CookShowInstigator="), CookShowInstigator); // We don't store this on COTFS, so reparse it from commandLine
@@ -915,6 +916,7 @@ FCbWriter& operator<<(FCbWriter& Writer, const UE::Cook::FBeginCookConfigSetting
 {
 	Writer.BeginObject();
 	Writer << "LegacyBuildDependencies" << Value.bLegacyBuildDependencies;
+	Writer << "LegacyBuildDependenciesDueToNoZenStore" << Value.bLegacyBuildDependenciesDueToNoZenStore;
 	Writer << "DependencyTrackingEnabled" << Value.bDependencyTrackingEnabled;
 	Writer << "CookIncrementalAllowAllClasses" << Value.bCookIncrementalAllowAllClasses;
 	Writer << "CookShowInstigator" << Value.CookShowInstigator;
@@ -938,6 +940,7 @@ bool LoadFromCompactBinary(FCbFieldView Field, UE::Cook::FBeginCookConfigSetting
 {
 	bool bOk = Field.IsObject();
 	bOk = LoadFromCompactBinary(Field["LegacyBuildDependencies"], OutValue.bLegacyBuildDependencies) & bOk;
+	bOk = LoadFromCompactBinary(Field["LegacyBuildDependenciesDueToNoZenStore"], OutValue.bLegacyBuildDependenciesDueToNoZenStore) & bOk;
 	bOk = LoadFromCompactBinary(Field["DependencyTrackingEnabled"], OutValue.bDependencyTrackingEnabled) & bOk;
 	bOk = LoadFromCompactBinary(Field["CookIncrementalAllowAllClasses"], OutValue.bCookIncrementalAllowAllClasses) & bOk;
 	bOk = LoadFromCompactBinary(Field["CookShowInstigator"], OutValue.CookShowInstigator) & bOk;
diff --git a/Engine/Source/Editor/UnrealEd/Private/Cooker/CookTypes.h b/Engine/Source/Editor/UnrealEd/Private/Cooker/CookTypes.h
--- a/Engine/Source/Editor/UnrealEd/Private/Cooker/CookTypes.h
+++ b/Engine/Source/Editor/UnrealEd/Private/Cooker/CookTypes.h
@@ -499,6 +499,7 @@ namespace UE::Cook
 
 		FString CookShowInstigator;
 		bool bLegacyBuildDependencies = false;
+		bool bLegacyBuildDependenciesDueToNoZenStore = false;
 		bool bDependencyTrackingEnabled = true;
 		bool bCookIncrementalAllowAllClasses = false;
 		TArray<FName> NeverCookPackageList;
diff --git a/Engine/Source/Editor/UnrealEd/Private/PlayLevelViaLauncher.cpp b/Engine/Source/Editor/UnrealEd/Private/PlayLevelViaLauncher.cpp
--- a/Engine/Source/Editor/UnrealEd/Private/PlayLevelViaLauncher.cpp
+++ b/Engine/Source/Editor/UnrealEd/Private/PlayLevelViaLauncher.cpp
@@ -250,7 +250,9 @@ void UEditorEngine::StartPlayUsingLauncherSession(FRequestPlaySessionParams& InR
 	bool bLegacyIterative = false;
 	GConfig->GetBool(TEXT("CookSettings"), TEXT("LegacyIterative"), bLegacyIterative, GEditorIni);
 
-	bool bIncrementalCooking = (CookerSettings.bIterativeCookingForLaunchOn || ExperimentalSettings.bSharedCookedBuilds) && !bCookOnTheFly;
+	bool bIncrementalCooking =
+		!bCookOnTheFly
+		&& (CookerSettings.bIterativeCookingForLaunchOn || (bLegacyIterative && ExperimentalSettings.bSharedCookedBuilds));
 
 	if (CurrentLauncherCookMode == ELauncherProfileCookModes::OnTheFlyInEditor ||
 		CurrentLauncherCookMode == ELauncherProfileCookModes::ByTheBookInEditor)
diff --git a/Engine/Source/Editor/UnrealEd/Private/UnrealEdEngine.cpp b/Engine/Source/Editor/UnrealEd/Private/UnrealEdEngine.cpp
--- a/Engine/Source/Editor/UnrealEd/Private/UnrealEdEngine.cpp
+++ b/Engine/Source/Editor/UnrealEd/Private/UnrealEdEngine.cpp
@@ -229,11 +229,23 @@ void UUnrealEdEngine::Init(IEngineLoop* InEngineLoop)
 		ECookInitializationFlags BaseCookingFlags = ECookInitializationFlags::AutoTick | ECookInitializationFlags::AsyncSave;
 		BaseCookingFlags |= CookerSettings->bEnableBuildDDCInBackground ? ECookInitializationFlags::BuildDDCInBackground : ECookInitializationFlags::None;
 
+		bool bLegacyIterative = false;
+		GConfig->GetBool(TEXT("CookSettings"), TEXT("LegacyIterative"), bLegacyIterative, GEditorIni);
+
 		if (CookerSettings->bIterativeCookingForLaunchOn)
 		{
+			// Despite the name, the LegacyIterative cook flag does not cause CookSettings.LegacyIterative, aka
+			// bLegacyBuildDependencies. It triggers the same code as the "-iterate" commandline, which causes incremental
+			// cooks when LegacyIterative=false, and causes the cook to cook iteratively using the legacy method when
+			// LegacyIterative=true.
 			BaseCookingFlags |= ECookInitializationFlags::LegacyIterative;
-			BaseCookingFlags |= CookerSettings->bIgnoreIniSettingsOutOfDateForIteration ? ECookInitializationFlags::IgnoreIniSettingsOutOfDate : ECookInitializationFlags::None;
-			BaseCookingFlags |= CookerSettings->bIgnoreScriptPackagesOutOfDateForIteration ? ECookInitializationFlags::IgnoreScriptPackagesOutOfDate : ECookInitializationFlags::None;
+			if (bLegacyIterative)
+			{
+				BaseCookingFlags |= CookerSettings->bIgnoreIniSettingsOutOfDateForIteration
+					? ECookInitializationFlags::IgnoreIniSettingsOutOfDate : ECookInitializationFlags::None;
+				BaseCookingFlags |= CookerSettings->bIgnoreScriptPackagesOutOfDateForIteration
+					? ECookInitializationFlags::IgnoreScriptPackagesOutOfDate : ECookInitializationFlags::None;
+			}
 		}
 		
 		if (CookerSettings->bEnableCookOnTheSide)
diff --git a/Engine/Source/Programs/AutomationTool/AutomationUtils/ProjectParams.cs b/Engine/Source/Programs/AutomationTool/AutomationUtils/ProjectParams.cs
--- a/Engine/Source/Programs/AutomationTool/AutomationUtils/ProjectParams.cs
+++ b/Engine/Source/Programs/AutomationTool/AutomationUtils/ProjectParams.cs
@@ -342,7 +342,7 @@ namespace AutomationTool
 			this.SkipBuildEditor = InParams.SkipBuildEditor;
 			this.Run = InParams.Run;
 			this.Cook = InParams.Cook;
-			this.IterativeCooking = InParams.IterativeCooking;
+			this.IncrementalCook = InParams.IncrementalCook;
 			this.IterateSharedCookedBuild = InParams.IterateSharedCookedBuild;
 			this.IterateSharedBuildUsePrecompiledExe = InParams.IterateSharedBuildUsePrecompiledExe;
             this.CookAll = InParams.CookAll;
@@ -412,7 +412,6 @@ namespace AutomationTool
 			this.KeepFileOpenLog = InParams.KeepFileOpenLog;
 			this.FastCook = InParams.FastCook;
 			this.Snapshot = InParams.Snapshot;
-			this.IncrementalCook = InParams.IncrementalCook;
 			this.Devices = InParams.Devices;
 			this.DeviceNames = InParams.DeviceNames;
 			this.ServerDevice = InParams.ServerDevice;
@@ -521,7 +520,7 @@ namespace AutomationTool
 			string AdditionalPakOptions = null,
 			string AdditionalIoStoreOptions = null,
 			string ForceOodleDllVersion = null,
-            bool? IterativeCooking = null,
+			bool? IncrementalCook = null,
 			string IterateSharedCookedBuild = null,
 			bool? IterateSharedBuildUsePrecompiledExe = null,
 			bool? CookAll = null,
@@ -636,7 +635,6 @@ namespace AutomationTool
 			bool? IterativeDeploy = null,
 			bool? FastCook = null,
 			bool? Snapshot = null,
-			bool? IncrementalCook = null,
 			bool? IgnoreCookErrors = null,
 			bool? KeepFileOpenLog = null,
 			bool? CodeSign = null,
@@ -865,7 +863,7 @@ namespace AutomationTool
 			}
 			this.AdditionalIoStoreOptions = ParseParamValueIfNotSpecified(Command, AdditionalIoStoreOptions, "AdditionalIoStoreOptions");
 			this.ForceOodleDllVersion = ParseParamValueIfNotSpecified(Command, ForceOodleDllVersion, "ForceOodleDllVersion");
-			this.IterativeCooking = GetParamValueIfNotSpecified(Command, IterativeCooking, this.IterativeCooking, new string[] { "iterativecooking", "iterate" });
+			this.IncrementalCook = GetParamValueIfNotSpecified(Command, IncrementalCook, this.IncrementalCook, new string[] { "iterativecooking", "iterate", "cookincremental" });
 			this.IterateSharedCookedBuild = GetParamValueIfNotSpecified(Command, false, false, "iteratesharedcookedbuild") ? "usesyncedbuild" : null;
 			this.IterateSharedCookedBuild = ParseParamValueIfNotSpecified(Command, IterateSharedCookedBuild, "IterateSharedCookedBuild", String.Empty);
 			this.IterateSharedBuildUsePrecompiledExe = GetParamValueIfNotSpecified(Command, IterateSharedBuildUsePrecompiledExe, this.IterateSharedBuildUsePrecompiledExe, new string[] { "IterateSharedBuildUsePrecompiledExe" });
@@ -1010,7 +1008,7 @@ namespace AutomationTool
 			this.Snapshot = GetParamValueIfNotSpecified(Command, Snapshot, this.Snapshot, "Snapshot");
 
 			// -snapshot implies the use of -cookincremental for cooking.
-			this.IncrementalCook = this.Snapshot || GetParamValueIfNotSpecified(Command, IncrementalCook, this.IncrementalCook, "cookincremental");
+			this.IncrementalCook = this.Snapshot || this.IncrementalCook;
 			this.IgnoreCookErrors = GetParamValueIfNotSpecified(Command, IgnoreCookErrors, this.IgnoreCookErrors, "IgnoreCookErrors");
 			this.KeepFileOpenLog = GetParamValueIfNotSpecified(Command, KeepFileOpenLog, this.KeepFileOpenLog, "KeepFileOpenLog");
 
@@ -2101,10 +2099,10 @@ namespace AutomationTool
 		public bool OptionalContent = false;
 
 		/// <summary>
-		/// Cook: Uses the iterative cooking, command line: -iterativecooking or -iterate
+		/// Cook: Cooks incrementally using previously-cooked data or a published snapshot as a base
 		/// </summary>
-		[Help( "iterativecooking", "Uses the iterative cooking, command line: -iterativecooking or -iterate" )]
-		public bool IterativeCooking;
+		[Help("cookincremental", "Cooks incrementally using previously-cooked data or a published snapshot as a base")]
+		public bool IncrementalCook;
 
 		/// <summary>
 		/// Cook: Iterate from a build on the network
@@ -2143,7 +2141,7 @@ namespace AutomationTool
 		/// <summary>
 		/// Cook: Uses the iterative deploy, command line: -iterativedeploy or -iterate
 		/// </summary>
-		[Help("iterativecooking", "Uses the iterative cooking, command line: -iterativedeploy or -iterate")]
+		[Help("IterativeDeploy", "Runs UnrealPak and targetplatform install iteratively, command line: -iterativedeploy or -iterate")]
 		public bool IterativeDeploy;
 
 		[Help("FastCook", "Uses fast cook path if supported by target")]
@@ -2155,12 +2153,6 @@ namespace AutomationTool
 		[Help("snapshot", "Imports the most recently published project snapshot (if cooking, this forces the use of -cookincremental)")]
 		public bool Snapshot;
 
-		/// <summary>
-		/// Cook: Cooks incrementally using previously-cooked data or a published snapshot as a base
-		/// </summary>
-		[Help("cookincremental", "Cooks incrementally using previously-cooked data or a published snapshot as a base")]
-		public bool IncrementalCook;
-
 		/// <summary>
 		/// Cook: Ignores cook errors and continues with packaging etc.
 		/// </summary>
@@ -3350,7 +3342,7 @@ namespace AutomationTool
                 throw new AutomationException("Can't create a release version at the same time as creating dlc.");
             }
 
-            if (HasBasedOnReleaseVersion && (IterativeCooking || IterativeDeploy || HasIterateSharedCookedBuild))
+            if (HasBasedOnReleaseVersion && (IncrementalCook || IterativeDeploy || HasIterateSharedCookedBuild))
             {
                 throw new AutomationException("Can't use iterative cooking / deploy on dlc or patching or creating a release");
             }
@@ -3477,7 +3469,7 @@ namespace AutomationTool
 				Logger.LogDebug("Foreign={Foreign}", Foreign);
 				Logger.LogDebug("IsCodeBasedProject={Arg0}", IsCodeBasedProject.ToString());
 				Logger.LogDebug("IsProgramTarget={Arg0}", IsProgramTarget.ToString());
-				Logger.LogDebug("IterativeCooking={IterativeCooking}", IterativeCooking);
+				Logger.LogDebug("IncrementalCook={IncrementalCook}", IncrementalCook);
 				Logger.LogDebug("IterateSharedCookedBuild={IterateSharedCookedBuild}", IterateSharedCookedBuild);
 				Logger.LogDebug("IterateSharedBuildUsePrecompiledExe={IterateSharedBuildUsePrecompiledExe}", IterateSharedBuildUsePrecompiledExe);
 				Logger.LogDebug("CookAll={CookAll}", CookAll);
diff --git a/Engine/Source/Programs/AutomationTool/Scripts/CookCommand.Automation.cs b/Engine/Source/Programs/AutomationTool/Scripts/CookCommand.Automation.cs
--- a/Engine/Source/Programs/AutomationTool/Scripts/CookCommand.Automation.cs
+++ b/Engine/Source/Programs/AutomationTool/Scripts/CookCommand.Automation.cs
@@ -1,20 +1,20 @@
 // Copyright Epic Games, Inc. All Rights Reserved.
+using AutomationTool;
+using EpicGames.Core;
+using EpicGames.ProjectStore;
+using Microsoft.Extensions.Logging;
 using System;
+using System.Collections.Concurrent;
 using System.Collections.Generic;
 using System.IO;
-using System.Threading;
-using System.Threading.Tasks;
-using System.Reflection;
 using System.Linq;
+using System.Reflection;
+using System.Runtime.InteropServices;
 using System.Text;
-using AutomationTool;
-using UnrealBuildTool;
-using System.Collections.Concurrent;
-using EpicGames.Core;
-using EpicGames.ProjectStore;
+using System.Threading;
+using System.Threading.Tasks;
 using UnrealBuildBase;
-using Microsoft.Extensions.Logging;
-using System.Runtime.InteropServices;
+using UnrealBuildTool;
 
 namespace AutomationScripts
 {
@@ -27,10 +27,15 @@ namespace AutomationScripts
 	/// </remarks>
 	public partial class Project : CommandUtils
 	{
-		private static string GetGenericCookCommandletParams(ProjectParams Params)
+		private static string GetGenericCookCommandletParams(ProjectParams Params, ConfigHierarchy EditorConfig)
 		{
 			string CommandletParams = "";
 
+			if (!EditorConfig.TryGetValue("CookSettings", "LegacyIterative", out bool bLegacyIterative))
+			{
+				bLegacyIterative = false;
+			}
+
 			if (IsBuildMachine)
 			{
 				CommandletParams += " -buildmachine";
@@ -45,7 +50,17 @@ namespace AutomationScripts
 			}
 			if (Params.UnversionedCookedContent)
 			{
-				CommandletParams += " -unversioned";
+				if (bLegacyIterative && Params.IncrementalCook)
+				{
+					Logger.LogInformation("UnversionedCookedContent was requested, but IncrementalCook is enabled and ini settings specify LegacyIterative "
+						+ "(Editor.ini:[CookSettings]:LegacyIterative=true). "
+						+ "UnversionedCookedContent is not supported with legacy iterative because any serialization change causes the IncrementalCook to crash. "
+						+ "Disabling UnversionedCookedContent.");
+				}
+				else
+				{
+					CommandletParams += " -unversioned";
+				}
 			}
 			if (Params.OptionalContent)
 			{
@@ -56,10 +71,6 @@ namespace AutomationScripts
 				CommandletParams += " -FastCook";
 			}
 			if (Params.IncrementalCook)
-			{
-				CommandletParams += " -forcerecook=false";
-			}
-			if (Params.IterativeCooking)
 			{
 				CommandletParams += " -iterate";
 			}
@@ -132,12 +143,17 @@ namespace AutomationScripts
 			return PlatformsToCookSet;
 		}
 
-		public static void GetCookByTheBookCommandletParams(ProjectParams Params, out string CommandletParams, out string[] Maps, out string[] DirectoriesToCook, out string InternationalizationPreset, out string[] CulturesToCook, out string PlatformsToCook)
+		public static void GetCookByTheBookCommandletParams(ProjectParams Params,
+			out string CommandletParams, out string[] Maps, out string[] DirectoriesToCook,
+			out string InternationalizationPreset, out string[] CulturesToCook, out string PlatformsToCook)
 		{
+			ConfigHierarchy EditorConfig = ConfigCache.ReadHierarchy(ConfigHierarchyType.Editor,
+				Params.RawProjectPath.Directory, HostPlatform.Current.HostEditorPlatform);
+
 			var PlatformsToCookSet = GetPlatformsToCook(Params);
 			PlatformsToCook = CombineCommandletParams(PlatformsToCookSet);
 
-			if (Params.Clean.HasValue && Params.Clean.Value && !Params.IterativeCooking)
+			if (Params.Clean.HasValue && Params.Clean.Value && !Params.IncrementalCook)
 			{
 				Logger.LogInformation("Cleaning cooked data.");
 				CleanupCookedData(PlatformsToCookSet.ToList(), Params);
@@ -176,7 +192,7 @@ namespace AutomationScripts
 				CulturesToCook = Params.CulturesToCook.ToArray();
 			}
 
-			CommandletParams = GetGenericCookCommandletParams(Params);
+			CommandletParams = GetGenericCookCommandletParams(Params, EditorConfig);
 
 			if (Params.KeepFileOpenLog)
 			{
@@ -296,7 +312,9 @@ namespace AutomationScripts
 					CreateDirectory(LogFolderOutsideOfSandbox);
 				}
 
-				string COTFCommandLine = GetGenericCookCommandletParams(Params);
+				ConfigHierarchy EditorConfig = ConfigCache.ReadHierarchy(ConfigHierarchyType.Editor,
+					Params.RawProjectPath.Directory, HostPlatform.Current.HostEditorPlatform);
+				string COTFCommandLine = GetGenericCookCommandletParams(Params, EditorConfig);
 				if (Params.ZenStore)
 				{
 					COTFCommandLine += " -messaging";
@@ -317,7 +335,9 @@ namespace AutomationScripts
 					string PlatformsToCook;
 
 					// get the parameters to CBTB
-					GetCookByTheBookCommandletParams(Params, out CommandletParams, out MapsToCook, out DirectoriesToCook, out InternationalizationPreset, out CulturesToCook, out PlatformsToCook);
+					GetCookByTheBookCommandletParams(Params, out CommandletParams, out MapsToCook,
+						out DirectoriesToCook, out InternationalizationPreset, out CulturesToCook,
+						out PlatformsToCook);
 
 					// run a blocking cook
 					CookCommandlet(Params.RawProjectPath, Params.UnrealExe, MapsToCook, DirectoriesToCook, InternationalizationPreset, CulturesToCook, PlatformsToCook, CommandletParams);
