Fixing typo "VisualStudio20122" -> "VisualStudio2022"

#Synced-CL 51674122

diff --git a/Engine/Source/Programs/AutomationTool/Scripts/BuildHlslcc.Automation.cs b/Engine/Source/Programs/AutomationTool/Scripts/BuildHlslcc.Automation.cs
--- a/Engine/Source/Programs/AutomationTool/Scripts/BuildHlslcc.Automation.cs
+++ b/Engine/Source/Programs/AutomationTool/Scripts/BuildHlslcc.Automation.cs
@@ -315,7 +315,7 @@ class BuildHlslcc : BuildCommand
 	private List<WindowsCompiler> GetTargetWindowsCompilers()
 	{
 		List<WindowsCompiler> TargetWindowsCompilers = new List<WindowsCompiler>();
-		string TargetWindowsCompilersFilter = ParseParamValue("TargetWindowsCompilers", "VisualStudio20122");
+		string TargetWindowsCompilersFilter = ParseParamValue("TargetWindowsCompilers", "VisualStudio2022");
 		if (TargetWindowsCompilersFilter != null)
 		{
 			foreach (string TargetWindowsCompilerName in TargetWindowsCompilersFilter.Split(new char[] { '+' }, StringSplitOptions.RemoveEmptyEntries))
