Rename format placeholder name File

#rnx
#Synced-CL 51691096

diff --git a/Engine/Source/Programs/AutomationTool/BuildGraph/Tasks/SanitizeReceiptTask.cs b/Engine/Source/Programs/AutomationTool/BuildGraph/Tasks/SanitizeReceiptTask.cs
--- a/Engine/Source/Programs/AutomationTool/BuildGraph/Tasks/SanitizeReceiptTask.cs
+++ b/Engine/Source/Programs/AutomationTool/BuildGraph/Tasks/SanitizeReceiptTask.cs
@@ -96,7 +96,7 @@ namespace AutomationTool.Tasks
 						}
 						else
 						{
-							Logger.LogInformation("Removing build product: {File}", buildProduct.Path);
+							Logger.LogInformation("Removing build product: {BuildProduct}", buildProduct.Path);
 						}
 					}
 					receipt.BuildProducts = newBuildProducts;
@@ -111,7 +111,7 @@ namespace AutomationTool.Tasks
 						}
 						else
 						{
-							Logger.LogInformation("Removing runtime dependency: {File}", runtimeDependency.Path);
+							Logger.LogInformation("Removing runtime dependency: {RuntimeDependency}", runtimeDependency.Path);
 						}
 					}
 					receipt.RuntimeDependencies = newRuntimeDependencies;
