Allow to reorder phases per end date moving unknown and not started at the end

#rnx

#Synced-CL 51708345

diff --git a/Engine/Source/Programs/AutomationTool/Gauntlet/Framework/Base/Gauntlet.TestTracking.cs b/Engine/Source/Programs/AutomationTool/Gauntlet/Framework/Base/Gauntlet.TestTracking.cs
--- a/Engine/Source/Programs/AutomationTool/Gauntlet/Framework/Base/Gauntlet.TestTracking.cs
+++ b/Engine/Source/Programs/AutomationTool/Gauntlet/Framework/Base/Gauntlet.TestTracking.cs
@@ -192,6 +192,7 @@ namespace Gauntlet
 			private Queue<Phase> _queuedPhases;
 			private Phase? _last;
 			private Phase? _current;
+			private bool _orderByCompletion = false;
 
 			/// <summary>
 			/// The last added phase to the queue
@@ -291,9 +292,21 @@ namespace Gauntlet
 			/// <inheritdoc/>
 			public IEnumerator<Phase> GetEnumerator()
 			{
+				if (_orderByCompletion)
+				{
+					return _queuedPhases.OrderBy(P => !P.StartedAt.HasValue).ThenBy(P => P.Outcome == TestPhaseOutcome.Unknown).ThenBy(P => P.StartedAt?.AddSeconds(P.ElapseTimeInSeconds)).GetEnumerator();
+				}
+
 				return _queuedPhases.GetEnumerator();
 			}
 
+			/// <summary>
+			/// Whether or not to order the phases by completion date when iterating the object
+			/// </summary>
+			public void SetOrderByCompletion(bool Value)
+			{
+				_orderByCompletion = Value;
+			}
 
 			/// <summary>
 			/// Clear the queue
