[Gauntlet] Fixing a bug that would prevent gauntlet from accepting multiple -dpcvars arguments. It would previously drop everything after the first in an effort to deduplicate arguments. Preflight passed far enough to validate the fix.

#Synced-CL 51657680

diff --git a/Engine/Source/Programs/AutomationTool/Gauntlet/Unreal/Base/Gauntlet.UnrealTestConfiguration.cs b/Engine/Source/Programs/AutomationTool/Gauntlet/Unreal/Base/Gauntlet.UnrealTestConfiguration.cs
--- a/Engine/Source/Programs/AutomationTool/Gauntlet/Unreal/Base/Gauntlet.UnrealTestConfiguration.cs
+++ b/Engine/Source/Programs/AutomationTool/Gauntlet/Unreal/Base/Gauntlet.UnrealTestConfiguration.cs
@@ -100,7 +100,7 @@ namespace Gauntlet
 		{
 			if (Params.ContainsKey(ParamName))
 			{
-				if (ParamName.ToLower() == "execcmds" && ParamVal != null)
+				if ((ParamName.ToLower() == "execcmds" || ParamName.ToLower() == "dpcvars") && ParamVal != null)
 				{
 					AddOrAppendParamValue(ParamName, ParamVal.ToString());
 					return true;
@@ -137,7 +137,7 @@ namespace Gauntlet
 
 			if (Params.ContainsKey(ParamName))
 			{
-				if (ParamName.ToLower() == "execcmds" && ParamVal != null)
+				if ((ParamName.ToLower() == "execcmds" || ParamName.ToLower() == "dpcvars") && ParamVal != null)
 				{
 					AddOrAppendParamValue(ParamName, ParamVal.ToString());
 					return;
diff --git a/Engine/Source/Programs/AutomationTool/Gauntlet/Unreal/BuildSource/Gauntlet.UnrealBuildSource.cs b/Engine/Source/Programs/AutomationTool/Gauntlet/Unreal/BuildSource/Gauntlet.UnrealBuildSource.cs
--- a/Engine/Source/Programs/AutomationTool/Gauntlet/Unreal/BuildSource/Gauntlet.UnrealBuildSource.cs
+++ b/Engine/Source/Programs/AutomationTool/Gauntlet/Unreal/BuildSource/Gauntlet.UnrealBuildSource.cs
@@ -751,6 +751,11 @@ namespace Gauntlet
 						// Special cased as execcmds is something that is totally able to be appended to. Requote everything when we're done.
 						CommandlineTokens[Name][0] = string.Format("\"{0}, {1}\"", CommandlineTokens[Name][0].Replace("\"", ""), Params.Replace("\"", ""));
 					}
+					else if (Name.ToLower() == "dpcvars")
+					{
+						// dpcvars is also appendable but is slightly different in syntax. It doesn't use quotes and cannot contain spaces
+						CommandlineTokens[Name][0] = string.Format("{0},{1}", CommandlineTokens[Name][0], Params);
+					}
 					else
 					{
 						bool bFoundDuplicateValue = false;
