Add -SkipEncryption support for IoStore containers.

[RN] fix Foundation/Core

#Synced-CL 51706094

diff --git a/Engine/Source/Programs/AutomationTool/Scripts/CopyBuildToStagingDirectory.Automation.cs b/Engine/Source/Programs/AutomationTool/Scripts/CopyBuildToStagingDirectory.Automation.cs
--- a/Engine/Source/Programs/AutomationTool/Scripts/CopyBuildToStagingDirectory.Automation.cs
+++ b/Engine/Source/Programs/AutomationTool/Scripts/CopyBuildToStagingDirectory.Automation.cs
@@ -4548,12 +4548,12 @@ namespace AutomationScripts
 					+ " " + AdditionalCompressionOptionsOnCommandLine
 					+ " " + Params.AdditionalPakOptions;
 
-				if (CryptoKeysCacheFilename != null)
+				if (!Params.SkipEncryption && CryptoKeysCacheFilename != null)
 				{
 					AdditionalArgs += String.Format(" -cryptokeys={0}", CommandUtils.MakePathSafeToUseWithCommandLine(CryptoKeysCacheFilename.FullName));
 				}
 
-				if (CryptoSettings != null)
+				if (!Params.SkipEncryption && CryptoSettings != null)
 				{
 					if (CryptoSettings.bDataCryptoRequired && CryptoSettings.bEnablePakSigning && CryptoSettings.SigningKey.IsValid())
 					{
