update wording when trying to make a GDK package with ARM64

#rnx

#Synced-CL 51451766

diff --git a/Engine/Source/Programs/AutomationTool/Win/MSGaming/MSGameStoreCustomDeploymentHandler.cs b/Engine/Source/Programs/AutomationTool/Win/MSGaming/MSGameStoreCustomDeploymentHandler.cs
--- a/Engine/Source/Programs/AutomationTool/Win/MSGaming/MSGameStoreCustomDeploymentHandler.cs
+++ b/Engine/Source/Programs/AutomationTool/Win/MSGaming/MSGameStoreCustomDeploymentHandler.cs
@@ -36,7 +36,7 @@ public class MSGameStoreCustomDeploymentHandler : CustomDeploymentHandler
 
 		if (GDKEdition < 251000 && Params.NoBootstrapExe && SC.StageTargets.Any(X => !X.Receipt.Architectures.SingleArchitecture.bIsX64))
 		{
-			throw new AutomationException("October 2025 GDK or higher is required to use the GDK with ARM64 unless you use the bootstrapper (omit -NoBootstrapExe)");
+			throw new AutomationException("October 2025 GDK or higher is required to create a GDK package with ARM64 unless you use the bootstrapper (omit -NoBootstrapExe)");
 		}
 	}
 
diff --git a/Engine/Source/Programs/AutomationTool/Win/MSGaming/MSGameStoreCustomStagingHandler.cs b/Engine/Source/Programs/AutomationTool/Win/MSGaming/MSGameStoreCustomStagingHandler.cs
--- a/Engine/Source/Programs/AutomationTool/Win/MSGaming/MSGameStoreCustomStagingHandler.cs
+++ b/Engine/Source/Programs/AutomationTool/Win/MSGaming/MSGameStoreCustomStagingHandler.cs
@@ -68,7 +68,7 @@ public class MSGameStoreCustomStagingHandler : CustomStagingHandler
 		// sanity check
 		if (GDKEdition < 251000 && Params.NoBootstrapExe && SC.StageTargets.Any(X => X.Receipt != null && !X.Receipt.Architectures.SingleArchitecture.bIsX64))
 		{
-			throw new AutomationException("October 2025 GDK or higher is required to use the GDK with ARM64 unless you use the bootstrapper (omit -NoBootstrapExe)");
+			throw new AutomationException("October 2025 GDK or higher is required to create a GDK package with ARM64 unless you use the bootstrapper (omit -NoBootstrapExe)");
 		}
 
 		return true;
