making sure dedicated server through RunUAT's BuildCookRun is using the right directory with macOS .app bundles to fix UE-313576

#Synced-CL 51715797

diff --git a/Engine/Source/Programs/AutomationTool/Scripts/RunProjectCommand.Automation.cs b/Engine/Source/Programs/AutomationTool/Scripts/RunProjectCommand.Automation.cs
--- a/Engine/Source/Programs/AutomationTool/Scripts/RunProjectCommand.Automation.cs
+++ b/Engine/Source/Programs/AutomationTool/Scripts/RunProjectCommand.Automation.cs
@@ -909,16 +909,30 @@ namespace AutomationScripts
 
 			var SC = DeployContextList[0];
 
+			TargetPlatformDescriptor ServerPlatformDesc = ServerParams.ServerTargetPlatforms[0];
 			var ServerApp = CombinePaths(CmdEnv.LocalRoot, "Engine/Binaries/Win64/UnrealEditor.exe");
 			bool bCooked = ServerParams.Cook || ServerParams.CookOnTheFly;
 			if (bCooked)
 			{
 				List<FileReference> Exes = SC.StageTargetPlatform.GetExecutableNames(SC);
 				ServerApp = Exes[0].FullName;
+				
+				// account for app bundles placed outside Binaries/Mac on macOS
+				if(ServerPlatformDesc.Type == UnrealTargetPlatform.Mac)
+				{
+					DirectoryReference BaseStageDir = new DirectoryReference(Params.BaseStageDirectory);
+					DirectoryReference StageSubDir = new FileReference(ServerApp).Directory;
+
+					while (StageSubDir.IsUnderDirectory(BaseStageDir) && StageSubDir.ParentDirectory != BaseStageDir)
+					{
+					    StageSubDir = StageSubDir.ParentDirectory;
+					}
+					ServerApp = FileReference.Combine(StageSubDir, 
+													  $"{Exes[0].GetFileNameWithoutExtension()}.app/Contents/MacOS/{Exes[0].GetFileNameWithoutExtension()}").FullName;
+				}
 			}
 			var Args = bCooked ? "" : (SC.ProjectArgForCommandLines + " ");
 			Console.WriteLine("Running dedicated server on device with address: " + Params.ServerDeviceAddress);
-			TargetPlatformDescriptor ServerPlatformDesc = ServerParams.ServerTargetPlatforms[0];
 			if (ServerParams.Cook && ServerPlatformDesc.Type == UnrealTargetPlatform.Linux && !String.IsNullOrEmpty(ServerParams.ServerDeviceAddress))
 			{
 				ServerApp = @"C:\Windows\system32\cmd.exe";
