Update AWSSDK.S3 package references to 3.7.501.14 to fix runtime dependency

#rnx
#Synced-CL 51784583

diff --git a/Engine/Source/Programs/AutomationTool/AutomationTool.csproj b/Engine/Source/Programs/AutomationTool/AutomationTool.csproj
--- a/Engine/Source/Programs/AutomationTool/AutomationTool.csproj
+++ b/Engine/Source/Programs/AutomationTool/AutomationTool.csproj
@@ -82,6 +82,7 @@
   </ItemGroup>
 
   <ItemGroup>
+    <PackageReference Include="AWSSDK.S3" Version="3.7.501.14" />
     <PackageReference Include="AWSSDK.SecurityToken" Version="3.7.501.14" />
     <PackageReference Include="Magick.NET-Q16-HDRI-AnyCPU" Version="14.9.1" />
     <PackageReference Include="Microsoft.Build" Version="17.11.48" ExcludeAssets="runtime" PrivateAssets="all">
diff --git a/Engine/Source/Programs/AutomationTool/Scripts/AutomationScripts.Automation.csproj b/Engine/Source/Programs/AutomationTool/Scripts/AutomationScripts.Automation.csproj
--- a/Engine/Source/Programs/AutomationTool/Scripts/AutomationScripts.Automation.csproj
+++ b/Engine/Source/Programs/AutomationTool/Scripts/AutomationScripts.Automation.csproj
@@ -34,6 +34,6 @@
   </ItemGroup>
   <ItemGroup>
     <PackageReference Include="Microsoft.VisualStudio.OLE.Interop" Version="17.14.40260" />
-    <PackageReference Include="AWSSDK.S3" Version="3.3.113.2" PrivateAssets="all" />
+    <PackageReference Include="AWSSDK.S3" Version="3.7.501.14" PrivateAssets="all" />
   </ItemGroup>
 </Project>
\ No newline at end of file
