Shared: Disable expensive C# analyzers

CA5390 and CA5403 are very expensive analyzers to run under some conditions, and provide very little value.
This is especially noticable when building UGS with analyzers. Build time for that single .csproj went from 117s -> 19s when disabled.
A full Horde.sln build 160s -> 92s.

diff --git a/Engine/Extras/P4VUtils/.editorconfig b/Engine/Extras/P4VUtils/.editorconfig
--- a/Engine/Extras/P4VUtils/.editorconfig
+++ b/Engine/Extras/P4VUtils/.editorconfig
@@ -769,8 +769,8 @@ dotnet_diagnostic.CA5388.severity = warning
 # CA5389: Do Not Add Archive Item's Path To The Target File System Path
 dotnet_diagnostic.CA5389.severity = warning
 
-# CA5390: Do not hard-code encryption key
-dotnet_diagnostic.CA5390.severity = warning
+# CA5390: Do not hard-code encryption key (disabled: expensive interprocedural dataflow analyzer)
+dotnet_diagnostic.CA5390.severity = none
 
 # CA5391: Use antiforgery tokens in ASP.NET Core MVC controllers
 dotnet_diagnostic.CA5391.severity = warning
@@ -808,8 +808,8 @@ dotnet_diagnostic.CA5401.severity = warning
 # CA5402: Use CreateEncryptor with the default IV 
 dotnet_diagnostic.CA5402.severity = warning
 
-# CA5403: Do not hard-code certificate
-dotnet_diagnostic.CA5403.severity = warning
+# CA5403: Do not hard-code certificate (disabled: expensive interprocedural dataflow analyzer)
+dotnet_diagnostic.CA5403.severity = none
 
 # CA5404: Do not disable token validation checks
 dotnet_diagnostic.CA5404.severity = warning
diff --git a/Engine/Source/Programs/AutomationTool/BuildGraph/.editorconfig b/Engine/Source/Programs/AutomationTool/BuildGraph/.editorconfig
--- a/Engine/Source/Programs/AutomationTool/BuildGraph/.editorconfig
+++ b/Engine/Source/Programs/AutomationTool/BuildGraph/.editorconfig
@@ -1050,8 +1050,8 @@ dotnet_diagnostic.CA5388.severity = warning
 # CA5389: Do Not Add Archive Item's Path To The Target File System Path
 dotnet_diagnostic.CA5389.severity = warning
 
-# CA5390: Do not hard-code encryption key
-dotnet_diagnostic.CA5390.severity = warning
+# CA5390: Do not hard-code encryption key (disabled: expensive interprocedural dataflow analyzer)
+dotnet_diagnostic.CA5390.severity = none
 
 # CA5391: Use antiforgery tokens in ASP.NET Core MVC controllers
 dotnet_diagnostic.CA5391.severity = warning
@@ -1089,8 +1089,8 @@ dotnet_diagnostic.CA5401.severity = warning
 # CA5402: Use CreateEncryptor with the default IV 
 dotnet_diagnostic.CA5402.severity = warning
 
-# CA5403: Do not hard-code certificate
-dotnet_diagnostic.CA5403.severity = warning
+# CA5403: Do not hard-code certificate (disabled: expensive interprocedural dataflow analyzer)
+dotnet_diagnostic.CA5403.severity = none
 
 # CA5404: Do not disable token validation checks
 dotnet_diagnostic.CA5404.severity = none
diff --git a/Engine/Source/Programs/Horde/.editorconfig b/Engine/Source/Programs/Horde/.editorconfig
--- a/Engine/Source/Programs/Horde/.editorconfig
+++ b/Engine/Source/Programs/Horde/.editorconfig
@@ -1052,8 +1052,8 @@ dotnet_diagnostic.CA5388.severity = warning
 # CA5389: Do Not Add Archive Item's Path To The Target File System Path
 dotnet_diagnostic.CA5389.severity = warning
 
-# CA5390: Do not hard-code encryption key
-dotnet_diagnostic.CA5390.severity = warning
+# CA5390: Do not hard-code encryption key (disabled: expensive interprocedural dataflow analyzer)
+dotnet_diagnostic.CA5390.severity = none
 
 # CA5391: Use antiforgery tokens in ASP.NET Core MVC controllers
 dotnet_diagnostic.CA5391.severity = warning
@@ -1091,8 +1091,8 @@ dotnet_diagnostic.CA5401.severity = warning
 # CA5402: Use CreateEncryptor with the default IV 
 dotnet_diagnostic.CA5402.severity = warning
 
-# CA5403: Do not hard-code certificate
-dotnet_diagnostic.CA5403.severity = warning
+# CA5403: Do not hard-code certificate (disabled: expensive interprocedural dataflow analyzer)
+dotnet_diagnostic.CA5403.severity = none
 
 # CA5404: Do not disable token validation checks
 dotnet_diagnostic.CA5404.severity = none
diff --git a/Engine/Source/Programs/Shared/.editorconfig b/Engine/Source/Programs/Shared/.editorconfig
--- a/Engine/Source/Programs/Shared/.editorconfig
+++ b/Engine/Source/Programs/Shared/.editorconfig
@@ -1084,8 +1084,8 @@ dotnet_diagnostic.CA5388.severity = warning
 # CA5389: Do Not Add Archive Item's Path To The Target File System Path
 dotnet_diagnostic.CA5389.severity = warning
 
-# CA5390: Do not hard-code encryption key
-dotnet_diagnostic.CA5390.severity = warning
+# CA5390: Do not hard-code encryption key (disabled: expensive interprocedural dataflow analyzer)
+dotnet_diagnostic.CA5390.severity = none
 
 # CA5391: Use antiforgery tokens in ASP.NET Core MVC controllers
 dotnet_diagnostic.CA5391.severity = warning
@@ -1123,8 +1123,8 @@ dotnet_diagnostic.CA5401.severity = warning
 # CA5402: Use CreateEncryptor with the default IV 
 dotnet_diagnostic.CA5402.severity = warning
 
-# CA5403: Do not hard-code certificate
-dotnet_diagnostic.CA5403.severity = warning
+# CA5403: Do not hard-code certificate (disabled: expensive interprocedural dataflow analyzer)
+dotnet_diagnostic.CA5403.severity = none
 
 # CA5404: Do not disable token validation checks
 dotnet_diagnostic.CA5404.severity = none
diff --git a/Engine/Source/Programs/UnrealCloudDDC/.editorconfig b/Engine/Source/Programs/UnrealCloudDDC/.editorconfig
--- a/Engine/Source/Programs/UnrealCloudDDC/.editorconfig
+++ b/Engine/Source/Programs/UnrealCloudDDC/.editorconfig
@@ -1040,8 +1040,8 @@ dotnet_diagnostic.CA5388.severity = warning
 # CA5389: Do Not Add Archive Item's Path To The Target File System Path
 dotnet_diagnostic.CA5389.severity = warning
 
-# CA5390: Do not hard-code encryption key
-dotnet_diagnostic.CA5390.severity = warning
+# CA5390: Do not hard-code encryption key (disabled: expensive interprocedural dataflow analyzer)
+dotnet_diagnostic.CA5390.severity = none
 
 # CA5391: Use antiforgery tokens in ASP.NET Core MVC controllers
 dotnet_diagnostic.CA5391.severity = warning
@@ -1079,8 +1079,8 @@ dotnet_diagnostic.CA5401.severity = warning
 # CA5402: Use CreateEncryptor with the default IV 
 dotnet_diagnostic.CA5402.severity = warning
 
-# CA5403: Do not hard-code certificate
-dotnet_diagnostic.CA5403.severity = warning
+# CA5403: Do not hard-code certificate (disabled: expensive interprocedural dataflow analyzer)
+dotnet_diagnostic.CA5403.severity = none
 
 # CA5404: Do not disable token validation checks
 dotnet_diagnostic.CA5404.severity = none
diff --git a/Engine/Source/Programs/UnrealGameSync/.editorconfig b/Engine/Source/Programs/UnrealGameSync/.editorconfig
--- a/Engine/Source/Programs/UnrealGameSync/.editorconfig
+++ b/Engine/Source/Programs/UnrealGameSync/.editorconfig
@@ -1029,8 +1029,8 @@ dotnet_diagnostic.CA5388.severity = warning
 # CA5389: Do Not Add Archive Item's Path To The Target File System Path
 dotnet_diagnostic.CA5389.severity = none
 
-# CA5390: Do not hard-code encryption key
-dotnet_diagnostic.CA5390.severity = warning
+# CA5390: Do not hard-code encryption key (disabled: 45s+ interprocedural dataflow cost)
+dotnet_diagnostic.CA5390.severity = none
 
 # CA5391: Use antiforgery tokens in ASP.NET Core MVC controllers
 dotnet_diagnostic.CA5391.severity = warning
@@ -1068,8 +1068,8 @@ dotnet_diagnostic.CA5401.severity = warning
 # CA5402: Use CreateEncryptor with the default IV 
 dotnet_diagnostic.CA5402.severity = warning
 
-# CA5403: Do not hard-code certificate
-dotnet_diagnostic.CA5403.severity = warning
+# CA5403: Do not hard-code certificate (disabled: 46s+ interprocedural dataflow cost)
+dotnet_diagnostic.CA5403.severity = none
 
 # CA5404: Do not disable token validation checks
 dotnet_diagnostic.CA5404.severity = none
diff --git a/Engine/Source/Programs/UnrealToolbox/.editorconfig b/Engine/Source/Programs/UnrealToolbox/.editorconfig
--- a/Engine/Source/Programs/UnrealToolbox/.editorconfig
+++ b/Engine/Source/Programs/UnrealToolbox/.editorconfig
@@ -1054,8 +1054,8 @@ dotnet_diagnostic.CA5388.severity = warning
 # CA5389: Do Not Add Archive Item's Path To The Target File System Path
 dotnet_diagnostic.CA5389.severity = warning
 
-# CA5390: Do not hard-code encryption key
-dotnet_diagnostic.CA5390.severity = warning
+# CA5390: Do not hard-code encryption key (disabled: expensive interprocedural dataflow analyzer)
+dotnet_diagnostic.CA5390.severity = none
 
 # CA5391: Use antiforgery tokens in ASP.NET Core MVC controllers
 dotnet_diagnostic.CA5391.severity = warning
@@ -1093,8 +1093,8 @@ dotnet_diagnostic.CA5401.severity = warning
 # CA5402: Use CreateEncryptor with the default IV 
 dotnet_diagnostic.CA5402.severity = warning
 
-# CA5403: Do not hard-code certificate
-dotnet_diagnostic.CA5403.severity = warning
+# CA5403: Do not hard-code certificate (disabled: expensive interprocedural dataflow analyzer)
+dotnet_diagnostic.CA5403.severity = none
 
 # CA5404: Do not disable token validation checks
 dotnet_diagnostic.CA5404.severity = none
