Add optional cooldown delay between test iterations for thermal management.

When benchmarking on physical devices, thermal throttling can skew results as devices heat up during repeated test runs. This adds a TestIterationCooldownSeconds parameter that allows users to configure a delay between iterations, giving devices time to cool down before the next benchmark run. The delay is tunable based on observed thermal behavior of specific devices. This is a simple, user-configurable solution that can be adjusted as needed, versus the more complex alternative of instrumenting device temperature monitoring before each test.

#Synced-CL 51769453

diff --git a/Engine/Build/Graph/Tasks/Inc/GauntletSettings.xml b/Engine/Build/Graph/Tasks/Inc/GauntletSettings.xml
--- a/Engine/Build/Graph/Tasks/Inc/GauntletSettings.xml
+++ b/Engine/Build/Graph/Tasks/Inc/GauntletSettings.xml
@@ -27,6 +27,8 @@
 	<Option Name="ExtraGauntletEditorTestArgs" DefaultValue="" Description="Additional Editor test arguments passed to Gauntlet command line"/>
 	<Option Name="UseTestDataV1" DefaultValue="false" Description="Publish report as TestData v1 for Horde"/>
 
+	<Option Name="TestIterationCooldownSeconds" DefaultValue="0" Restrict="^[0-9]{1,4}$" Description="Seconds to wait between test iterations (max 9999). Allows device thermal cooldown during benchmarking."/>
+	
 	<Option Name="AdbKeysPath" DefaultValue="" Description="Path to adbkeys location (for Android)"/>
 
 	<Option Name="DisableTestResultsForHorde" DefaultValue="false" Restrict="true|false" Description="Force disabling test results for Horde production"/>
@@ -34,6 +36,7 @@
 	<Property Name="GauntletCommonArgs" Value="-maxduration=$(GauntletTimeOut) -unattended -verbose -ResumeOnCriticalFailure $(ExtraGauntletCommonArgs)" />
 	<Property Name="GauntletCommonArgs" Value="$(GauntletCommonArgs) -WriteTestResultsForHorde" If="$(Horde) and !$(DisableTestResultsForHorde)"/>
 	<Property Name="GauntletCommonArgs" Value="$(GauntletCommonArgs) -UseTestDataV1" If="$(Horde) and !$(DisableTestResultsForHorde) and $(UseTestDataV1)"/>
+	<Property Name="GauntletCommonArgs" Value="$(GauntletCommonArgs) -TestIterationCooldownSeconds=$(TestIterationCooldownSeconds)" If="'$(TestIterationCooldownSeconds)' != '0'"/>
 	<Property Name="GauntletCommonArgs" Value="$(GauntletCommonArgs) -branch=&quot;$(Branch)&quot; -changelist=&quot;$(Change)&quot;" />
 
 	<Property Name="GauntletDeviceArgs" Value="" />
diff --git a/Engine/Source/Programs/AutomationTool/Gauntlet/Framework/Gauntlet.TestExecutor.cs b/Engine/Source/Programs/AutomationTool/Gauntlet/Framework/Gauntlet.TestExecutor.cs
--- a/Engine/Source/Programs/AutomationTool/Gauntlet/Framework/Gauntlet.TestExecutor.cs
+++ b/Engine/Source/Programs/AutomationTool/Gauntlet/Framework/Gauntlet.TestExecutor.cs
@@ -33,6 +33,14 @@ namespace Gauntlet
 		[AutoParam(true)]
 		public bool DeferReports;
 
+		/// <summary>
+		/// Optional delay in seconds to wait between test iterations. Primarily used to allow devices to cool down
+		/// when benchmarking to mitigate thermal throttling. Default is 0 (no additional delay beyond standard crash processing time).
+		/// Users can tune this value based on their specific device thermal characteristics and observed behavior.
+		/// </summary>
+		[AutoParam(0)]
+		public int TestIterationCooldownSeconds;
+
 		/// <summary>
 		/// Maximum number of seconds to wait for test start threads to join when a cancel is signalled
 		/// If this time elapses, the thread is abandoned and expected to close on it's own.
@@ -195,7 +203,16 @@ namespace Gauntlet
 				{
 					// if repeating tests wait a little bit. If there was a crash CR might still be
 					// processing things.
-					Thread.Sleep(10000);
+					const int CrashProcessingDelayTimeMs = 10000;
+
+					// Apply optional device cooldown delay to mitigate thermal throttling between benchmark iterations.
+					// The delay is user-configurable based on device thermal characteristics.
+					// Clamp to valid range [0, 9999] to prevent negative values and overflow.
+					const int MaxCooldownSeconds = 9999;
+					int CooldownSeconds = Math.Clamp(Options.TestIterationCooldownSeconds, 0, MaxCooldownSeconds);
+					int TestIterationCooldownMs = CooldownSeconds * 1000;
+
+					Thread.Sleep(Math.Max(CrashProcessingDelayTimeMs, TestIterationCooldownMs));
 				}
 
 				DateTime StartPassTime = DateTime.Now;
