[Gauntlet] Adds UE_GauntletCleanTempDir environment variable parsing to RunUnreal.
If the value is set to true, Gauntlet's TempDir will be deleted at the end of the command in the same way -Clean would

#Synced-CL 51699517

diff --git a/Engine/Source/Programs/AutomationTool/Gauntlet/Unreal/RunUnreal.cs b/Engine/Source/Programs/AutomationTool/Gauntlet/Unreal/RunUnreal.cs
--- a/Engine/Source/Programs/AutomationTool/Gauntlet/Unreal/RunUnreal.cs
+++ b/Engine/Source/Programs/AutomationTool/Gauntlet/Unreal/RunUnreal.cs
@@ -73,6 +73,11 @@ namespace Gauntlet
 		/// </summary>
 		public virtual string DefaultTestName { get { return "DefaultTest"; } }
 
+		/// <summary>
+		/// Constant for environment variable that can be set to trigger the cleanup of Options.TempDir upon completion
+		/// </summary>
+		private const string CleanEnvironmentVariable = "UE_GauntletCleanTempDir";
+
 		/// <summary>
 		/// Main UAT entrance point. Custom games can derive from RunUnrealTests to run custom setup steps or
 		/// directly set params on ContextOptions to remove the need for certain command line params (e.g.
@@ -370,14 +375,12 @@ namespace Gauntlet
 
 				DevicePool.Instance.Dispose();
 
-				if (ParseParam("clean"))
+				bool bCleanFromArgument = ParseParam("Clean");
+				bool bCleanFromEnvironmentVariable = bool.TryParse(Environment.GetEnvironmentVariable(CleanEnvironmentVariable), out bool bCleanEnvvar) && bCleanEnvvar;
+				if (bCleanFromArgument || bCleanFromEnvironmentVariable)
 				{
-					Logger.LogInformation("Deleting temp dir {Arg0}", Options.TempDir);
-					DirectoryInfo Di = new DirectoryInfo(Options.TempDir);
-					if (Di.Exists)
-					{
-						Di.Delete(true);
-					}
+					Logger.LogInformation("Cleaning TempDir '{GauntletTempDir}' by request...", Options.TempDir);
+					DeleteDirectory(Options.TempDir);
 				}
 
 				GC.Collect();
