[Gauntlet] Correct logic error in TargetDeviceMac installations when using the modern launch session flow

#Synced-CL 51699517

diff --git a/Engine/Source/Programs/AutomationTool/Gauntlet/Platform/Mac/Gauntlet.TargetDeviceMac.cs b/Engine/Source/Programs/AutomationTool/Gauntlet/Platform/Mac/Gauntlet.TargetDeviceMac.cs
--- a/Engine/Source/Programs/AutomationTool/Gauntlet/Platform/Mac/Gauntlet.TargetDeviceMac.cs
+++ b/Engine/Source/Programs/AutomationTool/Gauntlet/Platform/Mac/Gauntlet.TargetDeviceMac.cs
@@ -257,7 +257,7 @@ namespace Gauntlet
 
 			string BuildVolume = GetVolumeName(TargetDir);
 			string LocalRoot = GetVolumeName(Environment.CurrentDirectory);
-			if (!string.IsNullOrEmpty(BuildVolume) && BuildVolume.Equals(LocalRoot, StringComparison.OrdinalIgnoreCase))
+			if (!string.IsNullOrEmpty(BuildVolume) && !BuildVolume.Equals(LocalRoot, StringComparison.OrdinalIgnoreCase))
 			{
 				string SubDir = string.IsNullOrEmpty(Sandbox) ? Project : Sandbox;
 				string InstallDir = Path.Combine(InstallRoot, SubDir, ProcessType);
