Refactoring Requeue Matchmaking Test
#tests Tested locally and currently running Horde: 69b2e1925566cb0ff71d6368
#rnx

#Synced-CL 51577459

#Synced-CL 51577459

diff --git a/Engine/Source/Programs/AutomationTool/Gauntlet/Framework/RpcFramework/Gauntlet.RpcTargetRegistry.cs b/Engine/Source/Programs/AutomationTool/Gauntlet/Framework/RpcFramework/Gauntlet.RpcTargetRegistry.cs
--- a/Engine/Source/Programs/AutomationTool/Gauntlet/Framework/RpcFramework/Gauntlet.RpcTargetRegistry.cs
+++ b/Engine/Source/Programs/AutomationTool/Gauntlet/Framework/RpcFramework/Gauntlet.RpcTargetRegistry.cs
@@ -316,6 +316,16 @@ namespace Gauntlet
 			return RpcTargets.Values.Where(t => t.TargetType == InTargetType).ToList();
 		}
 
+		public List<RpcTarget> GetAllClients()
+		{
+			return GetAllTargetsOfType(UnrealTargetRole.Client);
+		}
+
+		public List<RpcTarget> GetAllServers()
+		{
+			return GetAllTargetsOfType(UnrealTargetRole.Server);
+		}
+
 		public virtual void UpdateTargetIp(string MessageBody, ResponseData CallResponse)
 		{
 			ConsoleColor PrevForegroundColor = Console.ForegroundColor;
