Unify gdk package generation requests, allowing easier extensibility
- also includes experimental support for -prestaged with gdk -package

#rnx
#Synced-CL 51451766

diff --git a/Engine/Source/Programs/AutomationTool/Microsoft/GDKPackageGenerator.cs b/Engine/Source/Programs/AutomationTool/Microsoft/GDKPackageGenerator.cs
--- a/Engine/Source/Programs/AutomationTool/Microsoft/GDKPackageGenerator.cs
+++ b/Engine/Source/Programs/AutomationTool/Microsoft/GDKPackageGenerator.cs
@@ -196,35 +196,22 @@ namespace AutomationUtils.GDK
 
 		public void GeneratePackages(ProjectParams Params, DeploymentContext SC, DirectoryReference PackageMetadataBasePath = null, string PackageOutputPathFragment = "")
 		{
-			// prepare the list of configurations to package
-			List<UnrealTargetConfiguration?> ConfigurationsToPackage = new();
-			if (GenerateSinglePackages || Params.HasCreateReleaseVersion)
-			{
-				ConfigurationsToPackage.AddRange(SC.StageTargetConfigurations.Select(X => (UnrealTargetConfiguration?)X));
-			}
-			if (GenerateCombinedPackage && !Params.HasCreateReleaseVersion)
-			{
-				ConfigurationsToPackage.Add(null); // null means "no specific single config - package all configurations"
-			}
-			if (ConfigurationsToPackage.Count == 0)
-			{
-				return;
-			}
+			List<PackageRequest> PackageRequests = MakePackageRequests(Params, SC);
 
 			// generate the packages
-			if (ParallelizePackageGeneration && ConfigurationsToPackage.Count > 1)
+			if (ParallelizePackageGeneration && PackageRequests.Count > 1 && !Params.PreStaged)
 			{
-				Logger.LogInformation("Generating {Num} packages in parallel...", ConfigurationsToPackage.Count);
+				Logger.LogInformation("Generating {Num} packages in parallel...", PackageRequests.Count);
 
 				var LockObject = new object();
-				Parallel.ForEach(ConfigurationsToPackage, (UnrealTargetConfiguration? Config) =>
+				Parallel.ForEach(PackageRequests, (PackageRequest Request) =>
 				{
 					CaptureLogger CaptureLogger = new();
 					try
 					{
-						GenerateSinglePackage(Params, SC, PackageMetadataBasePath, Config, CaptureLogger, PackageOutputPathFragment);
+						GenerateSinglePackage(Params, SC, PackageMetadataBasePath, Request.Config, CaptureLogger, PackageOutputPathFragment);
 					}
-					catch(AutomationException e)
+					catch (AutomationException e)
 					{
 						CaptureLogger.LogError("{Ex}", e.Message);
 						throw new AutomationException("Packaging failed");
@@ -232,7 +219,7 @@ namespace AutomationUtils.GDK
 					finally
 					{
 						// lock ensures each package's output is not interleaved
-						lock (LockObject) 
+						lock (LockObject)
 						{
 							CaptureLogger.RenderTo(Logger);
 						}
@@ -241,17 +228,54 @@ namespace AutomationUtils.GDK
 			}
 			else
 			{
-				foreach (UnrealTargetConfiguration? Config in ConfigurationsToPackage)
+				foreach (PackageRequest Request in PackageRequests)
 				{
-					GenerateSinglePackage(Params, SC, PackageMetadataBasePath, Config, null, PackageOutputPathFragment);
+					GenerateSinglePackage(Params, SC, PackageMetadataBasePath, Request.Config, null, PackageOutputPathFragment);
 				}
 			}
 		}
 
+		public void GeneratePackageMetadata(ProjectParams Params, DeploymentContext SC, DirectoryReference PackageMetadataBasePath)
+		{
+			List<PackageRequest> PackageRequests = MakePackageRequests(Params, SC);
+			foreach (PackageRequest Request in PackageRequests)
+			{
+				GenerateSinglePackageMetadata(Params, SC, PackageMetadataBasePath, Request.Config);
+			}
+		}
 
 
 
+		struct PackageRequest
+		{
+			public PackageRequest()
+			{
+			}
 
+			public PackageRequest(UnrealTargetConfiguration? InConfig)
+			{
+				Config = InConfig;
+			}
+
+			public readonly UnrealTargetConfiguration? Config = null; // null means "no specific single config - package all configurations"
+		};
+
+		private List<PackageRequest> MakePackageRequests(ProjectParams Params, DeploymentContext SC)
+		{
+			List<PackageRequest> PackageRequests = [];
+
+			// prepare the list of configurations to package
+			if (GenerateSinglePackages || Params.HasCreateReleaseVersion)
+			{
+				PackageRequests.AddRange(SC.StageTargetConfigurations.Select(X => new PackageRequest(X)));
+			}
+			if (GenerateCombinedPackage && !Params.HasCreateReleaseVersion)
+			{
+				PackageRequests.Add(new PackageRequest());
+			}
+
+			return PackageRequests;
+		}
 
 		public void GenerateSinglePackage(ProjectParams Params, DeploymentContext SC, DirectoryReference PackageMetadataBasePath = null, UnrealTargetConfiguration? OnlyConfig = null, ILogger CustomLogger = null, string PackageOutputPathFragment = "")
 		{
@@ -270,7 +294,7 @@ namespace AutomationUtils.GDK
 			}
 
 			// fall back to legacy packaging if there is no metadata
-			if (!PrecachePackageLayout)
+			if (!PrecachePackageLayout || Params.PreStaged)
 			{
 				PackageImmediate(Params, SC, PackageOutputDir, CustomLogger, OnlyConfig);
 				return;
@@ -1384,8 +1408,11 @@ namespace AutomationUtils.GDK
 		/// <summary>
 		/// Generates the layout.xml and associated files required for packaging in the given layout directory
 		/// </summary>
-		public void GeneratePackageMetadata(ProjectParams Params, DeploymentContext SC, DirectoryReference PackageMetadataDir, UnrealTargetConfiguration? OnlyConfig = null)
+		public void GenerateSinglePackageMetadata(ProjectParams Params, DeploymentContext SC, DirectoryReference PackageMetadataDir, UnrealTargetConfiguration? OnlyConfig = null)
 		{
+			string PackageConfig = (OnlyConfig?.ToString() ?? "Combined");
+			PackageMetadataDir = DirectoryReference.Combine(PackageMetadataDir, PackageConfig);
+
 			FileUtils.ForceDeleteDirectory(PackageMetadataDir);
 			FileUtils.CreateDirectoryTree(PackageMetadataDir);
 
diff --git a/Engine/Source/Programs/AutomationTool/Win/MSGaming/MSGameStoreCustomDeploymentHandler.cs b/Engine/Source/Programs/AutomationTool/Win/MSGaming/MSGameStoreCustomDeploymentHandler.cs
--- a/Engine/Source/Programs/AutomationTool/Win/MSGaming/MSGameStoreCustomDeploymentHandler.cs
+++ b/Engine/Source/Programs/AutomationTool/Win/MSGaming/MSGameStoreCustomDeploymentHandler.cs
@@ -85,24 +85,10 @@ public class MSGameStoreCustomDeploymentHandler : CustomDeploymentHandler
 
 		// get the package generator for the platform and make the packaging data
 		MSGameStorePackageGenerator PackageGenerator = new(SC.RawProjectPath, SC.StageTargetPlatform.PlatformType, IniSection_TargetSettings, SC.CustomConfig, Platform.CrashReportPlatform);
-		if (PackageGenerator.PrecachePackageLayout)
+		if (PackageGenerator.PrecachePackageLayout && !Params.PreStaged)
 		{
-			// create data that is used for generating per-configuration packages
-			if (PackageGenerator.GenerateSinglePackages || Params.HasCreateReleaseVersion)
-			{
-				foreach (var Config in SC.StageTargetConfigurations)
-				{
-					DirectoryReference PackageMetadataPath = DirectoryReference.Combine(SC.StageDirectory, ".pkgsrc", Config.ToString());
-					PackageGenerator.GeneratePackageMetadata(Params, SC, PackageMetadataPath, Config);
-				}
-			}
-
-			// create data that is used for generating combined configuration packages
-			if (PackageGenerator.GenerateCombinedPackage && !Params.HasCreateReleaseVersion)
-			{
-				DirectoryReference PackageMetadataPath = DirectoryReference.Combine(SC.StageDirectory, ".pkgsrc", "Combined");
-				PackageGenerator.GeneratePackageMetadata(Params, SC, PackageMetadataPath);
-			}
+			DirectoryReference PackageMetadataPath = DirectoryReference.Combine(SC.StageDirectory, ".pkgsrc");
+			PackageGenerator.GeneratePackageMetadata(Params, SC, PackageMetadataPath);
 		}
 
 		return false;
@@ -279,7 +265,7 @@ public class MSGameStoreCustomDeploymentHandler : CustomDeploymentHandler
 	{
 		// copy all custom install actions to the correct folder if we are not using package precaching
 		MSGameStorePackageGenerator PackageGenerator = new(SC.RawProjectPath, SC.StageTargetPlatform.PlatformType, IniSection_TargetSettings, SC.CustomConfig, SC.StageTargetPlatform.CrashReportPlatform);
-		if (!PackageGenerator.PrecachePackageLayout)
+		if (!PackageGenerator.PrecachePackageLayout || Params.PreStaged)
 		{
 			ConfigHierarchy EngineIni = ConfigCache.ReadHierarchy(ConfigHierarchyType.Engine, DirectoryReference.FromFile(SC.RawProjectPath), SC.StageTargetPlatform.PlatformType, SC.CustomConfig);
 
diff --git a/Engine/Source/Programs/AutomationTool/Win/MSGaming/MSGameStoreCustomStagingHandler.cs b/Engine/Source/Programs/AutomationTool/Win/MSGaming/MSGameStoreCustomStagingHandler.cs
--- a/Engine/Source/Programs/AutomationTool/Win/MSGaming/MSGameStoreCustomStagingHandler.cs
+++ b/Engine/Source/Programs/AutomationTool/Win/MSGaming/MSGameStoreCustomStagingHandler.cs
@@ -81,7 +81,7 @@ public class MSGameStoreCustomStagingHandler : CustomStagingHandler
 		{
 			// copy all custom install actions to the correct folder if we are not using package preaching
 			MSGameStorePackageGenerator PackageGenerator = new(SC.RawProjectPath, SC.StageTargetPlatform.PlatformType, IniSection_TargetSettings, SC.CustomConfig, SC.StageTargetPlatform.CrashReportPlatform);
-			if (!PackageGenerator.PrecachePackageLayout)
+			if (!PackageGenerator.PrecachePackageLayout || Params.PreStaged)
 			{
 				ConfigHierarchy EngineIni = ConfigCache.ReadHierarchy(ConfigHierarchyType.Engine, DirectoryReference.FromFile(SC.RawProjectPath), SC.StageTargetPlatform.PlatformType, SC.CustomConfig);
 
@@ -109,27 +109,12 @@ public class MSGameStoreCustomStagingHandler : CustomStagingHandler
 		{
 			// get the package generator for the platform and make the packaging data
 			MSGameStorePackageGenerator PackageGenerator = new(SC.RawProjectPath, SC.StageTargetPlatform.PlatformType, IniSection_TargetSettings, SC.CustomConfig, SC.StageTargetPlatform.CrashReportPlatform);
-			if (PackageGenerator.PrecachePackageLayout)
+			if (PackageGenerator.PrecachePackageLayout && !Params.PreStaged)
 			{
 				DirectoryReference PackageMetadataBasePath = GetPackageMetadataBasePath(SC);
 				CommandUtils.DeleteDirectory_NoExceptions(PackageMetadataBasePath.FullName);
 
-				// create data that is used for generating per-configuration packages
-				if (PackageGenerator.GenerateSinglePackages || Params.HasCreateReleaseVersion)
-				{
-					foreach (var Config in SC.StageTargetConfigurations)
-					{
-						DirectoryReference PackageMetadataPath = DirectoryReference.Combine(GetPackageMetadataBasePath(SC), Config.ToString());
-						PackageGenerator.GeneratePackageMetadata(Params, SC, PackageMetadataPath, Config);
-					}
-				}
-
-				// create data that is used for generating combined configuration packages
-				if (PackageGenerator.GenerateCombinedPackage && !Params.HasCreateReleaseVersion)
-				{
-					DirectoryReference PackageMetadataPath = DirectoryReference.Combine(GetPackageMetadataBasePath(SC), "Combined");
-					PackageGenerator.GeneratePackageMetadata(Params, SC, PackageMetadataPath);
-				}
+				PackageGenerator.GeneratePackageMetadata(Params, SC, PackageMetadataBasePath);
 			}
 		}
 	}
