[Gauntlet] UnrealSession will now use the modern launch session flow by default
#Synced-CL 51699517

diff --git a/Engine/Source/Programs/AutomationTool/Gauntlet/Unreal/Base/Gauntlet.UnrealSession.cs b/Engine/Source/Programs/AutomationTool/Gauntlet/Unreal/Base/Gauntlet.UnrealSession.cs
--- a/Engine/Source/Programs/AutomationTool/Gauntlet/Unreal/Base/Gauntlet.UnrealSession.cs
+++ b/Engine/Source/Programs/AutomationTool/Gauntlet/Unreal/Base/Gauntlet.UnrealSession.cs
@@ -896,7 +896,7 @@ namespace Gauntlet
 		/// <returns></returns>
 		public UnrealSessionInstance LaunchSession()
 		{
-			if(!Globals.UseExperimentalFeatures)
+			if(Globals.Params.ParseParam("LegacyLaunchSession"))
 			{
 				return Legacy_LaunchSession();
 			}
