[Android] Enable AFS for Gauntlet shipping builds to fix scoped storage file access

On Android 11+, files pushed via adb are owned by the shell user and
the FUSE scoped storage layer prevents the app from reading them,
even in its own data directory. This breaks PGO replay playback.

The fix parses the AFS token from the install batch script during
build discovery in CreateFromPath(), and threads it through the
existing UseAFSWithoutToken plumbing so UsingAndroidFileServer()
returns true with a valid token. This activates the full AFS flow
already implemented in Gauntlet: install AFS APK, push files via
AFS (owned by app UID), then replace with game APK before launch.

#rnx
#Synced-CL 51753475

diff --git a/Engine/Source/Programs/AutomationTool/Gauntlet/Platform/Android/Gauntlet.AndroidBuildSource.cs b/Engine/Source/Programs/AutomationTool/Gauntlet/Platform/Android/Gauntlet.AndroidBuildSource.cs
--- a/Engine/Source/Programs/AutomationTool/Gauntlet/Platform/Android/Gauntlet.AndroidBuildSource.cs
+++ b/Engine/Source/Programs/AutomationTool/Gauntlet/Platform/Android/Gauntlet.AndroidBuildSource.cs
@@ -38,9 +38,11 @@ namespace Gauntlet
 
 		public bool UseAFSWithoutToken { get; protected set; }
 
+		public string ParsedAFSToken { get; protected set; }
+
 		public bool UsesPublicLogs { get; protected set; }
 
-		public AndroidBuild(UnrealTargetConfiguration InConfig, string InAndroidPackageName, string InApkPath, Dictionary<string, string> InFilesToInstall, BuildFlags InFlags, bool InIs32Bit, bool bInUsesExternalFilesDir, bool bInUsesPublicLogs, bool bInUseAFSWithoutToken = false)
+		public AndroidBuild(UnrealTargetConfiguration InConfig, string InAndroidPackageName, string InApkPath, Dictionary<string, string> InFilesToInstall, BuildFlags InFlags, bool InIs32Bit, bool bInUsesExternalFilesDir, bool bInUsesPublicLogs, bool bInUseAFSWithoutToken = false, string InParsedAFSToken = null)
 		{
 			Configuration = InConfig;
 			AndroidPackageName = InAndroidPackageName;
@@ -51,6 +53,7 @@ namespace Gauntlet
 			UsesExternalFilesDir = bInUsesExternalFilesDir;
 			UsesPublicLogs = bInUsesPublicLogs;
 			UseAFSWithoutToken = bInUseAFSWithoutToken;
+			ParsedAFSToken = InParsedAFSToken;
 		}
 
 		public bool CanSupportRole(UnrealTargetRole RoleType)
@@ -169,7 +172,21 @@ namespace Gauntlet
 					Flags |= BuildFlags.NotBulk;
 				}
 
-				AndroidBuild NewBuild = new AndroidBuild(UnrealConfig, AndroidPackageName, SourceApkPath, FilesToInstall, Flags, PackageIs32Bit, bUsesExternalFilesDir, bUsesPublicLogs);
+				// Detect AFS token from install script for shipping builds with separate AFS APK.
+				bool bUseAFS = false;
+				string ParsedAFSToken = null;
+				if (Path.GetFileName(SourceApkPath).StartsWith("AFS_"))
+				{
+					var AFSTokenMatch = Regex.Match(BatContents, @"-k\s+([A-Fa-f0-9]{32})");
+					if (AFSTokenMatch.Success)
+					{
+						ParsedAFSToken = AFSTokenMatch.Groups[1].ToString();
+						bUseAFS = true;
+						Log.Info("Detected AFS build with token at {0}", AbsPath);
+					}
+				}
+
+				AndroidBuild NewBuild = new AndroidBuild(UnrealConfig, AndroidPackageName, SourceApkPath, FilesToInstall, Flags, PackageIs32Bit, bUsesExternalFilesDir, bUsesPublicLogs, bUseAFS, ParsedAFSToken);
 
 				DiscoveredBuilds.Add(NewBuild);
 
diff --git a/Engine/Source/Programs/AutomationTool/Gauntlet/Platform/Android/Gauntlet.TargetDeviceAndroid.cs b/Engine/Source/Programs/AutomationTool/Gauntlet/Platform/Android/Gauntlet.TargetDeviceAndroid.cs
--- a/Engine/Source/Programs/AutomationTool/Gauntlet/Platform/Android/Gauntlet.TargetDeviceAndroid.cs
+++ b/Engine/Source/Programs/AutomationTool/Gauntlet/Platform/Android/Gauntlet.TargetDeviceAndroid.cs
@@ -121,6 +121,8 @@ namespace Gauntlet
 
 		private bool bUseAFSWithoutToken;
 
+		private string ParsedAFSToken;
+
 		private Thread ScreenshotThread = null;
 
 		private AutoResetEvent ScreenshotTrigger = null;
@@ -383,6 +385,7 @@ namespace Gauntlet
 			}
 
 			bUseAFSWithoutToken = Build.UseAFSWithoutToken;
+			ParsedAFSToken = Build.ParsedAFSToken;
 
 			Log.Info("Cleaning up artifacts for {PackageName}...", Build.AndroidPackageName);
 
@@ -390,7 +393,7 @@ namespace Gauntlet
 
 			bool bIsPackageInstalled = PackageOutput.Contains($"package:{Build.AndroidPackageName}");
 
-			if (bIsPackageInstalled && UsingAndroidFileServer(AppConfig.ProjectFile, AppConfig.Configuration, out _, out string AFSToken, out _, out _, out _, bUseAFSWithoutToken))
+			if (bIsPackageInstalled && UsingAndroidFileServer(AppConfig.ProjectFile, AppConfig.Configuration, out _, out string AFSToken, out _, out _, out _, bUseAFSWithoutToken, ParsedAFSToken))
 			{
 				RunAFSDeviceCommand(string.Format("-p \"{0}\" -k \"{1}\" rmdir -r \"^saved/\"", Build.AndroidPackageName, AFSToken));
 			}
@@ -416,6 +419,7 @@ namespace Gauntlet
 			}
 
 			bUseAFSWithoutToken = Build.UseAFSWithoutToken;
+			ParsedAFSToken = Build.ParsedAFSToken;
 
 			string Package = Build.AndroidPackageName;
 			KillRunningProcess(Package);
@@ -444,6 +448,7 @@ namespace Gauntlet
 			}
 
 			bUseAFSWithoutToken = Build.UseAFSWithoutToken;
+			ParsedAFSToken = Build.ParsedAFSToken;
 
 			string ExternalStoragePath = StoragePath + "/UnrealGame/" + AppConfig.ProjectName;
 			string ExternalFilesPath = StoragePath + "/Android/data/" + Build.AndroidPackageName + "/files/UnrealGame/" + AppConfig.ProjectName;
@@ -502,6 +507,7 @@ namespace Gauntlet
 			// Now, create and push the UECommandline.txt file
 			bool bUseExternalFileDirectory = (Install.AppConfig.Build as AndroidBuild)?.UsesExternalFilesDir ?? false;
 			bUseAFSWithoutToken = (Install.AppConfig.Build as AndroidBuild)?.UseAFSWithoutToken ?? false;
+			ParsedAFSToken = (Install.AppConfig.Build as AndroidBuild)?.ParsedAFSToken;
 			CopyCommandlineFile(Install.AppConfig, Install.AppConfig.CommandLine, Install.AndroidPackageName, Install.AppConfig.ProjectName, bUseExternalFileDirectory);
 
 			// kill any currently running instance:
@@ -516,7 +522,7 @@ namespace Gauntlet
 			RunAdbDeviceCommand("logcat --clear");
 
 			// Ensure artifact directories exist
-			if (UsingAndroidFileServer(Install.ProjectFile, Install.Configuration, out _, out string AFSToken, out _, out _, out _, bUseAFSWithoutToken))
+			if (UsingAndroidFileServer(Install.ProjectFile, Install.Configuration, out _, out string AFSToken, out _, out _, out _, bUseAFSWithoutToken, ParsedAFSToken))
 			{
 				RunAFSDeviceCommand(string.Format("-p \"{0}\" -k \"{1}\" mkdir \"^saved\"", Install.AndroidPackageName, AFSToken));
 			}
@@ -635,12 +641,12 @@ namespace Gauntlet
 			}
 		}
 
-		public static bool UsingAndroidFileServer(FileReference RawProjectPath, UnrealTargetConfiguration TargetConfiguration, out bool bEnablePlugin, out string AFSToken, out bool bIsShipping, out bool bIncludeInShipping, out bool bAllowExternalStartInShipping, bool bUseAFSWithoutToken = false)
+		public static bool UsingAndroidFileServer(FileReference RawProjectPath, UnrealTargetConfiguration TargetConfiguration, out bool bEnablePlugin, out string AFSToken, out bool bIsShipping, out bool bIncludeInShipping, out bool bAllowExternalStartInShipping, bool bUseAFSWithoutToken = false, string ParsedAFSToken = null)
 		{
 			if (bUseAFSWithoutToken)
 			{
 				bEnablePlugin = true;
-				AFSToken = "";
+				AFSToken = ParsedAFSToken ?? "";
 				bIsShipping = false;
 				bIncludeInShipping = false;
 				bAllowExternalStartInShipping = false;
@@ -693,8 +699,9 @@ namespace Gauntlet
 			}
 
 			bUseAFSWithoutToken = Build.UseAFSWithoutToken;
+			ParsedAFSToken = Build.ParsedAFSToken;
 
-			return UsingAndroidFileServer(AppConfig.ProjectFile, Build.Configuration, out _, out _, out _, out _, out _, bUseAFSWithoutToken);
+			return UsingAndroidFileServer(AppConfig.ProjectFile, Build.Configuration, out _, out _, out _, out _, out _, bUseAFSWithoutToken, ParsedAFSToken);
 		}
 
 		/// <summary>
@@ -938,7 +945,7 @@ namespace Gauntlet
 			FileReference ProjectFile = null, UnrealTargetConfiguration Configuration = UnrealTargetConfiguration.Unknown, bool IgnoreDependencies = false)
 		{
 			string AFSToken;
-			bool UseAFS = UsingAndroidFileServer(ProjectFile, Configuration, out _, out AFSToken, out _, out _, out _, bUseAFSWithoutToken);
+			bool UseAFS = UsingAndroidFileServer(ProjectFile, Configuration, out _, out AFSToken, out _, out _, out _, bUseAFSWithoutToken, ParsedAFSToken);
 			bool IsAPK = string.Equals(Path.GetExtension(SourcePath), ".apk", StringComparison.OrdinalIgnoreCase);
 
 			// for the APK there's no easy/reliable way to get the date of the version installed, so
@@ -1180,8 +1187,9 @@ namespace Gauntlet
 				}
 
 				bUseAFSWithoutToken = Build.UseAFSWithoutToken;
+				ParsedAFSToken = Build.ParsedAFSToken;
 
-				if (UsingAndroidFileServer(AppConfig.ProjectFile, Build.Configuration, out bAFSEnablePlugin, out AFSToken, out bIsShipping, out bAFSIncludeInShipping, out bAFSAllowExternalStartInShipping, bUseAFSWithoutToken))
+				if (UsingAndroidFileServer(AppConfig.ProjectFile, Build.Configuration, out bAFSEnablePlugin, out AFSToken, out bIsShipping, out bAFSIncludeInShipping, out bAFSAllowExternalStartInShipping, bUseAFSWithoutToken, ParsedAFSToken))
 				{
 					useADB = false;
 				}
@@ -1284,7 +1292,7 @@ namespace Gauntlet
 			FileReference ProjectFile = null, UnrealTargetConfiguration Configuration = UnrealTargetConfiguration.Unknown)
 		{
 			string AFSToken;
-			bool UseAFS = UsingAndroidFileServer(ProjectFile, Configuration, out _, out AFSToken, out _, out _, out _, bUseAFSWithoutToken);
+			bool UseAFS = UsingAndroidFileServer(ProjectFile, Configuration, out _, out AFSToken, out _, out _, out _, bUseAFSWithoutToken, ParsedAFSToken);
 
 			IProcessResult AdbResult;
 			if (UseAFS)
@@ -1414,9 +1422,10 @@ namespace Gauntlet
 			bool useADB = true;
 
 			bUseAFSWithoutToken = Build.UseAFSWithoutToken;
+			ParsedAFSToken = Build.ParsedAFSToken;
 			if (AppConfig != null)
 			{
-				if (UsingAndroidFileServer(AppConfig.ProjectFile, Build.Configuration, out bAFSEnablePlugin, out AFSToken, out bIsShipping, out bAFSIncludeInShipping, out bAFSAllowExternalStartInShipping, bUseAFSWithoutToken))
+				if (UsingAndroidFileServer(AppConfig.ProjectFile, Build.Configuration, out bAFSEnablePlugin, out AFSToken, out bIsShipping, out bAFSIncludeInShipping, out bAFSAllowExternalStartInShipping, bUseAFSWithoutToken, ParsedAFSToken))
 				{
 					useADB = false;
 				}
@@ -1842,8 +1851,9 @@ namespace Gauntlet
 		public void StartRecording()
 		{
 			bool bAlwaysUseAFS = (Install.AppConfig.Build as AndroidBuild)?.UseAFSWithoutToken ?? false;
+			string RecordingParsedAFSToken = (Install.AppConfig.Build as AndroidBuild)?.ParsedAFSToken;
 			// Ensure artifact directories exist
-			if (TargetDeviceAndroid.UsingAndroidFileServer(Install.ProjectFile, Install.Configuration, out _, out string AFSToken, out _, out _, out _, bAlwaysUseAFS))
+			if (TargetDeviceAndroid.UsingAndroidFileServer(Install.ProjectFile, Install.Configuration, out _, out string AFSToken, out _, out _, out _, bAlwaysUseAFS, RecordingParsedAFSToken))
 			{
 				Install.AndroidDevice.RunAFSDeviceCommand(string.Format("-p \"{0}\" -k \"{1}\" mkdir \"^saved/Logs\"", Install.AndroidPackageName, AFSToken));
 			}
