Replacing ios-deploy artifact cleaning with devicectl blank directory replacement.

#Synced-CL 51642430

diff --git a/Engine/Source/Programs/AutomationTool/Gauntlet/Platform/IOS/Gauntlet.TargetDeviceIOS.cs b/Engine/Source/Programs/AutomationTool/Gauntlet/Platform/IOS/Gauntlet.TargetDeviceIOS.cs
--- a/Engine/Source/Programs/AutomationTool/Gauntlet/Platform/IOS/Gauntlet.TargetDeviceIOS.cs
+++ b/Engine/Source/Programs/AutomationTool/Gauntlet/Platform/IOS/Gauntlet.TargetDeviceIOS.cs
@@ -532,6 +532,124 @@ namespace Gauntlet
 
 		public void CleanArtifacts(string BundleID, string ProjectName)
 		{
+			if (UseDeviceCtl)
+			{
+				// devicectl doesn't support direct file deletion, so we use an empty directory replacement approach
+				// Create empty local directories and copy them to the device to clear the Saved folders
+
+				string DevicectlEngineSaved = "Documents/Engine/Saved";
+				string DevicectlProjectSaved = string.Format("Documents/{0}/Saved", ProjectName);
+
+				// Create temporary empty directories
+				string TempCleanupPath = Path.Combine(LocalCachePath, "CleanupTemp");
+				string EmptyEngineSaved = Path.Combine(TempCleanupPath, "Engine", "Saved");
+				string EmptyProjectSaved = Path.Combine(TempCleanupPath, ProjectName, "Saved");
+
+				try
+				{
+					// Clean up any existing temp cleanup directory
+					if (Directory.Exists(TempCleanupPath))
+					{
+						Directory.Delete(TempCleanupPath, true);
+					}
+
+					// Create empty directory structures with marker files
+					// devicectl doesn't handle truly empty directories well, so we add a small marker file
+					Directory.CreateDirectory(EmptyEngineSaved);
+					Directory.CreateDirectory(EmptyProjectSaved);
+					File.WriteAllText(Path.Combine(EmptyEngineSaved, ".cleanup"), "cleaned");
+					File.WriteAllText(Path.Combine(EmptyProjectSaved, ".cleanup"), "cleaned");
+
+					// Copy empty Engine/Saved directory to device
+					Log.Info("Cleaning Engine artifacts for {0}", BundleID);
+					string CleanEngineCommand = string.Format("device copy to --device {0} --source {1} --destination {2} --domain-type appDataContainer --domain-identifier {3}",
+						UUID, EmptyEngineSaved, DevicectlEngineSaved, BundleID);
+					IProcessResult EngineResult = ExecuteDevicectlCommand(CleanEngineCommand, UseDeviceID: false, WaitTime: 120);
+					EngineResult.WaitForExit();
+
+					if (EngineResult.ExitCode != 0)
+					{
+						// ContainerLookupErrorDomain error -1 means the app container doesn't exist (app not installed or was uninstalled)
+						if (EngineResult.Output.Contains("ContainerLookupErrorDomain error -1"))
+						{
+							Log.Verbose("App container for {0} not found - no artifacts to clean", BundleID);
+							return;
+						}
+						Log.Warning("Cleaning Engine artifacts for {0} exited with code {1}:\n{2}", BundleID, EngineResult.ExitCode, EngineResult.Output);
+					}
+
+					// Copy empty Project/Saved directory to device
+					Log.Info("Cleaning Project artifacts for {0}", BundleID);
+					string CleanProjectCommand = string.Format("device copy to --device {0} --source {1} --destination {2} --domain-type appDataContainer --domain-identifier {3}",
+						UUID, EmptyProjectSaved, DevicectlProjectSaved, BundleID);
+					IProcessResult ProjectResult = ExecuteDevicectlCommand(CleanProjectCommand, UseDeviceID: false, WaitTime: 120);
+					ProjectResult.WaitForExit();
+
+					if (ProjectResult.ExitCode != 0)
+					{
+						// ContainerLookupErrorDomain error -1 means the app container doesn't exist (app not installed or was uninstalled)
+						if (ProjectResult.Output.Contains("ContainerLookupErrorDomain error -1"))
+						{
+							Log.Verbose("App container for {0} not found - no artifacts to clean", BundleID);
+							return;
+						}
+						Log.Warning("Cleaning Project artifacts for {0} exited with code {1}:\n{2}", BundleID, ProjectResult.ExitCode, ProjectResult.Output);
+					}
+
+					// Verify cleanup by checking directory contents
+					Log.Verbose("Verifying artifact cleanup for {0}", BundleID);
+					string VerifyPath = Path.Combine(TempCleanupPath, "Verify");
+					if (Directory.Exists(VerifyPath))
+					{
+						Directory.Delete(VerifyPath, true);
+					}
+					Directory.CreateDirectory(VerifyPath);
+
+					// Download Engine/Saved to verify
+					string VerifyEngineCommand = string.Format("device copy from --device {0} --source {1} --destination {2} --domain-type appDataContainer --domain-identifier {3}",
+						UUID, DevicectlEngineSaved, VerifyPath, BundleID);
+					IProcessResult VerifyEngineResult = ExecuteDevicectlCommand(VerifyEngineCommand, UseDeviceID: false, WaitTime: 60, bQuiet: true);
+					VerifyEngineResult.WaitForExit();
+
+					if (VerifyEngineResult.ExitCode == 0 && Directory.Exists(VerifyPath))
+					{
+						// Count all files recursively from the verify path
+						string[] engineFiles = Directory.GetFiles(VerifyPath, "*", SearchOption.AllDirectories);
+						Log.Info("Engine/Saved contains {0} file(s) after cleanup", engineFiles.Length);
+						if (engineFiles.Length > 0 && Log.IsVeryVerbose)
+						{
+							foreach (string file in engineFiles)
+							{
+								string relativePath = file.Replace(VerifyPath, "").TrimStart(Path.DirectorySeparatorChar, Path.AltDirectorySeparatorChar);
+								Log.Verbose("  - {0}", relativePath);
+							}
+						}
+					}
+					else
+					{
+						Log.Verbose("Could not verify Engine/Saved cleanup (exit code {0})", VerifyEngineResult.ExitCode);
+					}
+				}
+				finally
+				{
+					// Clean up temporary directories
+					if (Directory.Exists(TempCleanupPath))
+					{
+						try
+						{
+							Directory.Delete(TempCleanupPath, true);
+						}
+						catch (Exception Ex)
+						{
+							Log.Verbose("Failed to delete temp cleanup directory: {0}", Ex.Message);
+						}
+					}
+				}
+
+				return;
+			}
+
+			// Use ios-deploy for cleanup (legacy approach)
 			string EngineSavedDirectory = "/Documents/Engine/Saved";
 			string ProjectSavedDirectory = string.Format("/Documents/{0}/Saved", ProjectName);
 
