[Gauntlet] Overlay executables will now use the project directory as the base location for executable searches rather than the UAT environment root
#Synced-CL 51699517

diff --git a/Engine/Source/Programs/AutomationTool/Gauntlet/Framework/Base/Gauntlet.OverlayExecutable.cs b/Engine/Source/Programs/AutomationTool/Gauntlet/Framework/Base/Gauntlet.OverlayExecutable.cs
--- a/Engine/Source/Programs/AutomationTool/Gauntlet/Framework/Base/Gauntlet.OverlayExecutable.cs
+++ b/Engine/Source/Programs/AutomationTool/Gauntlet/Framework/Base/Gauntlet.OverlayExecutable.cs
@@ -1,6 +1,7 @@
 // Copyright Epic Games, Inc. All Rights Reserved.
 
 using AutomationTool;
+using EpicGames.Core;
 using System;
 using System.IO;
 using UnrealBuildTool;
@@ -19,17 +20,24 @@ namespace Gauntlet
 	{
 		private string ProjectName;
 		private string Flavor;
+		private FileReference ProjectFile;
 		private UnrealTargetPlatform Platform;
 		private UnrealTargetConfiguration Configuration;
 		private UnrealTargetRole Role;
 
-		public OverlayExecutable(UnrealSessionRole SessionRole, string ProjectName)
+		public OverlayExecutable(UnrealSessionRole SessionRole, string ProjectName, FileReference ProjectFile)
+			: this(SessionRole.Platform.Value, SessionRole.Configuration, SessionRole.RoleType, ProjectName, ProjectFile, SessionRole.RequiredFlavor)
+		{ }
+
+		public OverlayExecutable(UnrealTargetPlatform Platform, UnrealTargetConfiguration Configuration, UnrealTargetRole RoleType,
+			string ProjectName, FileReference ProjectFile, string Flavor = "")
 		{
 			this.ProjectName = ProjectName;
-			this.Configuration = SessionRole.Configuration;
-			this.Role = SessionRole.RoleType;
-			this.Platform = SessionRole.Platform.Value;
-			this.Flavor = SessionRole.RequiredFlavor;
+			this.Configuration = Configuration;
+			this.Role = RoleType;
+			this.Platform = Platform;
+			this.Flavor = Flavor;
+			this.ProjectFile = ProjectFile;
 		}
 
 		/// <summary>
@@ -64,7 +72,7 @@ namespace Gauntlet
 			}
 
 			// Ex. /UnrealGame/Binaries/Win64/UnrealClient.exe
-			string PlatformBinariesDirectory = Path.Combine(Environment.CurrentDirectory, ProjectName, "Binaries", Platform.ToString());
+			string PlatformBinariesDirectory = Path.Combine(ProjectFile.Directory.FullName, "Binaries", Platform.ToString());
 			string ExecutableExtension = string.IsNullOrEmpty(ExtensionOverride) ? Path.GetExtension(BaseExecutable) : ExtensionOverride;
 
 			string ExecutableName = UnrealHelpers.GetExecutableName(ModuleName, Platform, Configuration, Role, Flavor, string.Empty);
diff --git a/Engine/Source/Programs/AutomationTool/Gauntlet/Unreal/BuildSource/Gauntlet.UnrealBuildSource.cs b/Engine/Source/Programs/AutomationTool/Gauntlet/Unreal/BuildSource/Gauntlet.UnrealBuildSource.cs
--- a/Engine/Source/Programs/AutomationTool/Gauntlet/Unreal/BuildSource/Gauntlet.UnrealBuildSource.cs
+++ b/Engine/Source/Programs/AutomationTool/Gauntlet/Unreal/BuildSource/Gauntlet.UnrealBuildSource.cs
@@ -703,7 +703,7 @@ namespace Gauntlet
 
 			if(Globals.IsRunningDev)
 			{
-				Config.OverlayExecutable = new OverlayExecutable(Role, Config.ProjectName);
+				Config.OverlayExecutable = new OverlayExecutable(Role, Config.ProjectName, Config.ProjectFile);
 			}
 
 			return Config;
