If not enough devices reserved, show message instead of invalidating step

#rnx
#Synced-CL 0

diff --git a/Engine/Source/Programs/AutomationTool/LowLevelTests/RunLowLevelTests.cs b/Engine/Source/Programs/AutomationTool/LowLevelTests/RunLowLevelTests.cs
--- a/Engine/Source/Programs/AutomationTool/LowLevelTests/RunLowLevelTests.cs
+++ b/Engine/Source/Programs/AutomationTool/LowLevelTests/RunLowLevelTests.cs
@@ -360,7 +360,7 @@ namespace LowLevelTests
 			bool Reserved = UnrealDeviceReservation.TryReserveDevices(RequiredDeviceTypes, 1);
 			if (!Reserved)
 			{
-				Gauntlet.Horde.ReportInvalidatingIssue(Gauntlet.Horde.GauntletHordeStepMessage.NotEnoughDevices, "The test was unable to secure enough devices to run.");
+				Gauntlet.Horde.UpdateStepMessage(Gauntlet.Horde.GauntletHordeStepMessage.NotEnoughDevices);
 			}
 			return Reserved;
 		}
