[Android] Fix AFS token flow naming in Gauntlet build discovery

#rnx
#Synced-CL 51877528

diff --git a/Engine/Source/Programs/AutomationTool/Gauntlet/Platform/Android/Gauntlet.AndroidBuildSource.cs b/Engine/Source/Programs/AutomationTool/Gauntlet/Platform/Android/Gauntlet.AndroidBuildSource.cs
--- a/Engine/Source/Programs/AutomationTool/Gauntlet/Platform/Android/Gauntlet.AndroidBuildSource.cs
+++ b/Engine/Source/Programs/AutomationTool/Gauntlet/Platform/Android/Gauntlet.AndroidBuildSource.cs
@@ -173,7 +173,6 @@ namespace Gauntlet
 				}
 
 				// Detect AFS token from install script for shipping builds with separate AFS APK.
-				bool bUseAFS = false;
 				string ParsedAFSToken = null;
 				if (Path.GetFileName(SourceApkPath).StartsWith("AFS_"))
 				{
@@ -181,12 +180,11 @@ namespace Gauntlet
 					if (AFSTokenMatch.Success)
 					{
 						ParsedAFSToken = AFSTokenMatch.Groups[1].ToString();
-						bUseAFS = true;
 						Log.Info("Detected AFS build with token at {0}", AbsPath);
 					}
 				}
 
-				AndroidBuild NewBuild = new AndroidBuild(UnrealConfig, AndroidPackageName, SourceApkPath, FilesToInstall, Flags, PackageIs32Bit, bUsesExternalFilesDir, bUsesPublicLogs, bUseAFS, ParsedAFSToken);
+				AndroidBuild NewBuild = new AndroidBuild(UnrealConfig, AndroidPackageName, SourceApkPath, FilesToInstall, Flags, PackageIs32Bit, bUsesExternalFilesDir, bUsesPublicLogs, InParsedAFSToken: ParsedAFSToken);
 
 				DiscoveredBuilds.Add(NewBuild);
 
diff --git a/Engine/Source/Programs/AutomationTool/Gauntlet/Platform/Android/Gauntlet.TargetDeviceAndroid.cs b/Engine/Source/Programs/AutomationTool/Gauntlet/Platform/Android/Gauntlet.TargetDeviceAndroid.cs
--- a/Engine/Source/Programs/AutomationTool/Gauntlet/Platform/Android/Gauntlet.TargetDeviceAndroid.cs
+++ b/Engine/Source/Programs/AutomationTool/Gauntlet/Platform/Android/Gauntlet.TargetDeviceAndroid.cs
@@ -646,7 +646,16 @@ namespace Gauntlet
 			if (bUseAFSWithoutToken)
 			{
 				bEnablePlugin = true;
-				AFSToken = ParsedAFSToken ?? "";
+				AFSToken = "";
+				bIsShipping = false;
+				bIncludeInShipping = false;
+				bAllowExternalStartInShipping = false;
+				return true;
+			}
+			if (!string.IsNullOrEmpty(ParsedAFSToken))
+			{
+				bEnablePlugin = true;
+				AFSToken = ParsedAFSToken;
 				bIsShipping = false;
 				bIncludeInShipping = false;
 				bAllowExternalStartInShipping = false;
