FreeType2:
Ensure Engine/Source/ThirdParty/Licenses/NOTICES.txt is staged in the root of the staging directory so that FreeType2 attributions (and future TPL attributions) are distributed with the binaries of the products built with Unreal Engine.

diff --git a/Engine/Source/Programs/AutomationTool/Scripts/CopyBuildToStagingDirectory.Automation.cs b/Engine/Source/Programs/AutomationTool/Scripts/CopyBuildToStagingDirectory.Automation.cs
--- a/Engine/Source/Programs/AutomationTool/Scripts/CopyBuildToStagingDirectory.Automation.cs
+++ b/Engine/Source/Programs/AutomationTool/Scripts/CopyBuildToStagingDirectory.Automation.cs
@@ -1647,6 +1647,19 @@ namespace AutomationScripts
 					SC.StageFile(StagedFileType.SystemNonUFS, CommandLineFile, StagedCommandLineFile);
 				}
 
+				// move Engine/Source/ThirdParty/Licenses/NOTICES.txt file to the root of the stage for necessary third party library attribution
+				// this file needs to be treated as a UFS file for casing, but NonUFS for being put into the .pak file
+				FileReference ThirdPartyNoticesFile = FileReference.Combine(SC.EngineRoot, "Source", "ThirdParty", "Licenses", "NOTICES.txt");
+				if (FileReference.Exists(ThirdPartyNoticesFile))
+				{
+					StagedFileReference StagedThirdPartyNoticesFile = new StagedFileReference("NOTICES.txt");
+					if (SC.StageTargetPlatform.DeployLowerCaseFilenames(StagedFileType.SystemNonUFS))
+					{
+						StagedThirdPartyNoticesFile = StagedThirdPartyNoticesFile.ToLowerInvariant();
+					}
+					SC.StageFile(StagedFileType.SystemNonUFS, ThirdPartyNoticesFile, StagedThirdPartyNoticesFile);
+				}
+
 				DirectoryReference ProjectContentRoot = DirectoryReference.Combine(SC.ProjectRoot, "Content");
 				StagedDirectoryReference StageContentRoot = StagedDirectoryReference.Combine(SC.RelativeProjectRootForStage, "Content");
 
