Cloud Artifacts - Fixed issue were total size was always written as a string into the metadata json

#Synced-CL 51823794

diff --git a/Engine/Source/Programs/AutomationTool/BuildGraph/Tasks/CreateCloudArtifactTask.cs b/Engine/Source/Programs/AutomationTool/BuildGraph/Tasks/CreateCloudArtifactTask.cs
--- a/Engine/Source/Programs/AutomationTool/BuildGraph/Tasks/CreateCloudArtifactTask.cs
+++ b/Engine/Source/Programs/AutomationTool/BuildGraph/Tasks/CreateCloudArtifactTask.cs
@@ -706,10 +706,16 @@ namespace AutomationTool.Tasks
 							continue;
 						}
 
-						if (pair.Value is int value)
+						// Use IConvertible + TypeCode to detect any numeric primitive (sbyte through decimal)
+						// rather than explicitly checking each type. This makes sure numerical values are written as numbers.
+						if (pair.Value is IConvertible conv)
 						{
-							jsonWriter.WriteValue(pair.Key, value);
-							continue;
+							TypeCode tc = conv.GetTypeCode();
+							if (tc is >= TypeCode.SByte and <= TypeCode.Decimal)
+							{
+								jsonWriter.WriteValue(pair.Key, conv.ToDouble(null));
+								continue;
+							}
 						}
 						jsonWriter.WriteValue(pair.Key, pair.Value.ToString());
 
