Using the -ApplyIoStoreOnDemand cmdline arg for BuildCookRun will now generate a global.uondemandtoc file rather than using the legacy .iochunktoc in the Cloud directory path.

#rnx

#Synced-CL 51774074

diff --git a/Engine/Source/Programs/AutomationTool/Scripts/CopyBuildToStagingDirectory.Automation.cs b/Engine/Source/Programs/AutomationTool/Scripts/CopyBuildToStagingDirectory.Automation.cs
--- a/Engine/Source/Programs/AutomationTool/Scripts/CopyBuildToStagingDirectory.Automation.cs
+++ b/Engine/Source/Programs/AutomationTool/Scripts/CopyBuildToStagingDirectory.Automation.cs
@@ -4808,6 +4808,7 @@ namespace AutomationScripts
 					string ConfigFilePath = CombinePaths(CloudDir, "IoStoreOnDemand.ini");
 
 					UploadArgs.Append(string.Format("Upload \"{0}\"", PakPath));
+					UploadArgs.Append(string.Format(" -GlobalContainerPath \"{0}\"", PakPath));
 					UploadArgs.Append(string.Format(" -ServiceUrl={0}/obj/bucket", ServiceUrl));
 					UploadArgs.Append(" -Bucket=" + SC.ShortProjectName.ToLower());
 					UploadArgs.Append(" -ConfigFilePath=\"" + ConfigFilePath + "\"");
