Improve missing message logging in RpcExecutor

#Synced-CL 51955855

diff --git a/Engine/Source/Programs/AutomationTool/Gauntlet/Framework/RpcFramework/Gauntlet.RpcExecutor.cs b/Engine/Source/Programs/AutomationTool/Gauntlet/Framework/RpcFramework/Gauntlet.RpcExecutor.cs
--- a/Engine/Source/Programs/AutomationTool/Gauntlet/Framework/RpcFramework/Gauntlet.RpcExecutor.cs
+++ b/Engine/Source/Programs/AutomationTool/Gauntlet/Framework/RpcFramework/Gauntlet.RpcExecutor.cs
@@ -243,7 +243,8 @@ namespace Gauntlet
 					RpcTarget Target = ActiveTargetRegistry.GetTarget(IncomingMessage.SenderId);
 					if( null == Target )
 					{
-						Log.Warning($"Discarded incoming message with category '{IncomingMessage.Category}' : No RPC target registered for {IncomingMessage.SenderId}");
+						Log.Warning($"No RPC target registered for '{IncomingMessage.SenderId}' in active registry; routing to central queue to prevent message loss (Category: '{IncomingMessage.Category}').");
+						IncomingMessageQueue.Enqueue(IncomingMessage);
 					}
 					else
 					{
@@ -381,6 +382,17 @@ namespace Gauntlet
 			return NewMessages;
 		}
 
+		/// <summary>
+		/// Returns the string representations of all messages across all queues, without modifying the queues.
+		/// </summary>
+		/// <returns>List of message strings formatted as "[SenderId][Category] Payload"</returns>
+		public List<string> GetAllMessageStrings()
+		{
+			return GetNewMessages(bClearMessageQueue: false)
+				.Select(Msg => $"[{Msg.SenderId}][{Msg.Category}] {Msg.Payload}")
+				.ToList();
+		}
+
 		private void PopulateNewMessagesFromQueue(ConcurrentQueue<GauntletIncomingRpcMessage> Queue, List<GauntletIncomingRpcMessage> Messages, bool bClearMessageQueue)
 		{
 			int InitialCount = Queue.Count;
diff --git a/Engine/Source/Programs/AutomationTool/Gauntlet/Framework/RpcFramework/Gauntlet.RpcLibrary.cs b/Engine/Source/Programs/AutomationTool/Gauntlet/Framework/RpcFramework/Gauntlet.RpcLibrary.cs
--- a/Engine/Source/Programs/AutomationTool/Gauntlet/Framework/RpcFramework/Gauntlet.RpcLibrary.cs
+++ b/Engine/Source/Programs/AutomationTool/Gauntlet/Framework/RpcFramework/Gauntlet.RpcLibrary.cs
@@ -502,8 +502,19 @@ namespace Gauntlet
 			if (DateTime.UtcNow >= TimeoutEnd)
 			{
 				string NamesMissing = string.Join(", ", InTargets.Where(Target => !Target.HasSentMessage).Select(Target => Target.TargetName));
-				string ErrorMsg = $"Hang detected waiting for targets to send requested '{MessageCategory}' message. {TotalMatches} targets reported ready within {TimeoutInSeconds}sec. Missing targets: {NamesMissing}";
+				string ErrorMsg = $"Hang detected waiting for targets to send requested '{MessageCategory}' message. {TotalMatches} targets reported ready within {TimeoutInSeconds}sec.";
+				if (NamesMissing.Length > 0)
+				{
+					ErrorMsg += $" The following targets never responded: {NamesMissing}";
+				}
 				Log.Error(ErrorMsg);
+
+				List<string> UnmatchedMessages = RpcExecutor.Instance.GetAllMessageStrings();
+				if (UnmatchedMessages.Count > 0)
+				{
+					Log.Warning($"At the time of failure, the following unhandled RPC messages were queued:\n{String.Join('\n', UnmatchedMessages)}");
+				}
+
 				throw new TestException(ErrorMsg);
 			}
 
