Add ParamHelp decorator for CookedEditorFlavor argument for RunUnreal run by Gauntlet

#rnx

#Synced-CL 51937299

diff --git a/Engine/Source/Programs/AutomationTool/Gauntlet/Unreal/RunUnreal.cs b/Engine/Source/Programs/AutomationTool/Gauntlet/Unreal/RunUnreal.cs
--- a/Engine/Source/Programs/AutomationTool/Gauntlet/Unreal/RunUnreal.cs
+++ b/Engine/Source/Programs/AutomationTool/Gauntlet/Unreal/RunUnreal.cs
@@ -37,6 +37,7 @@ namespace Gauntlet
 	[ParamHelp("NoTimeout", "No maximum timeout", ParamType = typeof(bool))]
 	[ParamHelp("TestIterations", "Number of iterations to repeat this test", ParamType = typeof(int), DefaultValue = 1)]
 	[ParamHelp("CookedEditor", "Restricts usage of uncooked editor role", ParamType = typeof(bool))]
+	[ParamHelp("CookedEditorFlavor", "Specifies RequiredFlavor for the cooked editor role", ParamType = typeof(string))]
 	[ParamHelp("Devices", "List of devices to use for tests using the syntax: <platform>:<address>,<platform>:<address>,...", Action = ParamHelpAttribute.ParamAction.Append)]
 	[ParamHelp("DeviceName", "Device name to constrain reservation", Action = ParamHelpAttribute.ParamAction.Append)]
 	[ParamHelp("PrefSpec", "Device perfomance specification to constrain reservation (Minimum,Recommended,High)", Action = ParamHelpAttribute.ParamAction.Append)]
