Clean up the OnDemand/IAX settings being passed around the UAT cooking/staging scripts.

#rnx

- Previously ProjectParams had 2 input parameters related to the OnDemand system but we intend to add more. Rather than continuously increasing the number of parameters I have gathered the existing values in a new class IoStoreOnDemandSettings.
-- Note that the very similar parameter 'GenerateOptionalPakForNonChunkedBuild' is not actually OnDemand related and so has not been touched.
- The empty constructor version of IoStoreOnDemandSettings will have OnDemand disabled. Calling IoStoreOnDemandSettings.CreateDefault, will return a version with the default OnDemand settings, which will force all bulkdata to OnDemand.
- Reading of values from the cmdline has been moved from ProjectParams to IoStoreOnDemandSettings.

#Synced-CL 51774074

diff --git a/Engine/Source/Programs/AutomationTool/AutomationUtils/IoStoreOnDemandUtils.cs b/Engine/Source/Programs/AutomationTool/AutomationUtils/IoStoreOnDemandUtils.cs
new file mode 100644
--- /dev/null
+++ b/Engine/Source/Programs/AutomationTool/AutomationUtils/IoStoreOnDemandUtils.cs
@@ -0,0 +1,89 @@
+// Copyright Epic Games, Inc. All Rights Reserved.
+
+using EpicGames.Core;
+
+namespace AutomationTool
+{
+
+	public class IoStoreOnDemandSettings
+	{
+		static public IoStoreOnDemandSettings CreateDefault()
+		{
+			IoStoreOnDemandSettings EnabledSettings = new IoStoreOnDemandSettings();
+
+			EnabledSettings.IASSettings = new IASSpecificSettings();
+			EnabledSettings.IASSettings.Mode = IASSpecificSettings.RuleMode.Default;
+			EnabledSettings.IASSettings.ApplyToStagingManifest = true;
+
+			return EnabledSettings;
+		}
+
+		public void ReadFromCommandLine(BuildCommand Command)
+		{
+			if (Command.ParseParam("ApplyIoStoreOnDemand"))
+			{
+				IASSettings ??= new IASSpecificSettings();
+				IASSettings.Mode = IASSpecificSettings.RuleMode.Default;
+			}
+
+			if (Command.ParseParam("CreateDefaultOnDemandPakRule"))
+			{
+				IASSettings ??= new IASSpecificSettings();
+				IASSettings.Mode = IASSpecificSettings.RuleMode.Default;
+			}
+
+			if (Command.ParseParam("GenerateOnDemandPakForNonChunkedBuild"))
+			{
+				IASSettings ??= new IASSpecificSettings();
+				IASSettings.ApplyToStagingManifest = true;
+			}
+		}
+
+		public bool IsIASEnabled()
+		{
+			return IASSettings != null && IASSettings.Mode != IASSpecificSettings.RuleMode.None;
+		}
+
+		public bool ShouldCreateDefaultOnDemandPakRule()
+		{
+			return IASSettings != null && IASSettings.Mode == IASSpecificSettings.RuleMode.Default;
+		}
+
+		/// <summary>
+		/// Poll if pak rules should be applied when staging from a manifest rather than a chunked build.
+		/// This allows IAS to force the creation of 'OnDemand' container files  even when chunked builds
+		/// are not being used.
+		/// </summary>
+		/// <returns>True if the rules should be applied, otherwise false</returns>
+		public bool ApplyRulesForStagingManifest()
+		{
+			return IASSettings != null && IASSettings.ApplyToStagingManifest;
+		}
+
+		public class IADSpecificSettings
+		{
+
+		}
+
+		public class IASSpecificSettings
+		{
+			public enum RuleMode
+			{
+				/** No additional paking rules should be produced. */
+				None,
+				/** The default rule (move all bulkdata to 'OnDemand' containers) should be produced */
+				Default
+			}
+
+			public RuleMode Mode { get; set; } = RuleMode.None;
+			public bool ApplyToStagingManifest = false;
+		}
+
+		/// <summary>
+		/// This is a placeholder for IAD in the future and is currently unused.
+		/// </summary>
+		private IADSpecificSettings IADSettings { get; }
+
+		private IASSpecificSettings IASSettings;
+	}
+}
\ No newline at end of file
diff --git a/Engine/Source/Programs/AutomationTool/AutomationUtils/ProjectParams.cs b/Engine/Source/Programs/AutomationTool/AutomationUtils/ProjectParams.cs
--- a/Engine/Source/Programs/AutomationTool/AutomationUtils/ProjectParams.cs
+++ b/Engine/Source/Programs/AutomationTool/AutomationUtils/ProjectParams.cs
@@ -463,9 +463,8 @@ namespace AutomationTool
 			this.ProjectDescriptor = InParams.ProjectDescriptor;
 			this.Upload = InParams.Upload;
 			this.ApplyIoStoreOnDemand = InParams.ApplyIoStoreOnDemand;
-			this.CreateDefaultOnDemandPakRule = InParams.CreateDefaultOnDemandPakRule;
+			this.OnDemandSettings = InParams.OnDemandSettings;
 			this.GenerateOptionalPakForNonChunkedBuild = InParams.GenerateOptionalPakForNonChunkedBuild;
-			this.GenerateOnDemandPakForNonChunkedBuild = InParams.GenerateOnDemandPakForNonChunkedBuild;
 			this.XcodeBuildOptions = InParams.XcodeBuildOptions;
 			this.bMakeMacNative = InParams.bMakeMacNative;
 			this.DeviceLockedTimeout = InParams.DeviceLockedTimeout;
@@ -652,9 +651,8 @@ namespace AutomationTool
 			ParamList<string> InMapsToRebuildHLOD = null,
 			ParamList<string> TitleID = null,
 			string Upload = null,
-			bool? CreateDefaultOnDemandPakRule = null,
+			IoStoreOnDemandSettings OnDemandSettings = null,
 			bool? GenerateOptionalPakForNonChunkedBuild = null,
-			bool? GenerateOnDemandPakForNonChunkedBuild = null,
 			string XcodeBuildOptions = null,
 			bool bMakeMacNative = false
 			)
@@ -1124,10 +1122,10 @@ namespace AutomationTool
 				this.Manifests = true;
 			}
 
-			this.CreateDefaultOnDemandPakRule = GetParamValueIfNotSpecified(Command, CreateDefaultOnDemandPakRule, false, "CreateDefaultOnDemandPakRule") || this.ApplyIoStoreOnDemand;
+			this.OnDemandSettings = OnDemandSettings ?? new IoStoreOnDemandSettings(); // Make sure we have a default (OnDemand disabled) settings if none were provided.
+			this.OnDemandSettings.ReadFromCommandLine(Command);
 
 			this.GenerateOptionalPakForNonChunkedBuild = GetParamValueIfNotSpecified(Command, GenerateOptionalPakForNonChunkedBuild, false, "GenerateOptionalPakForNonChunkedBuild");
-			this.GenerateOnDemandPakForNonChunkedBuild = GetParamValueIfNotSpecified(Command, GenerateOnDemandPakForNonChunkedBuild, false, "GenerateOnDemandPakForNonChunkedBuild");
 
 			this.XcodeBuildOptions = ParseParamValueIfNotSpecified(Command, XcodeBuildOptions, "xcodebuildoptions", null);
 			this.bMakeMacNative = GetParamValueIfNotSpecified(Command, null, this.bMakeMacNative, "macnative");
@@ -2531,15 +2529,12 @@ namespace AutomationTool
 		[Help("applyiostoreondemand", "Forces IoStoreOnDemand to be enabled for the project even if it is not set up for it")]
 		public bool ApplyIoStoreOnDemand { get; }
 
-		[Help("createdefaultondemandpakrule", "Creates a default pakfile rule for OnDemand data if no other such rule exists. The default rule will include all bulkdata types.")]
-		public bool CreateDefaultOnDemandPakRule { get; }
+		[Help("ondemandsettings", "The settings for OnDemand data to apply to the project")]
+		public IoStoreOnDemandSettings OnDemandSettings { get; }
 
 		[Help("GenerateOptionalPakForNonChunkedBuild", "When true non chunked builds will attempt to generate an optional pak for optional bulk data.")]
 		public bool GenerateOptionalPakForNonChunkedBuild { get; }
 
-		[Help("GenerateOnDemandPakForNonChunkedBuild", "When true non chunked builds will attempt to generate an OnDemand pak. Will also force ApplyIoStoreOnDemand to true")]
-		public bool GenerateOnDemandPakForNonChunkedBuild { get; }
-
 		[Help("XcodeBuildOptions", "Extra options to pass to xcodebuild")]
 		public string XcodeBuildOptions { get; set; }
 
diff --git a/Engine/Source/Programs/AutomationTool/Scripts/CopyBuildToStagingDirectory.Automation.cs b/Engine/Source/Programs/AutomationTool/Scripts/CopyBuildToStagingDirectory.Automation.cs
--- a/Engine/Source/Programs/AutomationTool/Scripts/CopyBuildToStagingDirectory.Automation.cs
+++ b/Engine/Source/Programs/AutomationTool/Scripts/CopyBuildToStagingDirectory.Automation.cs
@@ -3165,16 +3165,16 @@ namespace AutomationScripts
 
 					RulesList.Add(PakRules);
 				}
-				else 
+				else
 				{
 					Logger.LogDebug("No files found for pakfilerule {0}", PakRules.Name);
 				}
 			}
 
-			if (Params.CreateDefaultOnDemandPakRule)
+			// If there are no pre-existing pak rules for an 'OnDemand' container we can check the OnDemandSettings to see if we should create some
+			if (RulesList.FindIndex(x => x.bOnDemand == true) == -1)
 			{
-				// If there are no pak rules for an 'OnDemand' container we will need to create one via CreateIoStoreOnDemandPakRules
-				if (RulesList.FindIndex(x => x.bOnDemand == true) == -1)
+				if (Params.OnDemandSettings.ShouldCreateDefaultOnDemandPakRule())
 				{
 					RulesList.AddRange(CreateIoStoreOnDemandPakRules());
 				}
@@ -3334,7 +3334,7 @@ namespace AutomationScripts
 				string OriginalFilename = StagingFile.Key;
 				string OrigExt = Path.GetExtension(OriginalFilename);
 
-				if (Params.GenerateOnDemandPakForNonChunkedBuild && MatchingRule.HasValue && MatchingRule.Value.bOnDemand)
+				if (Params.OnDemandSettings.ApplyRulesForStagingManifest() && MatchingRule.HasValue && MatchingRule.Value.bOnDemand)
 				{
 					if (Params.GenerateOptionalPakForNonChunkedBuild && OrigExt.Equals(OptionalBulkDataFileExtension))
 					{
