AutomationScripts: Allow specifying plugin dependencies for the BuildPlugin command, required in the situation where the plugin dependency is foreign to Engine/Plugins, so is not visible to UBT by default.

#8051
#Synced-CL 52000912

diff --git a/Engine/Source/Programs/AutomationTool/Scripts/BuildPluginCommand.Automation.cs b/Engine/Source/Programs/AutomationTool/Scripts/BuildPluginCommand.Automation.cs
--- a/Engine/Source/Programs/AutomationTool/Scripts/BuildPluginCommand.Automation.cs
+++ b/Engine/Source/Programs/AutomationTool/Scripts/BuildPluginCommand.Automation.cs
@@ -2,7 +2,6 @@
 using System;
 using System.Collections.Generic;
 using System.IO;
-using System.Threading;
 using System.Linq;
 using System.Reflection;
 using AutomationTool;
@@ -12,10 +11,9 @@ using System.Text;
 using UnrealBuildBase;
 using Microsoft.Extensions.Logging;
 
-using static AutomationTool.CommandUtils;
-
 [Help("Builds a plugin, and packages it for distribution")]
 [Help("Plugin", "Specify the path to the descriptor file for the plugin that should be packaged")]
+[Help("Dependencies", "Specify the path to the plugin descriptor file for a dependency of the Plugin argument - only required if it is not accessible in the Engine/Plugins directory. Can be specified multiple times, one per dependency.")]
 [Help("NoHostPlatform", "Prevent compiling for the editor platform on the host")]
 [Help("HostPlatforms", "Specify a list of host platforms to build, separated by '+' characters (eg. -HostPlatforms=Win32+Win64). Default is the current host platforms")]
 [Help("TargetPlatforms", "Specify a list of target platforms to build, separated by '+' characters (eg. -TargetPlatforms=Win32+Win64). Default is all the Rocket target platforms.")]
@@ -61,6 +59,18 @@ public sealed class BuildPlugin : BuildCommand
 			throw new AutomationException("Missing -Package=... argument");
 		}
 
+		string[] Dependencies = ParseParamValues("Dependencies");
+		List<FileReference> DependencyPlugins = new(Dependencies.Length);
+		foreach (string Dependency in Dependencies)
+		{
+			FileReference DependencyPlugin = new FileReference(Dependency);
+			if (!FileReference.Exists(DependencyPlugin))
+			{
+				throw new AutomationException("Dependency plugin '{0}' not found", DependencyPlugin);
+			}
+			DependencyPlugins.Add(DependencyPlugin);
+		}
+
 		// Option for verifying that all include directive s
 		bool bStrictIncludes = ParseParam("StrictIncludes");
 
@@ -112,7 +122,7 @@ public sealed class BuildPlugin : BuildCommand
 
 		// Create a host project for the plugin. For script generator plugins, we need to have UHT be able to load it, which can only happen if it's enabled in a project.
 		FileReference HostProjectFile = FileReference.Combine(PackageDir, "HostProject", "HostProject.uproject");
-		FileReference HostProjectPluginFile = CreateHostProject(HostProjectFile, PluginFile);
+		FileReference HostProjectPluginFile = CreateHostProject(HostProjectFile, PluginFile, DependencyPlugins);
 
 		// Read the plugin
 		Logger.LogInformation("Reading plugin from {HostProjectPluginFile}...", HostProjectPluginFile);
@@ -158,19 +168,32 @@ public sealed class BuildPlugin : BuildCommand
 		}
 	}
 
-	FileReference CreateHostProject(FileReference HostProjectFile, FileReference PluginFile)
+	private FileReference CreateHostProject(FileReference HostProjectFile, FileReference PluginFile, IReadOnlyList<FileReference> DependencyPlugins)
 	{
 		DirectoryReference HostProjectDir = HostProjectFile.Directory;
 		DirectoryReference.CreateDirectory(HostProjectDir);
 
 		// Create the new project descriptor
-		File.WriteAllText(HostProjectFile.FullName, "{ \"FileVersion\": 3, \"Plugins\": [ { \"Name\": \"" + PluginFile.GetFileNameWithoutExtension() + "\", \"Enabled\": true } ] }");
+		StringBuilder PluginsArray = new($"{{ \"Name\": \"{PluginFile.GetFileNameWithoutExtension()}\", \"Enabled\": true }}");
+		foreach (FileReference DependencyPlugin in DependencyPlugins)
+		{
+			PluginsArray.Append($", {{ \"Name\": \"{DependencyPlugin.GetFileNameWithoutExtension()}\", \"Enabled\": true }}");
+		}
+
+		File.WriteAllText(HostProjectFile.FullName, $"{{ \"FileVersion\": 3, \"Plugins\": [ {PluginsArray} ] }}");
 
 		// Get the plugin directory in the host project, and copy all the files in
 		DirectoryReference HostProjectPluginDir = DirectoryReference.Combine(HostProjectDir, "Plugins", PluginFile.GetFileNameWithoutExtension());
 		CommandUtils.ThreadedCopyFiles(PluginFile.Directory.FullName, HostProjectPluginDir.FullName);
 		CommandUtils.DeleteDirectory(true, DirectoryReference.Combine(HostProjectPluginDir, "Intermediate").FullName);
 
+		foreach (FileReference DependencyPlugin in DependencyPlugins)
+		{
+			DirectoryReference HostProjectDependencyDir = DirectoryReference.Combine(HostProjectDir, "Plugins",  DependencyPlugin.GetFileNameWithoutExtension());
+			CommandUtils.ThreadedCopyFiles(DependencyPlugin.Directory.FullName, HostProjectDependencyDir.FullName);
+			CommandUtils.DeleteDirectory(true, DirectoryReference.Combine(HostProjectDependencyDir, "Intermediate").FullName);
+		}
+
 		// Return the path to the plugin file in the host project
 		return FileReference.Combine(HostProjectPluginDir, PluginFile.GetFileName());
 	}
