Add Horde Step Messages to distinguish failure types
Preflight failure necessary to check step messages.
#rnx

#Synced-CL 51581246

#Synced-CL 51581246

#Synced-CL 52011596

diff --git a/Engine/Source/Programs/AutomationTool/Gauntlet/Framework/Horde/Gauntlet.Horde.cs b/Engine/Source/Programs/AutomationTool/Gauntlet/Framework/Horde/Gauntlet.Horde.cs
--- a/Engine/Source/Programs/AutomationTool/Gauntlet/Framework/Horde/Gauntlet.Horde.cs
+++ b/Engine/Source/Programs/AutomationTool/Gauntlet/Framework/Horde/Gauntlet.Horde.cs
@@ -87,6 +87,11 @@ namespace Gauntlet
 			public static HordeStepMessage InvalidBackend = new HordeStepMessage("InvalidBackend", JobStepMessageLevels.Error);
 			public static HordeStepMessage WrongAccountUsed = new HordeStepMessage("WrongAccountUsed", JobStepMessageLevels.Error);
 			public static HordeStepMessage TestCancelled = new HordeStepMessage("TestCancelled", JobStepMessageLevels.Error);
+			public static HordeStepMessage InvalidProject = new HordeStepMessage("InvalidProject", JobStepMessageLevels.Error);
+			public static HordeStepMessage NoVerseFiles = new HordeStepMessage("NoVerseFiles", JobStepMessageLevels.Error);
+			public static HordeStepMessage InvalidValues = new HordeStepMessage("InvalidValues", JobStepMessageLevels.Error);
+			public static HordeStepMessage PermissionsFailure = new HordeStepMessage("PermissionsFailure", JobStepMessageLevels.Error);
+			public static HordeStepMessage InvalidTestData = new HordeStepMessage("InvalidTestData", JobStepMessageLevels.Error);
 
 			// Failures related to game/service functionality
 			public static HordeStepMessage LoginFailed = new HordeStepMessage("LoginFailed", JobStepMessageLevels.Error);
@@ -102,6 +107,8 @@ namespace Gauntlet
 			public static HordeStepMessage ServerCrashDetected = new HordeStepMessage("ServerCrashDetected", JobStepMessageLevels.Error);
 			public static HordeStepMessage ClientCrashDetected = new HordeStepMessage("ClientCrashDetected", JobStepMessageLevels.Error);
 			public static HordeStepMessage EditorCrashDetected = new HordeStepMessage("EditorCrashDetected", JobStepMessageLevels.Error);
+			public static HordeStepMessage ServerError = new HordeStepMessage("ServerError", JobStepMessageLevels.Error);
+			public static HordeStepMessage SyncError = new HordeStepMessage("SyncError", JobStepMessageLevels.Error);
 
 			// Warning states for if things are getting close to bad
 			public static HordeStepMessage TestWarningsDetected = new HordeStepMessage("TestWarningsDetected", JobStepMessageLevels.Warning);
@@ -116,6 +123,7 @@ namespace Gauntlet
 
 			// Failed states
 			public static HordeStepMessage ZeroTestsRun = new HordeStepMessage("ZeroTestsRun", JobStepMessageLevels.Error);
+			public static HordeStepMessage CleanupFailure = new HordeStepMessage("CleanupFailure", JobStepMessageLevels.Error);
 
 		}
 
