[Gauntlet] Ensure disposal of UnrealLog/UnrealLogParser file streams
This ensures no remaining file handles exist on process log files during test cleanup
#Synced-CL 51877528

diff --git a/Engine/Source/Programs/AutomationTool/Gauntlet/Framework/Devices/Gauntlet.DevicePool.cs b/Engine/Source/Programs/AutomationTool/Gauntlet/Framework/Devices/Gauntlet.DevicePool.cs
--- a/Engine/Source/Programs/AutomationTool/Gauntlet/Framework/Devices/Gauntlet.DevicePool.cs
+++ b/Engine/Source/Programs/AutomationTool/Gauntlet/Framework/Devices/Gauntlet.DevicePool.cs
@@ -301,6 +301,12 @@ namespace Gauntlet
 		/// </summary>
 		public static bool DeviceReservationBlock;
 
+		/// <summary>
+		/// Whether or not a device reservation block is in it's final active step
+		/// This is used to prevent cleaning of the tempdir where builds may be stored
+		/// </summary>
+		public static bool DeviceReservationBlockFinalStep;
+
 		/// <summary>
 		/// Whether or not this particular step is an installations tep in a device reservation block
 		/// This is primarily used within the context of UnrealSession
diff --git a/Engine/Source/Programs/AutomationTool/Gauntlet/Framework/Devices/Gauntlet.DeviceReservationService.cs b/Engine/Source/Programs/AutomationTool/Gauntlet/Framework/Devices/Gauntlet.DeviceReservationService.cs
--- a/Engine/Source/Programs/AutomationTool/Gauntlet/Framework/Devices/Gauntlet.DeviceReservationService.cs
+++ b/Engine/Source/Programs/AutomationTool/Gauntlet/Framework/Devices/Gauntlet.DeviceReservationService.cs
@@ -159,9 +159,14 @@ namespace Gauntlet
 						}
 						else if (Value.Equals("Install", StringComparison.InvariantCultureIgnoreCase))
 						{
-							Log.Verbose("Detected running step an explicit install step within a reservation block. Any devices used will be cleaned");
+							Log.Verbose("Detected running step as an explicit install step within a reservation block. Any devices used will be cleaned");
 							InstallStep = true;
 						}
+						else if (Value.Equals("End", StringComparison.OrdinalIgnoreCase))
+						{
+							Log.Verbose("Detected running step as the final step within a reservation block.");
+							DevicePool.DeviceReservationBlockFinalStep = true;
+						}
 
 						if (InstallStep)
 						{
@@ -190,6 +195,11 @@ namespace Gauntlet
 								break;
 							}
 						}
+
+						if (Step.Id == Batch.Steps.Last().Id)
+						{
+							DevicePool.DeviceReservationBlockFinalStep = true;
+						}
 					}
 				}
 				catch (Exception Ex)
diff --git a/Engine/Source/Programs/AutomationTool/Gauntlet/Unreal/Base/Gauntlet.UnrealTestNode.cs b/Engine/Source/Programs/AutomationTool/Gauntlet/Unreal/Base/Gauntlet.UnrealTestNode.cs
--- a/Engine/Source/Programs/AutomationTool/Gauntlet/Unreal/Base/Gauntlet.UnrealTestNode.cs
+++ b/Engine/Source/Programs/AutomationTool/Gauntlet/Unreal/Base/Gauntlet.UnrealTestNode.cs
@@ -1176,6 +1176,14 @@ namespace Gauntlet
 					UnrealApp.ShutdownInstance();
 				}
 			}
+
+			if (RoleResults != null)
+			{
+				foreach (var RoleResult in RoleResults)
+				{
+					RoleResult?.LogSummary?.Dispose();
+				}
+			}
 		}
 
 		/// <summary>
diff --git a/Engine/Source/Programs/AutomationTool/Gauntlet/Unreal/RunUnreal.cs b/Engine/Source/Programs/AutomationTool/Gauntlet/Unreal/RunUnreal.cs
--- a/Engine/Source/Programs/AutomationTool/Gauntlet/Unreal/RunUnreal.cs
+++ b/Engine/Source/Programs/AutomationTool/Gauntlet/Unreal/RunUnreal.cs
@@ -378,8 +378,13 @@ namespace Gauntlet
 
 				bool bCleanFromArgument = ParseParam("Clean");
 				bool bCleanFromEnvironmentVariable = bool.TryParse(Environment.GetEnvironmentVariable(CleanEnvironmentVariable), out bool bCleanEnvvar) && bCleanEnvvar;
-				if (bCleanFromArgument || bCleanFromEnvironmentVariable)
+				bool bCleanFromEnvvarAllowed = !DevicePool.DeviceReservationBlock || DevicePool.DeviceReservationBlockFinalStep;
+				if (bCleanFromArgument || (bCleanFromEnvironmentVariable && bCleanFromEnvvarAllowed))
 				{
+					// Flush any undisposed file streams that may be holding onto tempdir files
+					GC.Collect();
+					GC.WaitForPendingFinalizers();
+
 					Logger.LogInformation("Cleaning TempDir '{GauntletTempDir}' by request...", Options.TempDir);
 					DeleteDirectory(Options.TempDir);
 				}
diff --git a/Engine/Source/Programs/AutomationTool/Gauntlet/Unreal/Utils/Gauntlet.UnrealLogParser.cs b/Engine/Source/Programs/AutomationTool/Gauntlet/Unreal/Utils/Gauntlet.UnrealLogParser.cs
--- a/Engine/Source/Programs/AutomationTool/Gauntlet/Unreal/Utils/Gauntlet.UnrealLogParser.cs
+++ b/Engine/Source/Programs/AutomationTool/Gauntlet/Unreal/Utils/Gauntlet.UnrealLogParser.cs
@@ -18,7 +18,7 @@ namespace Gauntlet
 	/// Represents a prepared summary of the most relevant info in a log. Generated once by a
 	/// LogParser and cached.
 	/// </summary>
-	public class UnrealLog
+	public class UnrealLog : IDisposable
 	{
 		/// <summary>
 		/// Represents the level
@@ -207,6 +207,23 @@ namespace Gauntlet
 			_hasParseRequestedExitReason = new(() => GetRequestedExitReason());
 		}
 
+		#region IDisposable Support
+		private bool _disposed = false;
+		~UnrealLog()
+		{
+			Dispose();
+		}
+
+		public void Dispose()
+		{
+			if (!_disposed)
+			{
+				_parser?.Dispose();
+				_disposed = true;
+			}
+		}
+		#endregion
+
 		private bool GetRequestedExitReason()
 		{
 			// Check request exit and reason
@@ -609,7 +626,7 @@ namespace Gauntlet
 	/// <summary>
 	/// Helper class for parsing logs
 	/// </summary>
-	public class UnrealLogParser
+	public class UnrealLogParser : IDisposable
 	{
 		/// <summary>
 		/// DEPRECATED - Use GetLogReader() to read through the log stream efficiently.
@@ -655,14 +672,25 @@ namespace Gauntlet
 		public UnrealLogParser(string InContent) : this(new DynamicStringReader(() => InContent))
 		{ }
 
-		/// <summary>
-		/// Constructor that takes a UnrealLog instance
-		/// </summary>
-		/// <param name="InLog"></param>
-		public UnrealLogParser(UnrealLog InLog) : this(InLog.GetParser().GetLogReader())
+		#region IDisposable Support
+		private bool _disposed = false;
+		~UnrealLogParser()
+		{
+			Dispose();
+		}
+
+		public void Dispose()
 		{
-			_logEntries = new(() => InLog.GetParser().LogEntries);
+			if (!_disposed)
+			{
+				if (_logReader is IDisposable disposable)
+				{
+					disposable.Dispose();
+				}
+				_disposed = true;
+			}
 		}
+		#endregion
 
 		protected List<UnrealLog.LogEntry> ParseEntries()
 		{
