IoStoreOnDemandSettings may now provide a list of assets to force into OnDemand containers during staging.

#rnx

#Synced-CL 51774074

diff --git a/Engine/Source/Programs/AutomationTool/AutomationUtils/IoStoreOnDemandUtils.cs b/Engine/Source/Programs/AutomationTool/AutomationUtils/IoStoreOnDemandUtils.cs
--- a/Engine/Source/Programs/AutomationTool/AutomationUtils/IoStoreOnDemandUtils.cs
+++ b/Engine/Source/Programs/AutomationTool/AutomationUtils/IoStoreOnDemandUtils.cs
@@ -1,21 +1,18 @@
 // Copyright Epic Games, Inc. All Rights Reserved.
 
 using EpicGames.Core;
+using System.Collections.Generic;
 
 namespace AutomationTool
 {
 
 	public class IoStoreOnDemandSettings
 	{
-		static public IoStoreOnDemandSettings CreateDefault()
+		public void SetIASGloballyEnabled()
 		{
-			IoStoreOnDemandSettings EnabledSettings = new IoStoreOnDemandSettings();
-
-			EnabledSettings.IASSettings = new IASSpecificSettings();
-			EnabledSettings.IASSettings.Mode = IASSpecificSettings.RuleMode.Default;
-			EnabledSettings.IASSettings.ApplyToStagingManifest = true;
-
-			return EnabledSettings;
+			IASSettings = new IASSpecificSettings();
+			IASSettings.Mode = IASSpecificSettings.RuleMode.Default;
+			IASSettings.ApplyToStagingManifest = true;
 		}
 
 		public void ReadFromCommandLine(BuildCommand Command)
@@ -49,6 +46,23 @@ namespace AutomationTool
 			return IASSettings != null && IASSettings.Mode == IASSpecificSettings.RuleMode.Default;
 		}
 
+		public bool ShouldCreateAssetTypeOnDemandPakRule()
+		{
+			return IASSettings != null && IASSettings.Mode == IASSpecificSettings.RuleMode.AssetType;
+		}
+
+		public Dictionary<string, List<string>> GetEnabledAssetList()
+		{
+			if (IASSettings != null)
+			{
+				return IASSettings.Assets;
+			}
+			else
+			{
+				return new Dictionary<string, List<string>>();
+			}
+		}
+
 		/// <summary>
 		/// Poll if pak rules should be applied when staging from a manifest rather than a chunked build.
 		/// This allows IAS to force the creation of 'OnDemand' container files  even when chunked builds
@@ -72,18 +86,22 @@ namespace AutomationTool
 				/** No additional paking rules should be produced. */
 				None,
 				/** The default rule (move all bulkdata to 'OnDemand' containers) should be produced */
-				Default
+				Default,
+				/** Produce rules based on asset types/ package lists to create 'OnDemand' containers */
+				AssetType
 			}
 
 			public RuleMode Mode { get; set; } = RuleMode.None;
+
 			public bool ApplyToStagingManifest = false;
+
+			// AssetName->ListOfPackages
+			public Dictionary<string, List<string>> Assets { get; set; } = new();
 		}
 
-		/// <summary>
-		/// This is a placeholder for IAD in the future and is currently unused.
-		/// </summary>
-		private IADSpecificSettings IADSettings { get; }
+		protected IADSpecificSettings IADSettings { get; }
 
-		private IASSpecificSettings IASSettings;
+		protected IASSpecificSettings IASSettings;
 	}
+
 }
\ No newline at end of file
diff --git a/Engine/Source/Programs/AutomationTool/Scripts/CopyBuildToStagingDirectory.Automation.cs b/Engine/Source/Programs/AutomationTool/Scripts/CopyBuildToStagingDirectory.Automation.cs
--- a/Engine/Source/Programs/AutomationTool/Scripts/CopyBuildToStagingDirectory.Automation.cs
+++ b/Engine/Source/Programs/AutomationTool/Scripts/CopyBuildToStagingDirectory.Automation.cs
@@ -3178,6 +3178,11 @@ namespace AutomationScripts
 				{
 					RulesList.AddRange(CreateIoStoreOnDemandPakRules());
 				}
+				else if (Params.OnDemandSettings.ShouldCreateAssetTypeOnDemandPakRule())
+				{
+					Dictionary<string, List<string>> EnabledAssets = Params.OnDemandSettings.GetEnabledAssetList();
+					RulesList.AddRange(CreateAssetTypeOnDemandPakRules(EnabledAssets));
+				}
 			}
 
 			if (RulesList.Count == 0)
@@ -3216,6 +3221,29 @@ namespace AutomationScripts
 			return RulesList;
 		}
 
+		private static List<PakFileRules> CreateAssetTypeOnDemandPakRules(Dictionary<string, List<string>> AssetTypes)
+		{
+			List<PakFileRules> RulesList = new List<PakFileRules>();
+
+			PakFileRules PakRules = new PakFileRules();
+			PakRules.Name = "BulkDataOnDemand";
+			PakRules.bOnDemand = true;
+			PakRules.bUseExactFilePathMatch = true;
+			PakRules.ExactFileMatches = new HashSet<string>();
+			foreach (KeyValuePair<string, List<string>> EnabledAsset in AssetTypes)
+			{
+				foreach(string Package in EnabledAsset.Value)
+				{
+					//PakRules.ExactFileMatches.Add(Package + ".uptnl");
+					PakRules.ExactFileMatches.Add(Package + ".ubulk");
+				}
+			}
+
+			RulesList.Add(PakRules);
+
+			return RulesList;
+		}
+
 		/// <summary>
 		/// Attempts to apply the pak file rules to a specific staging file, returns false if it should be excluded
 		/// </summary>
