[Gauntlet] Remove a couple un-needed experimental feature flag checks now that the modern launch session flow is the default
#Synced-CL 51877528

#Synced-CL 51877528

diff --git a/Engine/Source/Programs/AutomationTool/Gauntlet/Framework/Base/Gauntlet.TargetDeviceDesktopCommon.cs b/Engine/Source/Programs/AutomationTool/Gauntlet/Framework/Base/Gauntlet.TargetDeviceDesktopCommon.cs
--- a/Engine/Source/Programs/AutomationTool/Gauntlet/Framework/Base/Gauntlet.TargetDeviceDesktopCommon.cs
+++ b/Engine/Source/Programs/AutomationTool/Gauntlet/Framework/Base/Gauntlet.TargetDeviceDesktopCommon.cs
@@ -458,7 +458,7 @@ namespace Gauntlet
 				}
 			}
 
-			if(Globals.UseExperimentalFeatures && Globals.IsRunningDev && AppConfig.OverlayExecutable.GetOverlay(ExecutablePath, out string OverlayExecutable))
+			if(Globals.IsRunningDev && AppConfig.OverlayExecutable.GetOverlay(ExecutablePath, out string OverlayExecutable))
 			{
 				CommandArguments += string.Format(" -basedir=\"{0}\"", Path.GetDirectoryName(ExecutablePath));
 				ExecutablePath = OverlayExecutable;
diff --git a/Engine/Source/Programs/AutomationTool/Gauntlet/Unreal/Base/Gauntlet.UnrealTestNode.cs b/Engine/Source/Programs/AutomationTool/Gauntlet/Unreal/Base/Gauntlet.UnrealTestNode.cs
--- a/Engine/Source/Programs/AutomationTool/Gauntlet/Unreal/Base/Gauntlet.UnrealTestNode.cs
+++ b/Engine/Source/Programs/AutomationTool/Gauntlet/Unreal/Base/Gauntlet.UnrealTestNode.cs
@@ -841,14 +841,7 @@ namespace Gauntlet
 				PrepareUnrealApp();
 			}
 
-			if(Globals.UseExperimentalFeatures)
-			{
-				return true;
-			}
-			else
-			{
-				return UnrealApp.TryReserveDevices();
-			}
+			return true;
 		}
 
 		/// <summary>
