Use new unified requests for archiving the packages too
- missed this function in the previous change

#rnx

#Synced-CL 51936548

diff --git a/Engine/Source/Programs/AutomationTool/Microsoft/GDKPackageGenerator.cs b/Engine/Source/Programs/AutomationTool/Microsoft/GDKPackageGenerator.cs
--- a/Engine/Source/Programs/AutomationTool/Microsoft/GDKPackageGenerator.cs
+++ b/Engine/Source/Programs/AutomationTool/Microsoft/GDKPackageGenerator.cs
@@ -172,23 +172,19 @@ namespace AutomationUtils.GDK
 
 		public void GetPackageFilesToArchive(ProjectParams Params, DeploymentContext SC, string PackageOutputPathFragment = "")
 		{
-			if (GenerateSinglePackages || Params.HasCreateReleaseVersion)
+			List<PackageRequest> PackageRequests = MakePackageRequests(Params, SC);
+
+			foreach (PackageRequest Request in PackageRequests)
 			{
-				foreach (var Config in SC.StageTargetConfigurations)
+				string PackageConfig = (Request.DLCPackage == null) ? (Request.Config?.ToString() ?? "Combined") : Request.DLCPackage;
+
+				DirectoryReference PackagePath = DirectoryReference.Combine(SC.ProjectRoot, "Saved", "Packages", SC.CookPlatform, PackageOutputPathFragment, PackageConfig);
+				if (Params.HasCreateReleaseVersion)
 				{
-					DirectoryReference PackagePath = DirectoryReference.Combine(SC.ProjectRoot, "Saved", "Packages", SC.CookPlatform, PackageOutputPathFragment, $"{Config}");
-					if (Params.HasCreateReleaseVersion)
-					{
-						PackagePath = DirectoryReference.FromString(Path.Combine(Params.GetCreateReleaseVersionPath(SC, Params.Client), PackageOutputPathFragment, Config.ToString()));
-					}
-					SC.ArchiveFiles(PackagePath.FullName, NewPath: Path.Combine(PackageOutputPathFragment, Config.ToString()));
+					PackagePath = DirectoryReference.FromString(Path.Combine(Params.GetCreateReleaseVersionPath(SC, Params.Client), PackageOutputPathFragment, PackageConfig));
 				}
-			}
 
-			if (GenerateCombinedPackage && !Params.HasCreateReleaseVersion)
-			{
-				DirectoryReference PackagePath = DirectoryReference.Combine(SC.ProjectRoot, "Saved", "Packages", SC.CookPlatform, PackageOutputPathFragment, "Combined");
-				SC.ArchiveFiles(PackagePath.FullName, NewPath: Path.Combine(PackageOutputPathFragment, "Combined")); ;
+				SC.ArchiveFiles(PackagePath.FullName, NewPath: Path.Combine(PackageOutputPathFragment, PackageConfig));
 			}
 		}
 
