Updating GDK package generator to support automatic generation of DLC packages
- enabled via bAutoCreateDLCPackages
- new DLCPackages and DLCChunks fields in the target settings let you choose chunks to move into DLC
- mutually exclusive with DLCPlugins, which now just references a DLCPackage (breaking change)

DLC plugin fixes:
- enforce that the selected DLC plugin should be explicitly loaded
- this allows it to be propertly unloaded too

#Synced-CL 51936548

diff --git a/Engine/Plugins/Experimental/Microsoft/GDKPlatformDLC/Source/GDKPlatformDLC/Private/GDKPlatformDLC.cpp b/Engine/Plugins/Experimental/Microsoft/GDKPlatformDLC/Source/GDKPlatformDLC/Private/GDKPlatformDLC.cpp
--- a/Engine/Plugins/Experimental/Microsoft/GDKPlatformDLC/Source/GDKPlatformDLC/Private/GDKPlatformDLC.cpp
+++ b/Engine/Plugins/Experimental/Microsoft/GDKPlatformDLC/Source/GDKPlatformDLC/Private/GDKPlatformDLC.cpp
@@ -59,6 +59,29 @@ void FGDKPlatformDLC::Initialize()
 	SetStoreUser( FPlatformUserId::CreateFromInternalId(0) );
 
 	// build the StoreId map
+	TArray<FString> DLCPackages;
+	if (GConfig->GetArray(IGDKRuntimeModule::Get().GetConfigSectionName(), TEXT("DLCPackages"), DLCPackages, GEngineIni))
+	{
+		FScopeLock Lock(&PackageDataLock);
+		for (FString& DLCPackage : DLCPackages)
+		{
+			DLCPackage.TrimStartAndEndInline();
+			DLCPackage.ReplaceInline(TEXT("("), TEXT(""));
+			DLCPackage.ReplaceInline(TEXT(")"), TEXT(""));
+		
+			FString DLCName;
+			FString StoreId;
+			if (FParse::Value(*DLCPackage, TEXT("DLCName="), DLCName) && FParse::Value(*DLCPackage, TEXT("StoreId="), StoreId))
+			{
+				DLCNameToStoreIdMap.Add(FName(*DLCName), StoreId);
+
+				TSharedPtr<FPackageData> Package = AllPackages.Add_GetRef(MakeShared<FPackageData>(this, FName(*DLCName), StoreId));
+				Package->Context.Package = Package;
+			}
+		}
+	}
+	
+	// build the plugin map
 	TArray<FString> DLCPlugins;
 	if (GConfig->GetArray(IGDKRuntimeModule::Get().GetConfigSectionName(), TEXT("DLCPlugins"), DLCPlugins, GEngineIni))
 	{
@@ -70,13 +93,17 @@ void FGDKPlatformDLC::Initialize()
 			DLCPlugin.ReplaceInline(TEXT(")"), TEXT(""));
 		
 			FString DLCName;
-			FString StoreId;
-			if (FParse::Value(*DLCPlugin, TEXT("DLCName="), DLCName) && FParse::Value(*DLCPlugin, TEXT("StoreId="), StoreId))
+			FString Plugin;
+			if (FParse::Value(*DLCPlugin, TEXT("DLCName="), DLCName) && FParse::Value(*DLCPlugin, TEXT("Plugin="), Plugin))
 			{
-				DLCNameToStoreIdMap.Add(FName(*DLCName), StoreId);
-
-				TSharedPtr<FPackageData> Package = AllPackages.Add_GetRef(MakeShared<FPackageData>(this, FName(*DLCName), StoreId));
-				Package->Context.Package = Package;
+				if (ensure(DLCNameToStoreIdMap.Contains(FName(*DLCName))))
+				{			
+					DLCNameToPluginNameMap.Add(FName(*DLCName), Plugin);
+				}
+				else
+				{
+					UE_LOG(LogPlatformDLC, Warning, TEXT("Ignoring DLCPlugin %s because DLC %s is not known"), *Plugin, *DLCName);
+				}
 			}
 		}
 	}
@@ -214,7 +241,7 @@ void FGDKPlatformDLC::AsyncInitDLC(const XPackageDetails& Details, bool bMount)
 
 		if (!Package.IsValid())
 		{
-			UE_LOG(LogPlatformDLC, Warning, TEXT("Ignoring DLC with StoreId %s because it isn't referenced in DLCPlugins config array"), UTF8_TO_TCHAR(Details.storeId));
+			UE_LOG(LogPlatformDLC, Warning, TEXT("Ignoring DLC with StoreId %s because it isn't referenced in DLCPackages config array"), UTF8_TO_TCHAR(Details.storeId));
 			return;
 		}
 
@@ -528,24 +555,31 @@ void FGDKPlatformDLC::MountPlugin_GameThread( TSharedPtr<FPackageData> Package )
 		}
 	}
 
-	if (bAutoMountPlugin && bAutoMountPakFiles)
+	if (bAutoMountPlugin && bAutoMountPakFiles && DLCNameToPluginNameMap.Contains(Package->DLCName))
 	{
+		const FString& PluginName = DLCNameToPluginNameMap.FindChecked(Package->DLCName);
+	
 		// find & load the .uplugin (should only be one per DLC, as specified by "RunUAT ... -DLCName=<plugin>")
 		TArray<FString> FoundPlugins;
-		IPlatformFile::GetPlatformPhysical().FindFilesRecursively(FoundPlugins, *Package->MountPath, TEXT(".uplugin") );
+		IPlatformFile::GetPlatformPhysical().FindFilesRecursively(FoundPlugins, *Package->MountPath, *FString::Printf(TEXT("%s.uplugin"), *PluginName) );
 		checkf(FoundPlugins.Num() != 0, TEXT("found %d plugins in %s - DLC should have 1!"), FoundPlugins.Num(), *Package->MountPath);
-		const FString PluginFileName = FoundPlugins.Last();
 
-		IPluginManager::Get().AddToPluginsList(PluginFileName);
-		FString DLCName = FPaths::GetBaseFilename(PluginFileName);
-
-		Package->DLCName = FName(DLCName);
-		Package->DLCPlugin = IPluginManager::Get().FindPlugin(DLCName);
-		checkf(Package->DLCPlugin.IsValid(), TEXT("failed to find DLC plugin %s"), *Package->DLCName.ToString());
+		Package->DLCPluginFilePath = FoundPlugins.Last();
+		IPluginManager::Get().AddToPluginsList(Package->DLCPluginFilePath);
 
+		Package->DLCPlugin = IPluginManager::Get().FindPlugin(PluginName);
+		checkf(Package->DLCPlugin.IsValid(), TEXT("failed to find DLC plugin %s"), *PluginName);
 
 		// ensure the plugin is mounted - will also load the shader code library if there is one
-		IPluginManager::Get().MountNewlyCreatedPlugin(Package->DLCName.ToString());
+		if (Package->DLCPlugin->GetDescriptor().bExplicitlyLoaded)
+		{
+			IPluginManager::Get().MountExplicitlyLoadedPlugin_FromFileName(Package->DLCPluginFilePath);
+		}
+		else
+		{
+			UE_LOG(LogPlatformDLC, Warning, TEXT("Plugin %s is not marked for explicit loading - unmounting DLC %s may fail"), *PluginName, *Package->DLCName.ToString());
+			IPluginManager::Get().MountNewlyCreatedPlugin(PluginName);
+		}
 
 		// load the asset registry for the new plugin
 		if (Package->DLCPlugin->CanContainContent())
@@ -629,13 +663,17 @@ bool FGDKPlatformDLC::Unmount( FName DLCName )
 		}
 
 		// attempt to unregister the plugin as much as possible
-		if (bAutoMountPlugin && bAutoMountPakFiles)
+		if (bAutoMountPlugin && bAutoMountPakFiles && Package->DLCPlugin.IsValid())
 		{
+			const FString& PluginName = DLCNameToPluginNameMap.FindChecked(Package->DLCName);
+
 			FText FailureReason;
-			if (Package->DLCPlugin->GetDescriptor().bExplicitlyLoaded && IPluginManager::Get().UnmountExplicitlyLoadedPlugin(Package->DLCName.ToString(), &FailureReason))
+			if (Package->DLCPlugin->GetDescriptor().bExplicitlyLoaded)
 			{
-				const FString PluginFileName = FPaths::Combine( Package->DLCPlugin->GetBaseDir(), Package->DLCName.ToString() + TEXT(".uplugin") );
-				IPluginManager::Get().RemoveFromPluginsList( PluginFileName );
+				bool bResult = IPluginManager::Get().UnmountExplicitlyLoadedPlugin(PluginName, &FailureReason);
+				UE_CLOG(!bResult, LogPlatformDLC, Warning, TEXT("Failed to unmount plugin %s : %s"), *PluginName, *FailureReason.ToString());
+
+				IPluginManager::Get().RemoveFromPluginsList( Package->DLCPluginFilePath );
 			}
 			else
 			{
@@ -643,6 +681,8 @@ bool FGDKPlatformDLC::Unmount( FName DLCName )
 				IPluginManager::Get().OnPluginUnmounted().Broadcast(*Package->DLCPlugin);
 				FPackageName::UnRegisterMountPoint( Package->DLCPlugin->GetMountedAssetPath(), Package->DLCPlugin->GetContentDir() );
 			}
+			Package->DLCPlugin.Reset();
+			Package->DLCPluginFilePath.Reset();
 		}
 
 		// unmount all pak files
@@ -975,6 +1015,7 @@ void FGDKPlatformDLC::FPackageData::OnUninstalled()
 	PackageIdentifier[0] = '\0';
 	MountPath.Reset();
 	DLCPlugin.Reset();
+	DLCPluginFilePath.Reset();
 	CachedInstalledBytes = 0; // not clearing CachedTotalBytes too because that will not change on a subsequent re-download
 	MountState = EMountState::NotMounted;
 	
diff --git a/Engine/Plugins/Experimental/Microsoft/GDKPlatformDLC/Source/GDKPlatformDLC/Private/GDKPlatformDLC.h b/Engine/Plugins/Experimental/Microsoft/GDKPlatformDLC/Source/GDKPlatformDLC/Private/GDKPlatformDLC.h
--- a/Engine/Plugins/Experimental/Microsoft/GDKPlatformDLC/Source/GDKPlatformDLC/Private/GDKPlatformDLC.h
+++ b/Engine/Plugins/Experimental/Microsoft/GDKPlatformDLC/Source/GDKPlatformDLC/Private/GDKPlatformDLC.h
@@ -10,6 +10,7 @@
 #include "Containers/Array.h"
 #include "Containers/Ticker.h"
 #include "Math/Color.h"
+#include "Misc/Optional.h"
 
 #include "Microsoft/AllowMicrosoftPlatformTypes.h"
 THIRD_PARTY_INCLUDES_START
@@ -109,6 +110,7 @@ private:
 		XStoreLicenseHandle               DLCLicense = nullptr;
 		XTaskQueueRegistrationToken       PackageLicenseLostToken = {};
 		TSharedPtr<class IPlugin>         DLCPlugin;
+		FString                           DLCPluginFilePath;
 		XPackageInstallationMonitorHandle InstallationMonitorHandle = nullptr;
 		XTaskQueueRegistrationToken       InstallationProgressToken = {};
 		std::atomic<uint64>               CachedInstalledBytes = 0;
@@ -123,6 +125,7 @@ private:
 	};
 	TArray<TSharedPtr<FPackageData>> AllPackages;
 	TMap<FName, FString>             DLCNameToStoreIdMap;
+	TMap<FName, FString>             DLCNameToPluginNameMap;
 	mutable FCriticalSection         PackageDataLock;
 	XTaskQueueRegistrationToken      PackageInstalledToken;
 	std::atomic<bool>                bIsInitialized;
diff --git a/Engine/Source/Developer/Microsoft/GDKPlatformEditor/Private/GDKTargetSettings.cpp b/Engine/Source/Developer/Microsoft/GDKPlatformEditor/Private/GDKTargetSettings.cpp
--- a/Engine/Source/Developer/Microsoft/GDKPlatformEditor/Private/GDKTargetSettings.cpp
+++ b/Engine/Source/Developer/Microsoft/GDKPlatformEditor/Private/GDKTargetSettings.cpp
@@ -34,13 +34,14 @@ TArray<FString> UGDKTargetSettings::GetKnownCultureStageIds()
 	return GetDefault<UProjectPackagingSettings>()->CulturesToStage;	
 }
 
-TArray<FString> UGDKTargetSettings::GetKnownDLCNames()
+TArray<FString> UGDKTargetSettings::GetKnownDLCPlugins()
 {
 	TArray<FString> KnownDLC;
 
-	for (TSharedRef<IPlugin> Plugin : IPluginManager::Get().GetEnabledPluginsWithContent())
+	for (TSharedRef<IPlugin> Plugin : IPluginManager::Get().GetDiscoveredPlugins())
 	{
-		if (Plugin->GetType() == EPluginType::Project && Plugin->GetLoadedFrom() == EPluginLoadedFrom::Project && Plugin->GetDescriptor().SupportsTargetPlatform(GetPlatformName()) )
+		// note: allowing Disabled plugins here too because typically the DLC plugin needs to be disabled to prevent it from being included in the main game package
+		if (Plugin->CanContainContent() && Plugin->GetType() == EPluginType::Project && Plugin->GetLoadedFrom() == EPluginLoadedFrom::Project && Plugin->GetDescriptor().bExplicitlyLoaded && Plugin->GetDescriptor().SupportsTargetPlatform(GetPlatformName()) )
 		{
 			FString DLCName = Plugin->GetName();
 			KnownDLC.Add(DLCName);
@@ -50,6 +51,19 @@ TArray<FString> UGDKTargetSettings::GetKnownDLCNames()
 	return KnownDLC;
 }
 
+TArray<FString> UGDKTargetSettings::GetKnownDLCNames()
+{
+	TArray<FString> DLCNames;
+
+	for (const FGDKDLCPackage& DLCPackage : DLCPackages)
+	{
+		DLCNames.Add(DLCPackage.DLCName);
+	}
+
+	return DLCNames;
+}
+
+
 bool UGDKTargetSettings::CanEditProductId() const
 {
 	return StoreId.IsEmpty();
diff --git a/Engine/Source/Developer/Microsoft/GDKPlatformEditor/Public/GDKTargetSettings.h b/Engine/Source/Developer/Microsoft/GDKPlatformEditor/Public/GDKTargetSettings.h
--- a/Engine/Source/Developer/Microsoft/GDKPlatformEditor/Public/GDKTargetSettings.h
+++ b/Engine/Source/Developer/Microsoft/GDKPlatformEditor/Public/GDKTargetSettings.h
@@ -195,12 +195,12 @@ struct FGDKDLCPackageStringResources
 
 
 USTRUCT()
-struct FGDKDLCConfiguration
+struct FGDKDLCPackage
 {
 	GENERATED_USTRUCT_BODY()
 
-	/* Name of the explicity-loaded DLC plugin this configuration applies to */
-	UPROPERTY(EditAnywhere, Config, Category=DLC, Meta=(GetOptions="GetKnownDLCNames", DisplayName="Plugin"))
+	/* Name of this DLC. Used by the various DLC support APIs */
+	UPROPERTY(EditAnywhere, Config, Category=DLC, DisplayName="DLC Name")
 	FString DLCName;
 
 	/* Store Id. Must match Partner Center configuration. */
@@ -229,6 +229,34 @@ struct FGDKDLCConfiguration
 	TArray<FGDKDLCPackageStringResources> PerCultureResources;
 };
 
+USTRUCT()
+struct FGDKDLCChunk
+{
+	GENERATED_USTRUCT_BODY()
+
+	/* Name of the DLC to associate with this chunk */
+	UPROPERTY(EditAnywhere, Config, Category=DLC, DisplayName="DLC Name", Meta=(GetOptions="GetKnownDLCNames"))
+	FString DLCName;
+
+	/* Unreal chunk identifier that will be moved into this DLC */
+	UPROPERTY(EditAnywhere, Config, Category=DLC, Meta=(ClampMin=1))
+	int32 ChunkId = 1;
+};
+
+USTRUCT()
+struct FGDKDLCPlugins
+{
+	GENERATED_USTRUCT_BODY()
+
+	/* Name of the DLC to associate with this plugin */
+	UPROPERTY(EditAnywhere, Config, Category=DLC, DisplayName="DLC Name", Meta=(GetOptions="GetKnownDLCNames"))
+	FString DLCName;
+
+	/* Name of the explicity-loaded content plugin this DLC applies to */
+	UPROPERTY(EditAnywhere, Config, Category=DLC, Meta=(GetOptions="GetKnownDLCPlugins", DisplayName="Plugin"))
+	FString Plugin;
+};
+
 
 USTRUCT()
 struct FGDKUriActivationProtocol
@@ -428,13 +456,30 @@ public:
 	UFUNCTION()
 	UE_API TArray<FString> GetKnownCultureStageIds();
 
-	// UI helper function to get the list of known DLC plugins
+	// UI helper function to get the list of DLC packages
 	UFUNCTION()
 	UE_API TArray<FString> GetKnownDLCNames();
+
+	// UI helper function to get the list of known DLC plugins
+	UFUNCTION()
+	UE_API TArray<FString> GetKnownDLCPlugins();
 	
-	/* Configuration for DLC. */
-	UPROPERTY(EditAnywhere, Config, Category=DLC, Meta=(TitleProperty="DLCName", DisplayName="DLC Plugins (EXPERIMENTAL)"))
-	TArray<FGDKDLCConfiguration> DLCPlugins;
+
+	/* Configuration for DLC packages. */
+	UPROPERTY(EditAnywhere, Config, Category=DLC, Meta=(TitleProperty="DLCName", DisplayName="DLC Packages (EXPERIMENTAL)"))
+	TArray<FGDKDLCPackage> DLCPackages;
+
+	/* When enabled, the DLC packages listed above will be generated automatically using `DLC Chunks` during normal packaging. This is not compatible with the "DLC Plugins" mechanism or packaging a single plugin in UAT via -dlcname= */
+	UPROPERTY(EditAnywhere, Config, Category=DLC, Meta=(DisplayName="Automatically Create DLC Packages (EXPERIMENTAL)"))
+	bool bAutoCreateDLCPackages = false;
+
+	/* Chunks to pull out of the main package to include in DLC instead */
+	UPROPERTY(EditAnywhere, Config, Category=DLC, Meta=(TitleProperty="ChunkId", DisplayName="DLC Chunks (EXPERIMENTAL)", EditCondition="bAutoCreateDLCPackages", EditConditionHides="true"))
+	TArray<FGDKDLCChunk> DLCChunks;
+
+	/* Configuration for DLC plugins */
+	UPROPERTY(EditAnywhere, Config, Category=DLC, Meta=(TitleProperty="DLCName", DisplayName="DLC Plugins (EXPERIMENTAL)", EditCondition="!bAutoCreateDLCPackages", EditConditionHides="true"))
+	TArray<FGDKDLCPlugins> DLCPlugins;
 };
 
 #undef UE_API
diff --git a/Engine/Source/Programs/AutomationTool/Microsoft/GDKPackageGenerator.cs b/Engine/Source/Programs/AutomationTool/Microsoft/GDKPackageGenerator.cs
--- a/Engine/Source/Programs/AutomationTool/Microsoft/GDKPackageGenerator.cs
+++ b/Engine/Source/Programs/AutomationTool/Microsoft/GDKPackageGenerator.cs
@@ -43,7 +43,7 @@ namespace AutomationUtils.GDK
 		public bool NoDebugInfoSkipsSymbolBundling { get; protected set; } = false;
 
 
-	protected static int GDKEdition => GDKExports.GetGDKVersionNumber() ?? 0;
+		protected static int GDKEdition => GDKExports.GetGDKVersionNumber() ?? 0;
 		protected static string GDKRootDir => Utils.CleanDirectorySeparators(Path.Combine(GDKExports.GetGSDKRoot()));
 		protected static string GDKBinariesDir => Path.Combine( GDKRootDir, "bin");
 
@@ -92,7 +92,7 @@ namespace AutomationUtils.GDK
 
 		public abstract string GetPlatformPackageFileExtension();
 
-		protected abstract string GetPlatformSpecificPackageCommandlineOptions(ProjectParams Params, DeploymentContext SC);
+		protected abstract string GetPlatformSpecificPackageCommandlineOptions(ProjectParams Params, DeploymentContext SC, string DLCPackage);
 
 		protected virtual void GetPlatformSpecificPackageSymbolPaths(HashSet<string> SymbolPaths, ProjectParams Params, DeploymentContext SC)
 		{
@@ -205,7 +205,7 @@ namespace AutomationUtils.GDK
 					CaptureLogger CaptureLogger = new();
 					try
 					{
-						GenerateSinglePackage(Params, SC, PackageMetadataBasePath, Request.Config, CaptureLogger, PackageOutputPathFragment);
+						GenerateSinglePackage(Params, SC, Request.DLCPackage, PackageMetadataBasePath, Request.Config, CaptureLogger, PackageOutputPathFragment);
 					}
 					catch (AutomationException e)
 					{
@@ -226,7 +226,7 @@ namespace AutomationUtils.GDK
 			{
 				foreach (PackageRequest Request in PackageRequests)
 				{
-					GenerateSinglePackage(Params, SC, PackageMetadataBasePath, Request.Config, null, PackageOutputPathFragment);
+					GenerateSinglePackage(Params, SC, Request.DLCPackage, PackageMetadataBasePath, Request.Config, null, PackageOutputPathFragment);
 				}
 			}
 		}
@@ -236,7 +236,7 @@ namespace AutomationUtils.GDK
 			List<PackageRequest> PackageRequests = MakePackageRequests(Params, SC);
 			foreach (PackageRequest Request in PackageRequests)
 			{
-				GenerateSinglePackageMetadata(Params, SC, PackageMetadataBasePath, Request.Config);
+				GenerateSinglePackageMetadata(Params, SC, PackageMetadataBasePath, Request.DLCPackage, Request.Config);
 			}
 		}
 
@@ -252,8 +252,13 @@ namespace AutomationUtils.GDK
 			{
 				Config = InConfig;
 			}
+			public PackageRequest(string InDLCPackage)
+			{
+				DLCPackage = InDLCPackage;
+			}
 
 			public readonly UnrealTargetConfiguration? Config = null; // null means "no specific single config - package all configurations"
+			public readonly string DLCPackage = null;
 		};
 
 		private List<PackageRequest> MakePackageRequests(ProjectParams Params, DeploymentContext SC)
@@ -270,12 +275,21 @@ namespace AutomationUtils.GDK
 				PackageRequests.Add(new PackageRequest());
 			}
 
+			// add any DLC packages we require
+			if (EngineIni.GetBool(IniSection_TargetSettings, "bAutoCreateDLCPackages", out bool bAutoCreateDLCPackages) && bAutoCreateDLCPackages)
+			{
+				foreach (string DLCPackage in DLCChunks.GetFromConfig(Params, SC, IniSection_TargetSettings).Keys)
+				{
+					PackageRequests.Add(new PackageRequest(DLCPackage));
+				}
+			}
+
 			return PackageRequests;
 		}
 
-		public void GenerateSinglePackage(ProjectParams Params, DeploymentContext SC, DirectoryReference PackageMetadataBasePath = null, UnrealTargetConfiguration? OnlyConfig = null, ILogger CustomLogger = null, string PackageOutputPathFragment = "")
+		public void GenerateSinglePackage(ProjectParams Params, DeploymentContext SC, string DLCPackage, DirectoryReference PackageMetadataBasePath = null, UnrealTargetConfiguration? OnlyConfig = null, ILogger CustomLogger = null, string PackageOutputPathFragment = "")
 		{
-			string PackageConfig = OnlyConfig?.ToString() ?? "Combined";
+			string PackageConfig = (DLCPackage == null) ? (OnlyConfig?.ToString() ?? "Combined") : DLCPackage;
 
 			// default to global logger unless one has been specified & output banner
 			CustomLogger ??= Logger;
@@ -292,7 +306,7 @@ namespace AutomationUtils.GDK
 			// fall back to legacy packaging if there is no metadata
 			if (!PrecachePackageLayout || Params.PreStaged)
 			{
-				PackageImmediate(Params, SC, PackageOutputDir, CustomLogger, OnlyConfig);
+				PackageImmediate(Params, SC, PackageOutputDir, CustomLogger, DLCPackage, OnlyConfig);
 				return;
 			}
 
@@ -306,13 +320,13 @@ namespace AutomationUtils.GDK
 			}
 
 			// attempt to package
-			PackageInternal(Params, SC, PackageOutputDir, LayoutXmlFile, PackageMetadataDir, CustomLogger, OnlyConfig);
+			PackageInternal(Params, SC, PackageOutputDir, LayoutXmlFile, PackageMetadataDir, CustomLogger, DLCPackage, OnlyConfig);
 		}
 
 
 
 
-		protected virtual void PackageInternal(ProjectParams Params, DeploymentContext SC, DirectoryReference PackagePath, FileReference LayoutPath, DirectoryReference SourceFolder, ILogger Logger, UnrealTargetConfiguration? OnlyConfig = null)
+		protected virtual void PackageInternal(ProjectParams Params, DeploymentContext SC, DirectoryReference PackagePath, FileReference LayoutPath, DirectoryReference SourceFolder, ILogger Logger, string DLCPackage, UnrealTargetConfiguration? OnlyConfig = null)
 		{
 			// sanity check
 			if (SC.StageDirectory == null)
@@ -357,7 +371,7 @@ namespace AutomationUtils.GDK
 			DirectoryReference.CreateDirectory(PackagePath);
 
 			// prepare the makepkg command line
-			string AdditionalOptions = GetPlatformSpecificPackageCommandlineOptions(Params, SC) + " " + GetPackagingEncryptionOptions(Params);
+			string AdditionalOptions = GetPlatformSpecificPackageCommandlineOptions(Params, SC, DLCPackage) + " " + GetPackagingEncryptionOptions(Params);
 			string PackageCommandline = string.Format("pack /f \"{0}\" {1} /d \"{2}\" /pd \"{3}\" /loggable", LayoutPath, AdditionalOptions, SourceFolder, PackagePath);
 
 			// add path to shared submission validator, if necessary
@@ -367,7 +381,7 @@ namespace AutomationUtils.GDK
 			}
 
 			// read store Id: we only include the legacy ProductId if the store Id is not specified
-			if (!Params.HasDLCName)
+			if (Params.DLCFile == null && DLCPackage == null)
 			{
 				bool bHasStoreId = (EngineIni.GetString(IniSection_TargetSettings, "StoreId", out string StoreId) && !String.IsNullOrEmpty(StoreId));
 				if (!bHasStoreId)
@@ -409,9 +423,14 @@ namespace AutomationUtils.GDK
 
 			// Look for a release package to use as base package for patch generation
 			string ReleasePackagePath = null;
-			if (Params.HasBasedOnReleaseVersion && !Params.HasDLCName)
+			if (Params.HasBasedOnReleaseVersion && Params.DLCFile == null)
 			{
 				DirectoryReference ReleaseVersionBasePath = new DirectoryReference(Params.GetBasedOnReleaseVersionPath(SC, Params.Client));
+				if (DLCPackage != null)
+				{
+					ReleaseVersionBasePath = DirectoryReference.Combine(ReleaseVersionBasePath, DLCPackage);
+				}
+
 				DirectoryReference[] ReleaseVersionPaths = new DirectoryReference[]
 				{
 					DirectoryReference.Combine(ReleaseVersionBasePath, "LatestPatch"), // legacy support. check this first as the package search is recursive.
@@ -553,7 +572,7 @@ namespace AutomationUtils.GDK
 
 
 
-		protected void PackageImmediate(ProjectParams Params, DeploymentContext SC, DirectoryReference PackagePath, ILogger Logger, UnrealTargetConfiguration? OnlyConfig = null)
+		protected void PackageImmediate(ProjectParams Params, DeploymentContext SC, DirectoryReference PackagePath, ILogger Logger, string DLCPackage, UnrealTargetConfiguration? OnlyConfig = null)
 		{
 			// prepare intermediate packaging directory
 			DirectoryReference IntermediatePath = DirectoryReference.Combine(SC.ProjectRoot, "Intermediate", PlatformType.ToString(), "Packaging");
@@ -572,16 +591,16 @@ namespace AutomationUtils.GDK
 
 			// generate a GDK manifest that just contains the selected configurations
 			List<AppXManifestExecutable> Executables = GDKAutomationUtils.GetGameConfigExecutables(Params, SC, GDKAutomationUtils.GameConfigPurpose.ForPackaging, IniSection_TargetSettings, OnlyConfig);
-			GDKGameConfigGenerator.Create(PlatformType, Logger, IntermediatePath, SC.ShortProjectName, Params.RawProjectPath, Executables, bDeleteIntermediates: true, DLCFile: Params.DLCFile, CustomConfig: SC.CustomConfig, NoBootstrapExe: Params.NoBootstrapExe);
+			GDKGameConfigGenerator.Create(PlatformType, Logger, IntermediatePath, SC.ShortProjectName, Params.RawProjectPath, Executables, bDeleteIntermediates: true, DLCFile: Params.DLCFile, DLCPackage: DLCPackage, CustomConfig: SC.CustomConfig, NoBootstrapExe: Params.NoBootstrapExe);
 			IEnumerable<StagedFileReference> GeneratedFilesToInclude =  GetGDKManifestFilesInDirectory(IntermediatePath)
 				.Select( File => StagedFileReference.Combine(File.MakeRelativeTo(IntermediatePath)));
 
 			// layout.xml & package.manifest generation
-			GeneratePackageMetadataInternal(Params, SC, FilesToInclude, GeneratedFilesToInclude, IntermediatePath, Logger);
+			GeneratePackageMetadataInternal(Params, SC, FilesToInclude, GeneratedFilesToInclude, IntermediatePath, Logger, DLCPackage);
 
 			// attempt to package
 			FileReference LayoutXmlFile = FileReference.Combine(IntermediatePath, "layout.xml");
-			PackageInternal(Params, SC, PackagePath, LayoutXmlFile, IntermediatePath, Logger, OnlyConfig);
+			PackageInternal(Params, SC, PackagePath, LayoutXmlFile, IntermediatePath, Logger, DLCPackage, OnlyConfig);
 		}
 
 
@@ -1030,6 +1049,47 @@ namespace AutomationUtils.GDK
 
 
 
+		protected struct DLCChunks
+		{
+			// these must match the the GDKTargetSettings field
+			public string DLCName;
+			public int ChunkId;
+
+			internal static Dictionary<string,IEnumerable<int>> GetFromConfig(ProjectParams Params, DeploymentContext SC, string IniSection_TargetSettings)
+			{
+				ConfigHierarchy EngineIni = ConfigCache.ReadHierarchy(ConfigHierarchyType.Engine, DirectoryReference.FromFile(SC.RawProjectPath), SC.StageTargetPlatform.PlatformType, SC.CustomConfig);
+
+				if (EngineIni.TryGetValuesGeneric(IniSection_TargetSettings, "DLCChunks", out DLCChunks[] DLCChunks) && DLCChunks != null)
+				{
+					foreach (DLCChunks DLCChunk in DLCChunks)
+					{
+						if (string.IsNullOrEmpty(DLCChunk.DLCName))
+						{
+							throw new AutomationException("Missing DLCName in DLCChunks item", DLCChunk.DLCName);
+						}
+
+						if (DLCChunk.DLCName.Equals("Combined", StringComparison.InvariantCultureIgnoreCase) || Enum.TryParse<UnrealTargetConfiguration>(DLCChunk.DLCName, out UnrealTargetConfiguration _))
+						{
+							throw new AutomationException("Invalid DLC name {0} - reserved for use internally", DLCChunk.DLCName);
+						}
+
+						if (DLCChunk.ChunkId <= 0)
+						{
+							throw new AutomationException("Invalid or missing DLC chunk id {0} for {1}... must be > 0", DLCChunk.ChunkId, DLCChunk.DLCName);
+						}
+					}
+
+					return DLCChunks.GroupBy( C => C.DLCName)
+						.ToDictionary(
+							G => G.Key,
+							G => G.Select(C => C.ChunkId)
+						);
+				}
+
+				return [];
+			}
+		}
+
 
 
 
@@ -1041,12 +1101,16 @@ namespace AutomationUtils.GDK
 			public string RelativeSourceDir;
 			public int ChunkId;
 			public bool bAllowInManifest;
+			public bool bAllowInLayout;
+			public bool bAutoGenerated;
 		};
 
 
 
-		private void GeneratePackageMetadataInternal(ProjectParams Params, DeploymentContext SC, IEnumerable<StagedFileReference> FilesToInclude, IEnumerable<StagedFileReference> GeneratedFilesToInclude, DirectoryReference LayoutDirectory, ILogger Logger)
+		private void GeneratePackageMetadataInternal(ProjectParams Params, DeploymentContext SC, IEnumerable<StagedFileReference> FilesToInclude, IEnumerable<StagedFileReference> GeneratedFilesToInclude, DirectoryReference LayoutDirectory, ILogger Logger, string DLCPackage)
 		{
+			ConfigHierarchy EngineIni = ConfigCache.ReadHierarchy(ConfigHierarchyType.Engine, DirectoryReference.FromFile(SC.RawProjectPath), SC.StageTargetPlatform.PlatformType, SC.CustomConfig);
+
 			IReadOnlyDictionary<string, GDKBundleSettings> InstallBundles = null;
 			if (BundleUtils.HasPlatformBundleSource(SC.RawProjectPath.Directory, SC.StageTargetPlatform.PlatformType, SC.CustomConfig))
 			{
@@ -1074,12 +1138,46 @@ namespace AutomationUtils.GDK
 				});
 			}
 
-			var IntelligentInstallChunks = IntelligentInstallChunk.GetFromConfig(Params, SC, IniSection_TargetSettings).ToList();
-			var IntelligentInstallFeatures = IntelligentInstallFeature.GetFromConfig(Params, SC, IniSection_TargetSettings).ToList();
-			var IntelligentInstallRecipes = IntelligentInstallRecipe.GetFromConfig(Params, SC, IniSection_TargetSettings).ToList();
+			// if we are creating DLC packages automatically, we need to make sure that the correct chunks end up in the correct packages
+			IEnumerable<int> AllowedChunkIds = null;
+			IEnumerable<int> DisallowedChunkIds = null;
+			EngineIni.GetBool(IniSection_TargetSettings, "bAutoCreateDLCPackages", out bool bAutoCreateDLCPackages);
+			bool bIsDLC = false;
+			if (bAutoCreateDLCPackages)
+			{
+				Dictionary<string, IEnumerable<int>> DLCChunkMap = DLCChunks.GetFromConfig(Params, SC, IniSection_TargetSettings);
+				if (DLCPackage == null)
+				{
+					DisallowedChunkIds = DLCChunkMap.SelectMany( X => X.Value ).Distinct();
+				}
+				else
+				{
+					DLCChunkMap.TryGetValue(DLCPackage, out AllowedChunkIds);
+					bIsDLC = true;
+				}
+			}
+			bool CanPackageChunk( int ChunkID )
+			{
+				if (DisallowedChunkIds != null && DisallowedChunkIds.Contains(ChunkID))
+				{
+					return false;
+				}
+
+				if (AllowedChunkIds != null && !AllowedChunkIds.Contains(ChunkID))
+				{
+					return false;
+				}
+
+				return true;
+			}
+
+			// read intelligent install data, unless we are making DLC
+			var IntelligentInstallChunks   = bIsDLC ? [] : IntelligentInstallChunk.GetFromConfig(Params, SC, IniSection_TargetSettings).ToList();
+			var IntelligentInstallFeatures = bIsDLC ? [] : IntelligentInstallFeature.GetFromConfig(Params, SC, IniSection_TargetSettings).ToList();
+			var IntelligentInstallRecipes  = bIsDLC ? [] : IntelligentInstallRecipe.GetFromConfig(Params, SC, IniSection_TargetSettings).ToList();
 
-			ConfigHierarchy EngineIni = ConfigCache.ReadHierarchy(ConfigHierarchyType.Engine, DirectoryReference.FromFile(SC.RawProjectPath), SC.StageTargetPlatform.PlatformType, SC.CustomConfig);
 			EngineIni.GetBool(IniSection_TargetSettings, "bUseFeaturesAndRecipes", out bool bUseFeaturesAndRecipes);
+			bUseFeaturesAndRecipes &= !bIsDLC;
 
 
 
@@ -1108,6 +1206,8 @@ namespace AutomationUtils.GDK
 						RelativeSourceDir = Path.Combine(RelativeLayoutDir, File.Directory.Name),
 						ChunkId = ChunkId,
 						bAllowInManifest = true,
+						bAllowInLayout = true,
+						bAutoGenerated = false,
 					});
 				}
 
@@ -1122,6 +1222,8 @@ namespace AutomationUtils.GDK
 							RelativeSourceDir = "./" + File.Directory.Name,
 							ChunkId = 0,
 							bAllowInManifest = false,
+							bAllowInLayout = true,
+							bAutoGenerated = true,
 						});
 					}
 				}
@@ -1135,6 +1237,8 @@ namespace AutomationUtils.GDK
 						RelativeSourceDir = "./" + StagedManifestFile.Directory.Name,
 						ChunkId = 0,
 						bAllowInManifest = false,
+						bAllowInLayout = true,
+						bAutoGenerated = true,
 					});
 				}
 			}
@@ -1152,8 +1256,13 @@ namespace AutomationUtils.GDK
 					GDKBundleSettings Bundle = BundleUtils.MatchBundleSettings(Filename, InstallBundles);
 					if (Bundle != null)
 					{
-						Logger.LogInformation("Mapping {Filename} to Chunk {ChunkId} for Bundle {BundleName}", Filename, Bundle.ChunkID, Bundle.Name);
 						Pair.Value.ChunkId = Bundle.ChunkID;
+						if (!CanPackageChunk(Bundle.ChunkID))
+						{
+							continue;
+						}
+
+						Logger.LogInformation("Mapping {Filename} to Chunk {ChunkId} for Bundle {BundleName} {DLCPackage}", Filename, Bundle.ChunkID, Bundle.Name, DLCPackage ?? "");
 
 						if (Bundle.ChunkID != -1 && Bundle.PlatformChunkName != "")
 						{
@@ -1249,9 +1358,20 @@ namespace AutomationUtils.GDK
 				}
 			}
 
+			// remove anything that's not autogenerated from the layout if it's for a different package
+			foreach (var Pair in FileChunkMapping)
+			{
+				if (!Pair.Value.bAutoGenerated && !CanPackageChunk(Pair.Value.ChunkId))
+				{
+					Pair.Value.bAllowInManifest = false;
+					Pair.Value.bAllowInLayout = false;
+				}
+			}
+
+
 			// Generate a contiguous list of chunks up to the maximum ID listed in either the FileChunkMappings or IntelligentInstallChunks
 			int MaxChunkId = (new[] { 0 })
-				.Concat(FileChunkMapping.Values.Select(c => c.ChunkId))
+				.Concat(FileChunkMapping.Values.Where(c => c.bAllowInLayout).Select(c => c.ChunkId))
 				.Concat(IntelligentInstallChunks.Select(c => c.ChunkId - OSChunkIdOffset))
 				.Where(c => c >= 0)
 				.Distinct()
@@ -1279,12 +1399,19 @@ namespace AutomationUtils.GDK
 			InitialChunkIds.Sort();
 			if (InitialChunkIds.Count == 0)
 			{
-				int MinChunkId = AllChunkIds.Min();
-				if (IntelligentInstallChunks.Count > 0)
+				if (bIsDLC)
+				{
+					InitialChunkIds.AddRange(AllChunkIds);
+				}
+				else
 				{
-					Logger.LogTrace($"No initial chunk ID specified. Assuming chunk {MinChunkId} is the initial chunk.");
+					int MinChunkId = AllChunkIds.Min();
+					if (IntelligentInstallChunks.Count > 0)
+					{
+						Logger.LogTrace($"No initial chunk ID specified. Assuming chunk {MinChunkId} is the initial chunk.");
+					}
+					InitialChunkIds.Add(MinChunkId);
 				}
-				InitialChunkIds.Add(MinChunkId);
 			}
 
 			// Gather the tags that can be used in the layout.xml
@@ -1352,7 +1479,7 @@ namespace AutomationUtils.GDK
 
 				var StagedFilesByChunk =
 					(from File in FileChunkMapping
-					 where File.Value.ChunkId >= 0
+					 where File.Value.bAllowInLayout && File.Value.ChunkId >= 0
 					 orderby File.Value.ChunkId ascending
 					 orderby InitialChunkIds.Contains(File.Value.ChunkId) descending // launch chunks first
 					 group File by File.Value.ChunkId)
@@ -1404,9 +1531,9 @@ namespace AutomationUtils.GDK
 		/// <summary>
 		/// Generates the layout.xml and associated files required for packaging in the given layout directory
 		/// </summary>
-		public void GenerateSinglePackageMetadata(ProjectParams Params, DeploymentContext SC, DirectoryReference PackageMetadataDir, UnrealTargetConfiguration? OnlyConfig = null)
+		public void GenerateSinglePackageMetadata(ProjectParams Params, DeploymentContext SC, DirectoryReference PackageMetadataDir, string DLCPackage, UnrealTargetConfiguration? OnlyConfig = null)
 		{
-			string PackageConfig = (OnlyConfig?.ToString() ?? "Combined");
+			string PackageConfig = (DLCPackage == null) ? (OnlyConfig?.ToString() ?? "Combined") : DLCPackage;
 			PackageMetadataDir = DirectoryReference.Combine(PackageMetadataDir, PackageConfig);
 
 			FileUtils.ForceDeleteDirectory(PackageMetadataDir);
@@ -1473,7 +1600,7 @@ namespace AutomationUtils.GDK
 
 				// generate the manifest and gather the files
 				List<AppXManifestExecutable> Executables = GDKAutomationUtils.GetGameConfigExecutables(Params, SC, Purpose: GDKAutomationUtils.GameConfigPurpose.ForStaging, IniSection_TargetSettings, OnlyConfig);
-				GDKGameConfigGenerator.Create(PlatformType, Logger, PackageMetadataDir, SC.ShortProjectName, Params.RawProjectPath, Executables, bDeleteIntermediates: true, DLCFile: Params.DLCFile, CustomConfig: SC.CustomConfig);
+				GDKGameConfigGenerator.Create(PlatformType, Logger, PackageMetadataDir, SC.ShortProjectName, Params.RawProjectPath, Executables, bDeleteIntermediates: true, DLCFile: Params.DLCFile, DLCPackage: DLCPackage, CustomConfig: SC.CustomConfig);
 				GeneratedFilesToInclude = GeneratedFilesToInclude.Union( GetGDKManifestFilesInDirectory(PackageMetadataDir)
 					.Select(File => StagedFileReference.Combine(File.MakeRelativeTo(PackageMetadataDir)))
 				);
@@ -1484,7 +1611,7 @@ namespace AutomationUtils.GDK
 
 
 			// Generate the packaging metadata
-			GeneratePackageMetadataInternal(Params, SC, FilesToInclude, GeneratedFilesToInclude, PackageMetadataDir, Logger);
+			GeneratePackageMetadataInternal(Params, SC, FilesToInclude, GeneratedFilesToInclude, PackageMetadataDir, Logger, DLCPackage);
 		}
 
 		protected virtual IEnumerable<StagedFileReference> GetPlatformFilesToPackage(ProjectParams Params, DeploymentContext SC, DirectoryReference PackageMetadataDir, UnrealTargetConfiguration? OnlyConfig)
diff --git a/Engine/Source/Programs/AutomationTool/Win/MSGaming/MSGameStorePackageGenerator.cs b/Engine/Source/Programs/AutomationTool/Win/MSGaming/MSGameStorePackageGenerator.cs
--- a/Engine/Source/Programs/AutomationTool/Win/MSGaming/MSGameStorePackageGenerator.cs
+++ b/Engine/Source/Programs/AutomationTool/Win/MSGaming/MSGameStorePackageGenerator.cs
@@ -18,10 +18,10 @@ namespace AutomationUtils.GDK
 		public MSGameStorePackageGenerator(FileReference InProjectPath, UnrealTargetPlatform InPlatformType, string InIniSection_TargetSettings, string CustomConfig, UnrealTargetPlatform? InCrashReportPlatform)
 			: base(InProjectPath, InPlatformType, InIniSection_TargetSettings, CustomConfig)
 		{
-			CrashReportPlatform = InCrashReportPlatform;
+			CrashReportPlatform = InCrashReportPlatform ?? UnrealTargetPlatform.Win64;
 		}
 
-		protected override string GetPlatformSpecificPackageCommandlineOptions(ProjectParams Params, DeploymentContext SC)
+		protected override string GetPlatformSpecificPackageCommandlineOptions(ProjectParams Params, DeploymentContext SC, string DLCPackage)
 		{
 			return "/pc";
 		}
@@ -94,13 +94,13 @@ namespace AutomationUtils.GDK
 		}
 
 
-		protected override void PackageInternal(ProjectParams Params, DeploymentContext SC, DirectoryReference PackagePath, FileReference LayoutPath, DirectoryReference SourceFolder, ILogger Logger, UnrealTargetConfiguration? OnlyConfig = null)
+		protected override void PackageInternal(ProjectParams Params, DeploymentContext SC, DirectoryReference PackagePath, FileReference LayoutPath, DirectoryReference SourceFolder, ILogger Logger, string DLCPackage, UnrealTargetConfiguration? OnlyConfig = null)
 		{
-			base.PackageInternal(Params, SC, PackagePath, LayoutPath, SourceFolder, Logger, OnlyConfig);
-			GenerateInstallBatchFile(PackagePath, Params, SC);
+			base.PackageInternal(Params, SC, PackagePath, LayoutPath, SourceFolder, Logger, DLCPackage, OnlyConfig);
+			GenerateInstallBatchFile(PackagePath, Params, SC, (DLCPackage != null) );
 		}
 
-		private void GenerateInstallBatchFile(DirectoryReference OutputPath, ProjectParams Params, DeploymentContext SC)
+		private void GenerateInstallBatchFile(DirectoryReference OutputPath, ProjectParams Params, DeploymentContext SC, bool bIsDLC)
 		{
 			ConfigHierarchy EngineIni = ConfigCache.ReadHierarchy(ConfigHierarchyType.Engine, DirectoryReference.FromFile(SC.RawProjectPath), SC.StageTargetPlatform.PlatformType, SC.CustomConfig);
 
@@ -118,10 +118,10 @@ namespace AutomationUtils.GDK
 				throw new AutomationException(ExitCode.Error_MissingExecutable, "Could not find an msixvc package in {0}.", OutputPath.FullName);
 			}
 
-			// create the game config info for the package, preferring the staged manifest if we find it
+			// create the game config info for the package, preferring the staged manifest for the base game if we find it
 			GDKGameConfigInfo GameConfig;
 			FileReference ManifestFileName = FileReference.Combine(SC.StageDirectory, "MicrosoftGame.config");
-			if (FileReference.Exists(ManifestFileName))
+			if (FileReference.Exists(ManifestFileName) && !bIsDLC)
 			{
 				GameConfig = new GDKGameConfigInfo(ManifestFileName);
 			}
@@ -161,7 +161,10 @@ namespace AutomationUtils.GDK
 			// save batch files
 			File.WriteAllText(Path.Combine(OutputPath.FullName, InstallBatchFileName), InstallBatchFile);
 			File.WriteAllText(Path.Combine(OutputPath.FullName, UninstallBatchFileName), UninstallBatchFile);
-			File.WriteAllText(Path.Combine(OutputPath.FullName, RunBatchFileName), RunBatchFile);
+			if (!bIsDLC)
+			{
+				File.WriteAllText(Path.Combine(OutputPath.FullName, RunBatchFileName), RunBatchFile);
+			}
 		}
 
 
diff --git a/Engine/Source/Programs/UnrealBuildTool/Platform/Microsoft/GDKGameConfigGenerator.cs b/Engine/Source/Programs/UnrealBuildTool/Platform/Microsoft/GDKGameConfigGenerator.cs
--- a/Engine/Source/Programs/UnrealBuildTool/Platform/Microsoft/GDKGameConfigGenerator.cs
+++ b/Engine/Source/Programs/UnrealBuildTool/Platform/Microsoft/GDKGameConfigGenerator.cs
@@ -24,14 +24,17 @@ namespace UnrealBuildTool.GDK
 		/// Platform-specific architecture
 		protected virtual string DefaultArchitecture => "neutral";
 
-		/// DLC file to use, if any
+		/// DLC plugin file to use, if any
 		protected FileReference? DLCFile;
 
+		/// DLC package name (when generating DLC alongside the main package)
+		protected string? DLCPackage;
+
 		/// Cached properties for the current DLC
 		protected Dictionary<string, string>? DLCConfig;
 
 		/// Helper to check if we are dealing with DLC
-		protected bool IsDLC => (DLCFile != null);
+		protected bool IsDLC => (DLCFile != null) || (DLCPackage != null);
 
 		/// Whether we have a bootstrap exe (Windows only) - affects dependencies
 		protected bool NoBootstrapExe = false;
@@ -50,10 +53,11 @@ namespace UnrealBuildTool.GDK
 		/// <summary>
 		/// Create a manifest for the given platform and return the list of modified files
 		/// </summary>
-		public static List<string>? Create(UnrealTargetPlatform Platform, ILogger Logger, DirectoryReference InOutputPath, string? InTargetName, FileReference? InProjectFile, IEnumerable<AppXManifestExecutable> InExecutables, bool bDeleteIntermediates = false, FileReference? DLCFile = null, string CustomConfig = "", bool NoBootstrapExe = true)
+		public static List<string>? Create(UnrealTargetPlatform Platform, ILogger Logger, DirectoryReference InOutputPath, string? InTargetName, FileReference? InProjectFile, IEnumerable<AppXManifestExecutable> InExecutables, bool bDeleteIntermediates = false, FileReference? DLCFile = null, string? DLCPackage = null, string CustomConfig = "", bool NoBootstrapExe = true)
 		{
 			GDKGameConfigGenerator Generator = GDKGameConfigGeneratorFactory.CreateGenerator(Platform, Logger);
 			Generator.DLCFile = DLCFile;
+			Generator.DLCPackage = DLCPackage;
 			Generator.CustomConfig = CustomConfig;
 			Generator.NoBootstrapExe = NoBootstrapExe;
 
@@ -362,8 +366,14 @@ namespace UnrealBuildTool.GDK
 					return DefaultDisplayName;
 				}
 
-				Logger.LogWarning("Default display name for {DLC} is not specified", DLCFile!.GetFileNameWithoutExtension());
-				return DLCFile!.GetFileNameWithoutAnyExtensions();
+				if (DLCFile != null)
+				{
+					Logger.LogWarning("Default display name for {DLC} is not specified", DLCFile!.GetFileNameWithoutExtension());
+					return DLCFile!.GetFileNameWithoutAnyExtensions();
+				}
+
+				Logger.LogWarning("Default display name for {DLC} is not specified", DLCPackage);
+				return DLCPackage!;
 			}
 			else
 			{
@@ -396,8 +406,15 @@ namespace UnrealBuildTool.GDK
 					return PackageName;
 				}
 
-				Logger.LogWarning("Package name for {DLC} is not specified", DLCFile!.GetFileNameWithoutExtension());
-				return DLCFile!.GetFileNameWithoutAnyExtensions();
+				if (DLCFile != null)
+				{
+					Logger.LogWarning("Package name for {DLC} is not specified", DLCFile!.GetFileNameWithoutExtension());
+					return DLCFile!.GetFileNameWithoutAnyExtensions();
+				}
+
+				Logger.LogWarning("Package name for {DLC} is not specified", DLCPackage);
+				return DLCPackage!;
+
 			}
 
 			return base.GetIdentityPackageName();
@@ -465,19 +482,19 @@ namespace UnrealBuildTool.GDK
 
 			// find and cache the configuration for the selected DLC
 			DLCConfig = null;
-			if (DLCFile != null)
+			if (IsDLC)
 			{
-				string DLCName = DLCFile!.GetFileNameWithoutExtension();
+				string DLCName = DLCPackage ?? DLCFile!.GetFileNameWithoutExtension();
 
-				if (EngineIni!.GetArray(IniSection_PlatformTargetSettings, "DLCPlugins", out List<string>? DLCPlugins))
+				if (EngineIni!.GetArray(IniSection_PlatformTargetSettings, "DLCPackages", out List<string>? DLCPackages))
 				{
-					foreach (string DLCPlugin in DLCPlugins)
+					foreach (string DLCPackage in DLCPackages)
 					{
-						if (ConfigHierarchy.TryParse(DLCPlugin, out Dictionary<string, string>? Properties)
+						if (ConfigHierarchy.TryParse(DLCPackage, out Dictionary<string, string>? Properties)
 							&& Properties.TryGetValue("DLCName", out string? ThisDLCName)
 							&& DLCName == ThisDLCName)
 						{
-							// we have found the configuration for this DLC plugin
+							// we have found the configuration for this DLC package
 							DLCConfig = Properties;
 							break;
 						}
@@ -651,13 +668,20 @@ namespace UnrealBuildTool.GDK
 		/// </summary>
 		protected override void PrepareResourceBinaryPaths()
 		{
-			// check for DLC images first
+			// check for DLC plugin images first, in the <plugin>/Build/<plaform>/Resources/ folder (and Platform extension equivalent)
 			if (DLCFile != null)
 			{
 				AppXResources!.ProjectBinaryResourceDirectories.Add(DirectoryReference.Combine(DLCFile.Directory, "Build", Platform.ToString(), BuildResourceSubPath));
 				AppXResources!.ProjectBinaryResourceDirectories.Add(DirectoryReference.Combine(DLCFile.Directory, "Platforms", Platform.ToString(), "Build", BuildResourceSubPath));
 			}
 
+			// check for DLC package images next, in <project>/Build/<plaform>/<dlc>/Resources/ folder (and Platform extension equivalent)
+			if (DLCPackage != null && ProjectFile != null)
+			{
+				AppXResources!.ProjectBinaryResourceDirectories.Add(DirectoryReference.Combine(ProjectFile.Directory, "Build", Platform.ToString(), DLCPackage, BuildResourceSubPath));
+				AppXResources!.ProjectBinaryResourceDirectories.Add(DirectoryReference.Combine(ProjectFile.Directory, "Platforms", Platform.ToString(), "Build", DLCPackage, BuildResourceSubPath));
+			}
+
 			base.PrepareResourceBinaryPaths();
 		}
 
diff --git a/Engine/Source/Programs/UnrealBuildTool/Platform/Windows/MSGaming/MSGamingGameConfigGenerator.cs b/Engine/Source/Programs/UnrealBuildTool/Platform/Windows/MSGaming/MSGamingGameConfigGenerator.cs
--- a/Engine/Source/Programs/UnrealBuildTool/Platform/Windows/MSGaming/MSGamingGameConfigGenerator.cs
+++ b/Engine/Source/Programs/UnrealBuildTool/Platform/Windows/MSGaming/MSGamingGameConfigGenerator.cs
@@ -53,13 +53,20 @@ namespace UnrealBuildTool.GDK
 		{
 			base.PrepareResourceBinaryPaths();
 
-			// check for DLC images first
+			// check for DLC plugin images first, in the <plugin>/Build/<platform>/MSGaming/Resources/ folder (and Platform extension equivalent)
 			if (DLCFile != null)
 			{
 				AppXResources!.ProjectBinaryResourceDirectories.Add(DirectoryReference.Combine(DLCFile.Directory, "Build", IniPlatformName, "MSGaming", BuildResourceSubPath));
 				AppXResources!.ProjectBinaryResourceDirectories.Add(DirectoryReference.Combine(DLCFile.Directory, "Platforms", IniPlatformName, "Build", "MSGaming", BuildResourceSubPath));
 			}
 
+			// check for DLC package images next, in <project>/Build/<platform>/MSGaming/<dlc>/Resources/ folder (and Platform extension equivalent)
+			if (DLCPackage != null && ProjectFile != null)
+			{
+				AppXResources!.ProjectBinaryResourceDirectories.Add(DirectoryReference.Combine(ProjectFile.Directory, "Build", IniPlatformName, "MSGaming", DLCPackage, BuildResourceSubPath));
+				AppXResources!.ProjectBinaryResourceDirectories.Add(DirectoryReference.Combine(ProjectFile.Directory, "Platforms", IniPlatformName, "Build", "MSGaming", DLCPackage, BuildResourceSubPath));
+			}
+
 			// allow projects to keep their resources in an MSGaming subfolder to help with project organization
 			if (ProjectFile != null)
 			{
