Remove unused references to the submitted fastJSON binary.

#Synced-CL 52096572

diff --git a/Engine/Source/Programs/AutomationTool/AutomationUtils/AutomationUtils.Automation.csproj b/Engine/Source/Programs/AutomationTool/AutomationUtils/AutomationUtils.Automation.csproj
--- a/Engine/Source/Programs/AutomationTool/AutomationUtils/AutomationUtils.Automation.csproj
+++ b/Engine/Source/Programs/AutomationTool/AutomationUtils/AutomationUtils.Automation.csproj
@@ -13,12 +13,6 @@
   <PropertyGroup Condition="'$(EpicGamesMsBuild)' == 'true'">
     <OutputPath>..\..\..\..\Binaries\DotNET\AutomationTool\AutomationUtils</OutputPath>
   </PropertyGroup>
-  <ItemGroup>
-    <Reference Include="fastJSON">
-      <HintPath>..\..\..\..\Binaries\ThirdParty\fastJSON\netstandard2.0\fastJSON.dll</HintPath>
-      <Private>False</Private>
-    </Reference>
-  </ItemGroup>
 
   <!-- The following will glob all platform extensions' AutomationUtils source
   files and include them in the build. They will also appear in VisualStudio -->
