[Gauntlet] Prevent LogParser finalizers from closing LogStreamReader file handles as is prevents multiple reads from the same stream
#Synced-CL 52050243
#Synced-CL 52050243

diff --git a/Engine/Source/Programs/AutomationTool/Gauntlet/Unreal/Utils/Gauntlet.UnrealLogParser.cs b/Engine/Source/Programs/AutomationTool/Gauntlet/Unreal/Utils/Gauntlet.UnrealLogParser.cs
--- a/Engine/Source/Programs/AutomationTool/Gauntlet/Unreal/Utils/Gauntlet.UnrealLogParser.cs
+++ b/Engine/Source/Programs/AutomationTool/Gauntlet/Unreal/Utils/Gauntlet.UnrealLogParser.cs
@@ -211,14 +211,22 @@ namespace Gauntlet
 		private bool _disposed = false;
 		~UnrealLog()
 		{
-			Dispose();
+			Dispose(false);
 		}
 
 		public void Dispose()
+		{
+			Dispose(true);
+		}
+
+		private void Dispose(bool bDisposing = true)
 		{
 			if (!_disposed)
 			{
-				_parser?.Dispose();
+				if (bDisposing)
+				{
+					_parser?.Dispose(true);
+				}
 				_disposed = true;
 			}
 		}
@@ -676,17 +684,26 @@ namespace Gauntlet
 		private bool _disposed = false;
 		~UnrealLogParser()
 		{
-			Dispose();
+			Dispose(false);
 		}
 
 		public void Dispose()
+		{
+			Dispose(true);
+		}
+
+		public void Dispose(bool bDisposing = true)
 		{
 			if (!_disposed)
 			{
-				if (_logReader is IDisposable disposable)
+				if (bDisposing)
 				{
-					disposable.Dispose();
+					if (_logReader is IDisposable disposable)
+					{
+						disposable.Dispose();
+					}
 				}
+				
 				_disposed = true;
 			}
 		}
