Increase basic timeouts for downloading/uploading *.po files to Smartling
Set explicit HttpClient timeout so it won't break the existing algorithm of downloading data of increasing timout with multiple retries.
Also, raised up the initial timeout for downloads.

#Synced-CL 52139840

diff --git a/Engine/Source/Programs/AutomationTool/SmartlingLocalization/SmartlingLocalizationProvider.cs b/Engine/Source/Programs/AutomationTool/SmartlingLocalization/SmartlingLocalizationProvider.cs
--- a/Engine/Source/Programs/AutomationTool/SmartlingLocalization/SmartlingLocalizationProvider.cs
+++ b/Engine/Source/Programs/AutomationTool/SmartlingLocalization/SmartlingLocalizationProvider.cs
@@ -77,7 +77,7 @@ namespace EpicGames.SmartlingLocalization
 				Logger.LogWarning("LocalizationBranchName is null or empty. The branch name is used to create the correct file URIs for Smartling. Wrong files may be downloaded from Smartling and files may be clobbered in Smartling. Please append a value to LocalizationName as a command line argument to fix this.");
 			}
 			Config = new SmartlingConfig();
-			Client = new HttpClient();
+			Client = new HttpClient() { Timeout = TimeSpan.FromSeconds(1300) };
 			JsonOptions = new JsonSerializerOptions();
 			JsonOptions.PropertyNameCaseInsensitive = true;
 			JsonOptions.DefaultIgnoreCondition = JsonIgnoreCondition.WhenWritingNull;
@@ -86,7 +86,7 @@ namespace EpicGames.SmartlingLocalization
 
 		// public override Task InitializeProjectWithLocalizationProvider(string ProjectName, ProjectImportExportInfo ProjectImportInfo)
 		// {
-		// 	// @TODOLocalization: Throw an exeption if the config data is null/invalid 
+		// 	// @TODOLocalization: Throw an exception if the config data is null/invalid 
 		// }
 
 		public async override Task DownloadProjectFromLocalizationProvider(string ProjectName, ProjectImportExportInfo ProjectImportInfo)
@@ -163,8 +163,8 @@ namespace EpicGames.SmartlingLocalization
 			int MaxTries = 5;
 			// 1 to account for the initial try 
 			int CurrentTries = 1;
-			// The base for which all the exponential back off will be derived from. By default HttpClient has a default timeout of 100s 
-			int InitialTimeOut = 150;
+			// The base for which all the exponential back off will be derived from.
+			int InitialTimeOut = 250;
 			int CurrentTimeOut = CurrentTries * InitialTimeOut;
 			// Measures how long it takes from sending a response to successfully getting a response 
 			Stopwatch DownloadStopWatch = new Stopwatch();
