Fix ambiguous call converting method group to string

#Synced-CL 51157358

diff --git a/Engine/Source/Programs/AutomationTool/Gauntlet/Framework/RpcFramework/Gauntlet.RpcLibrary.cs b/Engine/Source/Programs/AutomationTool/Gauntlet/Framework/RpcFramework/Gauntlet.RpcLibrary.cs
--- a/Engine/Source/Programs/AutomationTool/Gauntlet/Framework/RpcFramework/Gauntlet.RpcLibrary.cs
+++ b/Engine/Source/Programs/AutomationTool/Gauntlet/Framework/RpcFramework/Gauntlet.RpcLibrary.cs
@@ -225,7 +225,7 @@ namespace Gauntlet
 				{
 					if (!bSuppressExceptions)
 					{
-						Log.Info($"{Err.GetType} encountered during call {CallCount} of {FuncName}");
+						Log.Info($"{Err.GetType()} encountered during call {CallCount} of {FuncName}");
 						throw;
 					}
 					else
