Throw exceptions in tests with exe-only preflights

#rnx
#tests preflight
#Synced-CL 52127503

diff --git a/Engine/Source/Programs/AutomationTool/Gauntlet/Framework/Base/Gauntlet.OverlayExecutable.cs b/Engine/Source/Programs/AutomationTool/Gauntlet/Framework/Base/Gauntlet.OverlayExecutable.cs
--- a/Engine/Source/Programs/AutomationTool/Gauntlet/Framework/Base/Gauntlet.OverlayExecutable.cs
+++ b/Engine/Source/Programs/AutomationTool/Gauntlet/Framework/Base/Gauntlet.OverlayExecutable.cs
@@ -93,7 +93,14 @@ namespace Gauntlet
 
 			if (!File.Exists(LocalExecutable) && !Directory.Exists(LocalExecutable))
 			{
-				Log.Verbose("No local executable for {Platform} exists. Skipping overlay for this role", Platform);
+				if (CommandUtils.IsBuildMachine)
+				{
+					throw new AutomationException("Overlay Error: No local executable for {0} exists. Skipping overlay for this role", Platform);
+				}
+				else
+				{
+					Log.Verbose("No local executable for {Platform} exists. Skipping overlay for this role", Platform);
+				}
 				return false;
 			}
 
