Removing BgEnvironment class

It makes mocking out or trying to instantiate a BuildGraph script from code more difficult. It also serves no purpose at this time, as evident by the fact that removing it only required touching method signatures and not a single line of actual code.

diff --git a/Engine/Source/Programs/AutomationTool/BuildGraph/BgGraphBuilder.cs b/Engine/Source/Programs/AutomationTool/BuildGraph/BgGraphBuilder.cs
--- a/Engine/Source/Programs/AutomationTool/BuildGraph/BgGraphBuilder.cs
+++ b/Engine/Source/Programs/AutomationTool/BuildGraph/BgGraphBuilder.cs
@@ -21,7 +21,6 @@ namespace AutomationTool
 		/// <summary>
 		/// Callback used to instantiate the graph
 		/// </summary>
-		/// <param name="env">The graph context</param>
-		public abstract BgGraph CreateGraph(BgEnvironment env);
+		public abstract BgGraph CreateGraph();
 	}
 }
diff --git a/Engine/Source/Programs/AutomationTool/BuildGraph/BuildGraph.cs b/Engine/Source/Programs/AutomationTool/BuildGraph/BuildGraph.cs
--- a/Engine/Source/Programs/AutomationTool/BuildGraph/BuildGraph.cs
+++ b/Engine/Source/Programs/AutomationTool/BuildGraph/BuildGraph.cs
@@ -155,54 +155,6 @@ namespace AutomationTool
 		}
 	}
 
-	/// <summary>
-	/// Environment for graph evaluation
-	/// </summary>
-	public class BgEnvironment
-	{
-		/// <summary>
-		/// The stream being built
-		/// </summary>
-		public BgString Stream { get; set; }
-
-		/// <summary>
-		/// Current changelist
-		/// </summary>
-		public BgInt Change { get; }
-
-		/// <summary>
-		/// Current code changelist
-		/// </summary>
-		public BgInt CodeChange { get; }
-
-		/// <summary>
-		/// Whether the graph is being run on a build machine
-		/// </summary>
-		public BgBool IsBuildMachine { get; }
-
-		/// <summary>
-		/// The current engine version
-		/// </summary>
-		public (BgInt Major, BgInt Minor, BgInt Patch) EngineVersion { get; }
-
-		/// <summary>
-		/// Constructor
-		/// </summary>
-		/// <param name="branch">The current branch</param>
-		/// <param name="change">Changelist being built</param>
-		/// <param name="codeChange">Code changelist being built</param>
-		internal BgEnvironment(string? branch, int? change, int? codeChange)
-		{
-			Stream = branch ?? "Unknown";
-			Change = change ?? 0;
-			CodeChange = codeChange ?? 0;
-			IsBuildMachine = CommandUtils.IsBuildMachine;
-
-			ReadOnlyBuildVersion version = ReadOnlyBuildVersion.Current;
-			EngineVersion = (version.MajorVersion, version.MinorVersion, version.PatchVersion);
-		}
-	}
-
 	/// <summary>
 	/// Tool to execute build automation scripts for UE projects, which can be run locally or in parallel across a build farm (assuming synchronization and resource allocation implemented by a separate system).
 	///
@@ -457,7 +409,7 @@ namespace AutomationTool
 				}
 
 				BgGraphBuilder builder = (BgGraphBuilder)Activator.CreateInstance(builderType)!;
-				BgGraph graphSpec = builder.CreateGraph(new BgEnvironment(branch, change, codeChange));
+				BgGraph graphSpec = builder.CreateGraph();
 
 				(byte[] data, BgThunkDef[] methods) = BgCompiler.Compile(graphSpec);
 
diff --git a/Engine/Source/Programs/AutomationTool/BuildGraph/Scripts/InstalledBuild.cs b/Engine/Source/Programs/AutomationTool/BuildGraph/Scripts/InstalledBuild.cs
--- a/Engine/Source/Programs/AutomationTool/BuildGraph/Scripts/InstalledBuild.cs
+++ b/Engine/Source/Programs/AutomationTool/BuildGraph/Scripts/InstalledBuild.cs
@@ -75,7 +75,7 @@ namespace AutomationTool
 			return projects;
 		}
 
-		public override BgGraph CreateGraph(BgEnvironment context)
+		public override BgGraph CreateGraph()
 		{
 			string rootDir = CommandUtils.CmdEnv.LocalRoot;
 
