Fixup Script to download IOS builds.
Add configurability for preflight and buildgroup.

#Synced-CL 52175530

#Synced-CL 52175530

diff --git a/Engine/Source/Programs/AutomationTool/Scripts/BuildAcquire.Automation.cs b/Engine/Source/Programs/AutomationTool/Scripts/BuildAcquire.Automation.cs
--- a/Engine/Source/Programs/AutomationTool/Scripts/BuildAcquire.Automation.cs
+++ b/Engine/Source/Programs/AutomationTool/Scripts/BuildAcquire.Automation.cs
@@ -180,6 +180,9 @@ class BuildAcquire : BuildCommand, IProjectParamsHelpers
 			throw new AutomationException("No CL number specified. Use -CL=<changelist>");
 		}
 
+		IsPreflight = ParseParam("Preflight");
+		Buildgroup = ParseParamValue("Buildgroup", null);
+
 		if (!DownloadBuild && !SyncUGS && !DownloadSnapshot)
 		{
 			throw new AutomationException("No action to perform specified on the commandline.");
@@ -238,7 +241,11 @@ class BuildAcquire : BuildCommand, IProjectParamsHelpers
 		CloudStorageQueryBuilder query = new CloudStorageQueryBuilder();
 		query.Equals("type", BuildTypeDesc);
 		query.Equals("commit", Changelist);
-		query.Equals("isPreflight", "false");
+		query.Equals("isPreflight", IsPreflight ? "true" : "false");
+		if (!string.IsNullOrEmpty(Buildgroup))
+		{
+			query.Equals("buildgroup", Buildgroup);
+		}
 
 		string[] BuildPlatforms = BuildPlatform?.Split('+');
 		string[] Overlay = ParseOptionalStringParam("Overlay")?.Split('+');
@@ -517,6 +524,8 @@ class BuildAcquire : BuildCommand, IProjectParamsHelpers
 
 	string Stream;
 	string Changelist;
+	string Buildgroup = null;
+	bool IsPreflight = false;
 	FileReference ProjectPath;
 	string BuildPlatform;
 	bool DownloadServer;
