Enable screenshot captures on fail for Mobile platforms

#rnx
#tests horde
#Synced-CL 52127503

diff --git a/Engine/Source/Programs/AutomationTool/Gauntlet/Framework/Base/Gauntlet.TargetDevice.cs b/Engine/Source/Programs/AutomationTool/Gauntlet/Framework/Base/Gauntlet.TargetDevice.cs
--- a/Engine/Source/Programs/AutomationTool/Gauntlet/Framework/Base/Gauntlet.TargetDevice.cs
+++ b/Engine/Source/Programs/AutomationTool/Gauntlet/Framework/Base/Gauntlet.TargetDevice.cs
@@ -1,13 +1,14 @@
 // Copyright Epic Games, Inc. All Rights Reserved.
 
+using AutomationTool;
+using EpicGames.Core;
 using System;
 using System.Collections.Generic;
 using System.IO;
-using AutomationTool;
-using UnrealBuildTool;
-using System.Threading;
-using System.Text.RegularExpressions;
 using System.Linq;
+using System.Text.RegularExpressions;
+using System.Threading;
+using UnrealBuildTool;
 
 namespace Gauntlet
 {
@@ -128,6 +129,15 @@ namespace Gauntlet
 		/// <param name="FilesToCopy">The collection of files to copy</param>
 		void CopyAdditionalFiles(IEnumerable<UnrealFileToCopy> FilesToCopy);
 
+		/// <summary>
+		/// Take a screenshot using device soft
+		/// Filename is saved in the format of: yy-MM-dd-HH-mm-ss-ff
+		/// </summary>
+		public void TakeScreenshot(string ScreenshotName) 
+		{
+			Log.Warning(KnownLogEvents.Gauntlet_TestEvent, "Screenshot capture is not implemented");
+		}
+
 		// End new flow //
 	};
 
diff --git a/Engine/Source/Programs/AutomationTool/Gauntlet/Platform/Android/Gauntlet.TargetDeviceAndroid.cs b/Engine/Source/Programs/AutomationTool/Gauntlet/Platform/Android/Gauntlet.TargetDeviceAndroid.cs
--- a/Engine/Source/Programs/AutomationTool/Gauntlet/Platform/Android/Gauntlet.TargetDeviceAndroid.cs
+++ b/Engine/Source/Programs/AutomationTool/Gauntlet/Platform/Android/Gauntlet.TargetDeviceAndroid.cs
@@ -1346,7 +1346,7 @@ namespace Gauntlet
 
 			while (!ScreenshotTrigger.WaitOne(Period))
 			{
-				TakeScreenshot();
+				TakeScreenshot(string.Empty);
 			}
 		}
 		/// <summary>
@@ -1395,15 +1395,18 @@ namespace Gauntlet
 		/// <summary>
 		/// Takes a screenshot using adb and saves to the phone's sdcard in a temp directory. 
 		/// </summary>
-		public void TakeScreenshot()
+		public void TakeScreenshot(string ScreenshotName)
 		{
-			string FileName = $"{DateTime.Now.ToString("yy-MM-dd-HH-mm-ss-ff")}.png";
+			if (string.IsNullOrEmpty(ScreenshotName))
+			{
+				ScreenshotName = $"{DateTime.Now.ToString("yy-MM-dd-HH-mm-ss-ff")}.png";
+			}
 			try
 			{
 				RunAdbDeviceCommand("shell mkdir -p /sdcard/Screenshots/");
-				RunAdbDeviceCommand(string.Format($"shell screencap -p /sdcard/Screenshots/{FileName}"));
+				RunAdbDeviceCommand(string.Format($"shell screencap -p /sdcard/Screenshots/{ScreenshotName}"));
 
-				Log.Verbose("Saved screenshot: {ScreenshotFile}", FileName);
+				Log.Verbose("Saved screenshot: {ScreenshotFile}", ScreenshotName);
 			}
 			catch (Exception Ex)
 			{
diff --git a/Engine/Source/Programs/AutomationTool/Gauntlet/Platform/IOS/Gauntlet.IOSAppium.cs b/Engine/Source/Programs/AutomationTool/Gauntlet/Platform/IOS/Gauntlet.IOSAppium.cs
--- a/Engine/Source/Programs/AutomationTool/Gauntlet/Platform/IOS/Gauntlet.IOSAppium.cs
+++ b/Engine/Source/Programs/AutomationTool/Gauntlet/Platform/IOS/Gauntlet.IOSAppium.cs
@@ -694,8 +694,13 @@ namespace Gauntlet
 		/// Filename is saved as png in the format of: yy-MM-dd-HH-mm-ss-ff
 		/// </summary>
 		/// <param name="OutputDirectory">Folder to copy screenshots to</param>
-		public void TakeScreenshot(string OutputDirectory)
+		public void TakeScreenshot(string OutputDirectory, string ScreenshotName)
 		{
+			if (string.IsNullOrEmpty(ScreenshotName))
+			{
+				ScreenshotName = $"{DateTime.Now.ToString("yy-MM-dd-HH-mm-ss-ff")}.png";
+			}
+
 			if (string.IsNullOrEmpty(OutputDirectory))
 			{
 				throw new AutomationException("Screenshot output directory must not be null or empty!");
@@ -706,8 +711,7 @@ namespace Gauntlet
 				CreateDirectory(OutputDirectory);
 			}
 
-			string FileName = $"{DateTime.Now.ToString("yy-MM-dd-HH-mm-ss-ff")}.png";
-			string FilePath = Path.Combine(OutputDirectory, FileName);
+			string FilePath = Path.Combine(OutputDirectory, ScreenshotName);
 			try
 			{
 				(Driver as ITakesScreenshot).GetScreenshot().SaveAsFile(FilePath);
@@ -767,7 +771,7 @@ namespace Gauntlet
 				iterationCounter++;
 				if (iterationCounter >= screenshotCheckCount)
 				{
-					TakeScreenshot(ScreenshotDirectory);
+					TakeScreenshot(ScreenshotDirectory, string.Empty);
 					iterationCounter = 0;
 				}
 			}
@@ -812,7 +816,7 @@ namespace Gauntlet
 			}
 		}
 
-		private string GetScreenshotDirectory()
+		public string GetScreenshotDirectory()
 		{
 			string ScreenshotDirectory = AppiumConfig.ScreenshotDirectory;
 			if (string.IsNullOrEmpty(ScreenshotDirectory))
diff --git a/Engine/Source/Programs/AutomationTool/Gauntlet/Platform/IOS/Gauntlet.TargetDeviceIOS.cs b/Engine/Source/Programs/AutomationTool/Gauntlet/Platform/IOS/Gauntlet.TargetDeviceIOS.cs
--- a/Engine/Source/Programs/AutomationTool/Gauntlet/Platform/IOS/Gauntlet.TargetDeviceIOS.cs
+++ b/Engine/Source/Programs/AutomationTool/Gauntlet/Platform/IOS/Gauntlet.TargetDeviceIOS.cs
@@ -440,6 +440,17 @@ namespace Gauntlet
 		{
 			return Name;
 		}
+
+		public void TakeScreenshot(string ScreenshotName)
+		{
+			if (AppiumContainer is null)
+			{
+				Log.Warning(KnownLogEvents.Gauntlet_TestEvent, "Cannot perform screenshot, AppiumContainer is null");
+				return;
+			}
+
+			AppiumContainer.TakeScreenshot(AppiumContainer.GetScreenshotDirectory(), ScreenshotName);
+		}
 		#endregion
 
 		public void FullClean()
diff --git a/Engine/Source/Programs/AutomationTool/Gauntlet/Unreal/Base/Gauntlet.UnrealTestNode.cs b/Engine/Source/Programs/AutomationTool/Gauntlet/Unreal/Base/Gauntlet.UnrealTestNode.cs
--- a/Engine/Source/Programs/AutomationTool/Gauntlet/Unreal/Base/Gauntlet.UnrealTestNode.cs
+++ b/Engine/Source/Programs/AutomationTool/Gauntlet/Unreal/Base/Gauntlet.UnrealTestNode.cs
@@ -1263,6 +1263,13 @@ namespace Gauntlet
 		/// <returns></returns>
 		public override void StopTest(StopReason InReason)
 		{
+			if (Context.TestParams.ParseParam("FailurePrtScr") && (GetTestResult() != TestResult.Passed || InReason != StopReason.Completed))
+			{
+				foreach (ITargetDevice TargetDevice in UnrealApp.UnrealDeviceReservation.ReservedDevices)
+				{
+					TargetDevice.TakeScreenshot("FailureScreenshot.png");
+				}
+			}
 			base.StopTest(InReason);
 
 			// Warn if there are still deferred roles that have not been launched when the test is finished
