Add error handling for when mismatched_content_length_body_too_large

#Synced-CL 52269044

diff --git a/Engine/Source/Programs/AutomationTool/Gauntlet/Framework/RpcFramework/Gauntlet.RpcTargetRegistry.cs b/Engine/Source/Programs/AutomationTool/Gauntlet/Framework/RpcFramework/Gauntlet.RpcTargetRegistry.cs
--- a/Engine/Source/Programs/AutomationTool/Gauntlet/Framework/RpcFramework/Gauntlet.RpcTargetRegistry.cs
+++ b/Engine/Source/Programs/AutomationTool/Gauntlet/Framework/RpcFramework/Gauntlet.RpcTargetRegistry.cs
@@ -84,6 +84,10 @@ namespace Gauntlet
 						{
 							Log.Error($"Failed to call {StatusUrl} - The RPC was not identified by the target \"{TargetName}\".");
 						}
+						else if(ResponseString.Contains("mismatched_content_length_body_too_large"))
+						{
+							Log.Warning($"Failed to call {StatusUrl} (\"{TargetName}\") - Response body exceeded the declared Content-Length. The target's RPC list may have changed during serialization. This is typically transient and will resolve on the next update.");
+						}
 						else
 						{
 							Log.Error($"Failed to call {StatusUrl} (\"{TargetName}\") - Reason unknown. Please check the following Response String and update the error handler: {ResponseString}");
