PakFileRules - Warn when OverridePaks or GeneratedChunkName do not start with pakchunk<N>
Adds validation in GetPakFileRules() to warn when OverridePaks entries or GeneratedChunkName do not match the required pakchunk<N> format.

#tests Local test & preflight
#Synced-CL 52256608

diff --git a/Engine/Source/Programs/AutomationTool/Scripts/CopyBuildToStagingDirectory.Automation.cs b/Engine/Source/Programs/AutomationTool/Scripts/CopyBuildToStagingDirectory.Automation.cs
--- a/Engine/Source/Programs/AutomationTool/Scripts/CopyBuildToStagingDirectory.Automation.cs
+++ b/Engine/Source/Programs/AutomationTool/Scripts/CopyBuildToStagingDirectory.Automation.cs
@@ -2930,6 +2930,14 @@ namespace AutomationScripts
 			}
 		};
 
+		/// <summary>
+		/// Returns true if the given pak chunk name matches the required "pakchunk<N>" format.
+		/// </summary>
+		private static bool IsValidPakChunkName(string ChunkName)
+		{
+			return new Regex("^pakchunk\\d+", RegexOptions.IgnoreCase).IsMatch(ChunkName);
+		}
+
 		/// <summary>
 		/// Reads Default/BasePakFileRules.ini and returns all PakFileRules objects that apply for this deployment
 		/// </summary>
@@ -3082,10 +3090,21 @@ namespace AutomationScripts
 
 				PakRulesConfig.TryGetValue(SectionName, "bGenerateNewChunk", out PakRules.bGenerateNewChunk);
 				PakRulesConfig.TryGetValue(SectionName, "GeneratedChunkName", out PakRules.GeneratedChunkName);
+				if (!string.IsNullOrEmpty(PakRules.GeneratedChunkName) && !IsValidPakChunkName(PakRules.GeneratedChunkName))
+				{
+					Logger.LogError("Error in PakFileRules {Arg0}, GeneratedChunkName '{Arg1}' does not start with 'pakchunk<N>'.", PakRules.Name, PakRules.GeneratedChunkName);
+				}
 
 				if (PakString != null && PakString.Length > 0)
 				{
 					PakRules.OverridePaks = PakString.Split(new char[] { ',', ' ' }, StringSplitOptions.RemoveEmptyEntries).ToList();
+					foreach (string OverridePakName in PakRules.OverridePaks)
+					{
+						if (!IsValidPakChunkName(OverridePakName))
+						{
+							Logger.LogError("Error in PakFileRules {Arg0}, OverridePaks entry '{Arg1}' does not start with 'pakchunk<N>'.", PakRules.Name, OverridePakName);
+						}
+					}
 				}
 				else
 				{
