Timeout for google spreadsheets authorization

#tests 69ccfee7537c0e682324e82a
#rnx
#Synced-CL 52242993

#Synced-CL 50934557

diff --git a/Engine/Source/Programs/AutomationTool/AutomationUtils/GoogleSheetsHelper.cs b/Engine/Source/Programs/AutomationTool/AutomationUtils/GoogleSheetsHelper.cs
--- a/Engine/Source/Programs/AutomationTool/AutomationUtils/GoogleSheetsHelper.cs
+++ b/Engine/Source/Programs/AutomationTool/AutomationUtils/GoogleSheetsHelper.cs
@@ -11,6 +11,7 @@ using Google.Apis.Sheets.v4.Data;
 using System.Threading;
 using Google.Apis.Util.Store;
 using Google.Apis.Services;
+using System.Threading.Tasks;
 
 
 namespace SheetsHelper
@@ -99,21 +100,22 @@ namespace SheetsHelper
 			// JSON that has been addressed via a rebinding, but has a habit of coming back and is tricky to solve.
 			try
 			{
-				UserCredential credential;
-
-				using (FileStream stream =
-					new FileStream(SecretKeyPath, FileMode.Open, FileAccess.Read))
+				UserCredential credential = null;
+				using (var stream = new FileStream(SecretKeyPath, FileMode.Open, FileAccess.Read))
 				{
-					string credPath = System.Environment.GetFolderPath(
-						System.Environment.SpecialFolder.Personal);
-					credPath = Path.Combine(Environment.CurrentDirectory, "Engine/Restricted/NotForLicensees/Source/Programs/AutomationTool/Gauntlet/Sheets");
-
-					credential = GoogleWebAuthorizationBroker.AuthorizeAsync(
+					Task<UserCredential> Tsk = GoogleWebAuthorizationBroker.AuthorizeAsync(
 						GoogleClientSecrets.Load(stream).Secrets,
 						Scopes,
 						"user",
 						CancellationToken.None,
-						new FileDataStore(CredentialPath, true)).Result;
+						new FileDataStore(CredentialPath, true));
+
+					if (!Tsk.Wait(TimeSpan.FromMinutes(1)))
+					{
+						throw new TimeoutException("Google Sheets authorization timed out after 1 minute");
+					}
+
+					credential = Tsk.Result;
 				}
 
 				// Create Google Sheets API service.
@@ -126,16 +128,16 @@ namespace SheetsHelper
 			catch (System.Exception ex)
 			{
 				Console.Write("Error: Failed to create Sheet.\r\n\r\nException={0}\r\n\r\nInnerException={1}", ex, ex.InnerException);
-				throw ex.InnerException;				
+				throw ex.InnerException;
 			}
-			
-            // Define request parameters.
+
+			// Define request parameters.
 			SpreadsheetID = InSpreadsheetID;
 
 			Refresh();
 
 			UpdateSpreadsheetBatchRequest = new BatchUpdateSpreadsheetRequest();
-            UpdateValuesBatchRequest = new BatchUpdateValuesRequest();
+			UpdateValuesBatchRequest = new BatchUpdateValuesRequest();
 		}
 
 		public void Refresh()
