Auto download the command-line tools for streamlined Android Turnkey

#android
#Synced-CL 52302008

diff --git a/Engine/Build/Turnkey/TurnkeyManifest.xml b/Engine/Build/Turnkey/TurnkeyManifest.xml
--- a/Engine/Build/Turnkey/TurnkeyManifest.xml
+++ b/Engine/Build/Turnkey/TurnkeyManifest.xml
@@ -20,4 +20,14 @@
         <Source HostPlatform="Mac">https://redirector.gvt1.com/edgedl/android/studio/install/2024.1.2.13/android-studio-2024.1.2.13-mac.dmg</Source>
     </FileSource>
 
+    <FileSource>
+        <Platform>Android</Platform>
+        <Type>Misc</Type>
+        <Name>AndroidSDKCommandlineTools</Name>
+		<!-- Android SDK Command-line Tools 20.0 -->
+        <Source HostPlatform="Win64">https://dl.google.com/android/repository/commandlinetools-win-14742923_latest.zip</Source>
+        <Source HostPlatform="Linux">https://dl.google.com/android/repository/commandlinetools-linux-14742923_latest.zip</Source>
+		<Source HostPlatform="Mac">https://dl.google.com/android/repository/commandlinetools-mac-14742923_latest.zip</Source>
+    </FileSource>
+
 </TurnkeyManifest>
diff --git a/Engine/Source/Programs/AutomationTool/Android/AndroidPlatform.Automation.cs b/Engine/Source/Programs/AutomationTool/Android/AndroidPlatform.Automation.cs
--- a/Engine/Source/Programs/AutomationTool/Android/AndroidPlatform.Automation.cs
+++ b/Engine/Source/Programs/AutomationTool/Android/AndroidPlatform.Automation.cs
@@ -186,11 +186,55 @@ public class AndroidPlatform : Platform
 			return base.InstallSDK(BuildCommand, TurnkeyContext, Device, bUnattended, bSdkAlreadyInstalled);
 		}
 
-		string SdkDir = GetSdkDir();
-		bool bIsInstalled = Directory.Exists(SdkDir);
-
-		if (!bIsInstalled)
+		bool InstallCommandlineTools()
 		{
+			string SdkDir = GetSdkDir();
+			if (Directory.Exists(SdkDir))
+			{
+				return true;
+			}
+
+			string OutputPath = TurnkeyContext.RetrieveFileSource("AndroidSDKCommandlineTools");
+			if (OutputPath != null)
+			{
+				string SdkManager = Path.Combine(OutputPath, "cmdline-tools", "bin", "sdkmanager");
+				if (!OperatingSystem.IsWindows())
+				{
+					// Mark executable before running
+					File.SetUnixFileMode(SdkManager, File.GetUnixFileMode(SdkManager) | UnixFileMode.UserExecute | UnixFileMode.GroupExecute | UnixFileMode.OtherExecute);
+				}
+				else
+				{
+					SdkManager += ".bat";
+				}
+
+				// Use the JBR from Android Studio for the initial call to sdkmanager, if JAVA_HOME doesn't exist
+				string JavaHome = Environment.GetEnvironmentVariable("JAVA_HOME");
+				if (JavaHome == null)
+				{
+					string StudioExe = GetAndroidStudioExe();
+					string StudioDir = OperatingSystem.IsMacOS() ? Path.Combine(StudioExe, "Contents") : Path.GetDirectoryName(Path.GetDirectoryName(StudioExe))!;
+					Environment.SetEnvironmentVariable("JAVA_HOME", Path.Combine(StudioDir, "jbr"));
+				}
+
+				try
+				{
+					// This may bring up a license acceptance message that needs the user to respond, so we make a new window
+					int ExitCode = TurnkeyContext.RunExternalCommand(SdkManager, $"--sdk_root=\"{SdkDir}\" \"cmdline-tools;latest\"", false, bUnattended, true);
+					if (ExitCode == 0)
+					{
+						return true;
+					}
+				}
+				finally
+				{
+					if (JavaHome == null)
+					{
+						Environment.SetEnvironmentVariable("JAVA_HOME", null);
+					}
+				}
+			}
+
 			int Option = 2;
 			while (Option == 2)
 			{
@@ -199,19 +243,18 @@ public class AndroidPlatform : Platform
 					"  - Wait for Android Studio to start, you will see an initial dialog asking how to proceed (called \"Welcome to Android Studio\")\n" +
 					"  - Click the \"Configure\" dropdown in the bottom right, and select \"SDK Manager\"\n" +
 					"  - Click on the \"SDK Tools\" tab near the top middle of the right pane\n" +
-					"  - Check the box next to Android SDK COmmand-line Tools (latest)\n" +
+					"  - Check the box next to Android SDK Command-line Tools (latest)\n" +
 					"  - Click OK in the bottom right\n" +
 					"  - It will probably ask for you to accept a license - you MUST do this\n" +
 					"  - Once installation has completed, close/quit Android Studio to continue\n";
 
-				List<string> Options = new()
-				{
+				List<string> Options =
+				[
 					"Run Android Studio to install the Command Line Tools",
-					"Get detailed step by step guide",
-				};
+					"Get detailed step by step guide"
+				];
 
 				Option = TurnkeyContext.ReadInputInt(Prompt, Options, true, 1);
-
 				if (Option == 0)
 				{
 					return false;
@@ -224,8 +267,6 @@ public class AndroidPlatform : Platform
 				}
 			}
 
-			string AndroidStudioExe = GetAndroidStudioExe();
-
 			if (HostPlatform.Platform == UnrealTargetPlatform.Mac)
 			{
 				TurnkeyContext.RunExternalCommand("open", $"-W \"{GetAndroidStudioExe()}\"", false, true, true);
@@ -234,6 +275,13 @@ public class AndroidPlatform : Platform
 			{
 				TurnkeyContext.RunExternalCommand(GetAndroidStudioExe(), "", false, true, false);
 			}
+
+			return true;
+		}
+
+		if (!InstallCommandlineTools())
+		{
+			return false;
 		}
 
 		// run the Setup.bat in the engine, not coming from a normal FileSource
