Horde Step Message updates
Preflight failures necessary to test code
#rnx

#Synced-CL 52288889

diff --git a/Engine/Source/Programs/AutomationTool/Gauntlet/Framework/Horde/Gauntlet.Horde.cs b/Engine/Source/Programs/AutomationTool/Gauntlet/Framework/Horde/Gauntlet.Horde.cs
--- a/Engine/Source/Programs/AutomationTool/Gauntlet/Framework/Horde/Gauntlet.Horde.cs
+++ b/Engine/Source/Programs/AutomationTool/Gauntlet/Framework/Horde/Gauntlet.Horde.cs
@@ -92,6 +92,8 @@ namespace Gauntlet
 			public static HordeStepMessage InvalidValues = new HordeStepMessage("InvalidValues", JobStepMessageLevels.Error);
 			public static HordeStepMessage PermissionsFailure = new HordeStepMessage("PermissionsFailure", JobStepMessageLevels.Error);
 			public static HordeStepMessage InvalidTestData = new HordeStepMessage("InvalidTestData", JobStepMessageLevels.Error);
+			public static HordeStepMessage PushChangesFailure = new HordeStepMessage("PushChangesFailure", JobStepMessageLevels.Error);
+			public static HordeStepMessage CompileError = new HordeStepMessage("CompileError", JobStepMessageLevels.Error);
 
 			// Failures related to game/service functionality
 			public static HordeStepMessage LoginFailed = new HordeStepMessage("LoginFailed", JobStepMessageLevels.Error);
@@ -109,6 +111,9 @@ namespace Gauntlet
 			public static HordeStepMessage EditorCrashDetected = new HordeStepMessage("EditorCrashDetected", JobStepMessageLevels.Error);
 			public static HordeStepMessage ServerError = new HordeStepMessage("ServerError", JobStepMessageLevels.Error);
 			public static HordeStepMessage SyncError = new HordeStepMessage("SyncError", JobStepMessageLevels.Error);
+			public static HordeStepMessage ContentError = new HordeStepMessage("ContentError", JobStepMessageLevels.Error);
+			public static HordeStepMessage UploadError = new HordeStepMessage("UploadError", JobStepMessageLevels.Error);
+			public static HordeStepMessage LobbyFailure = new HordeStepMessage("LobbyFailure", JobStepMessageLevels.Error);
 
 			// Warning states for if things are getting close to bad
 			public static HordeStepMessage TestWarningsDetected = new HordeStepMessage("TestWarningsDetected", JobStepMessageLevels.Warning);
