Sql failures during logging have been causing test failures, making SQLfailures not fatal
#Synced-CL 52269601

diff --git a/Engine/Source/Programs/AutomationTool/Gauntlet/Unreal/Utils/Gauntlet.DeviceUsageReporter.cs b/Engine/Source/Programs/AutomationTool/Gauntlet/Unreal/Utils/Gauntlet.DeviceUsageReporter.cs
--- a/Engine/Source/Programs/AutomationTool/Gauntlet/Unreal/Utils/Gauntlet.DeviceUsageReporter.cs
+++ b/Engine/Source/Programs/AutomationTool/Gauntlet/Unreal/Utils/Gauntlet.DeviceUsageReporter.cs
@@ -2,6 +2,7 @@
 
 using UnrealBuildTool;
 using Gauntlet.Utils;
+using System;
 using System.Collections.Generic;
 
 namespace Gauntlet
@@ -87,8 +88,16 @@ namespace Gauntlet
 		{
 			foreach (IDeviceUsageReporter reporter in GetEnabledReporters())
 			{
-				Gauntlet.Log.Verbose("Reporting DeviceUsage event {0} via reporter {1}", ev.ToString(), reporter.GetType().Name);
-				reporter.RecordEvent(deviceName, platform, ev, bStarting, state, comment);
+				try
+				{
+					Gauntlet.Log.Verbose("Reporting DeviceUsage event {0} via reporter {1}", ev.ToString(), reporter.GetType().Name);
+					reporter.RecordEvent(deviceName, platform, ev, bStarting, state, comment);
+				}
+				catch (Exception ex)
+				{
+					Gauntlet.Log.Warning("Failed to report DeviceUsage event {0} via {1}: {2}",
+						ev.ToString(), reporter.GetType().Name, ex.Message);
+				}
 			}
 		}
 
@@ -96,8 +105,16 @@ namespace Gauntlet
 		{
 			foreach (IDeviceUsageReporter reporter in GetEnabledReporters())
 			{
-				Gauntlet.Log.Verbose("Adding Comment to DeviceUsage event {0} via reporter {1}", ev.ToString(), reporter.GetType().Name);
-				reporter.AddCommentToEvent(DeviceName, platform, ev, comment);
+				try
+				{
+					Gauntlet.Log.Verbose("Adding Comment to DeviceUsage event {0} via reporter {1}", ev.ToString(), reporter.GetType().Name);
+					reporter.AddCommentToEvent(DeviceName, platform, ev, comment);
+				}
+				catch (Exception ex)
+				{
+					Gauntlet.Log.Warning("Failed to add comment to DeviceUsage event {0} via {1}: {2}",
+						ev.ToString(), reporter.GetType().Name, ex.Message);
+				}
 			}
 		}
 
