Moving Megazord to Engine Utils
Adding Reboot logic
Adding DeviceInfo response
#Synced-CL 52269601

#Synced-CL 52269601

diff --git a/Engine/Source/Programs/AutomationTool/Gauntlet/Gauntlet.Automation.csproj b/Engine/Source/Programs/AutomationTool/Gauntlet/Gauntlet.Automation.csproj
--- a/Engine/Source/Programs/AutomationTool/Gauntlet/Gauntlet.Automation.csproj
+++ b/Engine/Source/Programs/AutomationTool/Gauntlet/Gauntlet.Automation.csproj
@@ -77,6 +77,7 @@
   <ItemGroup>
     <PackageReference Include="Appium.WebDriver" Version="8.0.0" />
     <PackageReference Include="Magick.NET-Q16-HDRI-AnyCPU" Version="14.9.1" PrivateAssets="all" />
+    <PackageReference Include="Microsoft.IdentityModel.Tokens" Version="6.36.0" />
     <PackageReference Include="MySql.Data" Version="6.10.9" PrivateAssets="all" />
     <PackageReference Include="Newtonsoft.Json" Version="13.0.4" PrivateAssets="all" />
   </ItemGroup>
