Update several references to AutomationTool to not use the .exe

#rnx
#Synced-CL 52421477

diff --git a/Engine/Source/Programs/AutomationTool/AutomationUtils/Automation.cs b/Engine/Source/Programs/AutomationTool/AutomationUtils/Automation.cs
--- a/Engine/Source/Programs/AutomationTool/AutomationUtils/Automation.cs
+++ b/Engine/Source/Programs/AutomationTool/AutomationUtils/Automation.cs
@@ -367,14 +367,14 @@ namespace AutomationTool
 		}
 
 		/// <summary>
-		/// Display AutomationTool.exe help.
+		/// Display AutomationTool help.
 		/// </summary>
 		private static void DisplayHelp(Dictionary<string, string> ParamDict)
 		{
 			HelpUtils.PrintHelp("Automation Help:",
 @"Executes scripted commands
 
-AutomationTool.exe [-verbose] [-Help|-List|-CompileOnly] [-p4] Command0 [-Arg0 -Arg1 -Arg2 ...] Command1 [-Arg0 -Arg1 ...] Command2 [-Arg0 ...] CommandN ... [EnvVar0=MyValue0 ... EnvVarN=MyValueN]",
+AutomationTool [-verbose] [-Help|-List|-CompileOnly] [-p4] Command0 [-Arg0 -Arg1 -Arg2 ...] Command1 [-Arg0 -Arg1 ...] Command2 [-Arg0 ...] CommandN ... [EnvVar0=MyValue0 ... EnvVarN=MyValueN]",
 				ParamDict.ToList());
 			CommandUtils.LogHelp(typeof(Automation));
 		}
diff --git a/Engine/Source/Programs/UnrealBuildTool/ProjectFiles/CodeLite/CodeLiteProject.cs b/Engine/Source/Programs/UnrealBuildTool/ProjectFiles/CodeLite/CodeLiteProject.cs
--- a/Engine/Source/Programs/UnrealBuildTool/ProjectFiles/CodeLite/CodeLiteProject.cs
+++ b/Engine/Source/Programs/UnrealBuildTool/ProjectFiles/CodeLite/CodeLiteProject.cs
@@ -458,7 +458,7 @@ namespace UnrealBuildTool
 					//
 					if (ProjectTargetType == TargetType.Game)
 					{
-						string CookGameCommandLine = "dotnet AutomationTool.exe BuildCookRun ";
+						string CookGameCommandLine = "dotnet AutomationTool.dll BuildCookRun ";
 
 						// Projects filename
 						if (OnlyGameProject != null)
