BuildCMakeLib: Fix msbuild path

#rnx

#Synced-CL 52421477

diff --git a/Engine/Source/Programs/AutomationTool/Scripts/BuildCMakeLib.Automation.cs b/Engine/Source/Programs/AutomationTool/Scripts/BuildCMakeLib.Automation.cs
--- a/Engine/Source/Programs/AutomationTool/Scripts/BuildCMakeLib.Automation.cs
+++ b/Engine/Source/Programs/AutomationTool/Scripts/BuildCMakeLib.Automation.cs
@@ -301,6 +301,9 @@ public sealed class BuildCMakeLib : BuildCommand
 			return Args;
 		}
 
+		// Returns extra environment variables passed to every cmake invocation, or null for none.
+		protected virtual Dictionary<string, string> GetCMakeEnvironment() => null;
+
 		public virtual string GetCMakeBuildArguments(TargetLib TargetLib, string TargetConfiguration)
 		{
 			StringBuilder CMakeArgs = new StringBuilder();
@@ -387,7 +390,7 @@ public sealed class BuildCMakeLib : BuildCommand
 			Logger.LogInformation("Generating projects for {Arg0} for {FriendlyName}", TargetLib.Name, FriendlyName);
 
 			string CMakeArgs = GetCMakeSetupArguments(TargetLib, TargetConfiguration);
-			if (Run(CMakeCommand, CMakeArgs).ExitCode != 0)
+			if (Run(CMakeCommand, CMakeArgs, Env: GetCMakeEnvironment()).ExitCode != 0)
 			{
 				throw new AutomationException("Unable to generate projects for {0}.", TargetLib.ToString() + ", " + FriendlyName);
 			}
@@ -395,7 +398,7 @@ public sealed class BuildCMakeLib : BuildCommand
 
 		public virtual void BuildTargetLib(TargetLib TargetLib, string TargetConfiguration)
 		{
-			if (Run(CMakeCommand, GetCMakeBuildArguments(TargetLib, TargetConfiguration)).ExitCode != 0)
+			if (Run(CMakeCommand, GetCMakeBuildArguments(TargetLib, TargetConfiguration), Env: GetCMakeEnvironment()).ExitCode != 0)
 			{
 				throw new AutomationException("Unable to build target {0}, {1}.", TargetLib.ToString(), FriendlyName);
 			}
@@ -451,6 +454,14 @@ public sealed class BuildCMakeLib : BuildCommand
 	public abstract class VSTargetPlatform : TargetPlatform
 	{
 		public override bool SeparateProjectPerConfig => false;
+
+		// Clear value of msbuild related vars in DotNetTask.cs
+		protected override Dictionary<string, string> GetCMakeEnvironment() => new()
+			{
+				["MSBuildSDKsPath"] = String.Empty,
+				["MSBuildExtensionsPath"] = String.Empty,
+				["MSBUILD_EXE_PATH"] = String.Empty,
+			};
 	}
 
 	public abstract class VS2017TargetPlatform : VSTargetPlatform
