Adding Shaders directory to ZipUpProject automation command

#Synced-CL 52428918

diff --git a/Engine/Source/Programs/AutomationTool/Scripts/BuildCommonTools.Automation.cs b/Engine/Source/Programs/AutomationTool/Scripts/BuildCommonTools.Automation.cs
--- a/Engine/Source/Programs/AutomationTool/Scripts/BuildCommonTools.Automation.cs
+++ b/Engine/Source/Programs/AutomationTool/Scripts/BuildCommonTools.Automation.cs
@@ -159,6 +159,7 @@ public class ZipProjectUp : BuildCommand
 		Filter.Include("/Plugins/...");
 		Filter.Exclude("/Plugins/.../Intermediate/...");
 		Filter.Exclude("/Plugins/.../Binaries/...");
+		Filter.Include("/Shaders/...");
 		Filter.Include("/Source/...");
 		Filter.Include("*.uproject");
 
