Cloud Storage - Add .exp and .map file extensions to symbols build part filter

#Synced-CL 51557602

diff --git a/Engine/Source/Programs/AutomationTool/BuildGraph/Tasks/CreateCloudArtifactTask.cs b/Engine/Source/Programs/AutomationTool/BuildGraph/Tasks/CreateCloudArtifactTask.cs
--- a/Engine/Source/Programs/AutomationTool/BuildGraph/Tasks/CreateCloudArtifactTask.cs
+++ b/Engine/Source/Programs/AutomationTool/BuildGraph/Tasks/CreateCloudArtifactTask.cs
@@ -187,6 +187,12 @@ namespace AutomationTool.Tasks
 	{
 		readonly CreateCloudArtifactTaskParameters _parameters;
 
+		private readonly HashSet<string> _supportedSymbolsExtensions = [
+			".exp",
+			".map",
+			".pdb"
+		];
+
 		/// <summary>
 		/// Constructor.
 		/// </summary>
@@ -533,7 +539,7 @@ namespace AutomationTool.Tasks
 							{
 								writer.WriteValue("partId", CbObjectId.NewObjectId().ToString());
 								WriteFilesArray(_parameters.SymbolFiles!,
-									filter: fileName => Path.GetExtension(fileName).Equals(".pdb", StringComparison.OrdinalIgnoreCase));
+									filter: fileName => _supportedSymbolsExtensions.Contains(Path.GetExtension(fileName).ToLowerInvariant()));
 							}
 							writer.WriteObjectEnd();
 						}
