Improve logging  if we fail to get a chunk from zenserver while staging.

#Synced-CL 52199356

diff --git a/Engine/Source/Programs/AutomationTool/Scripts/CopyBuildToStagingDirectory.Automation.cs b/Engine/Source/Programs/AutomationTool/Scripts/CopyBuildToStagingDirectory.Automation.cs
--- a/Engine/Source/Programs/AutomationTool/Scripts/CopyBuildToStagingDirectory.Automation.cs
+++ b/Engine/Source/Programs/AutomationTool/Scripts/CopyBuildToStagingDirectory.Automation.cs
@@ -5366,8 +5366,8 @@ namespace AutomationScripts
 
 					if (!HttpGetResult.IsSuccessStatusCode)
 					{
-						throw new AutomationException(String.Format("Failed reading chunk {0}.{1} from Zen. Ensure that cooking was successful.",
-							SC.PackageStoreData.ZenServerStore.ProjectId, SC.PackageStoreData.ZenServerStore.OplogId));
+						throw new AutomationException(String.Format("Failed reading chunk {0} for '{1}' from Zen with status code {2} ({3}). Ensure that cooking was successful.",
+							Request.RequestUri.PathAndQuery, Filename, (int)HttpGetResult.StatusCode, HttpGetResult.StatusCode));
 					}
 					Task<byte[]> ReadChunkTask = HttpGetResult.Content.ReadAsByteArrayAsync();
 					ReadChunkTask.Wait();
