Adding virtual bundle handling to patch diff reports
Separate patch descriptors now output for each virtual bundle

#Synced-CL 52466916

diff --git a/Engine/Source/Programs/AutomationTool/AutomationUtils/BundleUtils.cs b/Engine/Source/Programs/AutomationTool/AutomationUtils/BundleUtils.cs
--- a/Engine/Source/Programs/AutomationTool/AutomationUtils/BundleUtils.cs
+++ b/Engine/Source/Programs/AutomationTool/AutomationUtils/BundleUtils.cs
@@ -4,6 +4,7 @@ using EpicGames.Core;
 using EpicGames.Horde.Projects;
 using EpicGames.MCP.Automation;
 using EpicGames.UHT.Utils;
+using Microsoft.Extensions.Logging;
 using System;
 using System.Collections.Generic;
 using System.Linq;
@@ -421,6 +422,8 @@ namespace AutomationUtils.Automation
 				InstallBundleSources.Contains("Platform");
 		}
 
+
+
 		// Searches through all tagged files in a manifest, and determines virtual bundle tags from any files with multiple tags
 		//	Virtual bundles should always contain 'virtual' in their name
 		public static List<string> InferVirtualBundleTagsFromManifest(BuildPatchToolBase.ManifestMetadataOutput ManifestContents)
@@ -465,9 +468,29 @@ namespace AutomationUtils.Automation
 			return VirtualBundleTags;
 		}
 
+		public static List<string> InferNonVirtualBundleTagsFromManifest(BuildPatchToolBase.ManifestMetadataOutput ManifestContents)
+		{
+			List<string> NonVirtualTags = new List<string>();
+			if (ManifestContents != null)
+			{
+				List<string> VirtualBundleTags = InferVirtualBundleTagsFromManifest(ManifestContents);
+
+				// We expect very few virtual bundles (like 1) so doing this n^2 isn't terrible.
+				foreach (string Tag in ManifestContents.FileTagList)
+				{
+					if (!VirtualBundleTags.Contains(Tag))
+					{
+						NonVirtualTags.Add(Tag);
+					}
+				}
+			}
+			return NonVirtualTags;
+		}
+
 		// Loads the metadata for a given manifest and returns any NON virtual bundle tags inferred from its contents.
 		public static List<string> InferNonVirtualBundleTagsFromManifest(string ManifestPath, BuildPatchToolBase.ToolVersion BPTVersion = BuildPatchToolBase.ToolVersion.Live)
 		{
+			List<string> NonVirtualTags = new List<string>();
 			if (!string.IsNullOrEmpty(ManifestPath))
 			{
 				BuildPatchToolBase.ManifestMetadataOptions MetadataOptions = new BuildPatchToolBase.ManifestMetadataOptions
@@ -476,23 +499,39 @@ namespace AutomationUtils.Automation
 				};
 				BuildPatchToolBase.ManifestMetadataOutput ManifestMetadataOutput = null;
 				BuildPatchToolBase.Get().Execute(MetadataOptions, out ManifestMetadataOutput, BPTVersion);
-				if (ManifestMetadataOutput != null)
+				NonVirtualTags = InferNonVirtualBundleTagsFromManifest(ManifestMetadataOutput);
+			}
+			return NonVirtualTags;
+		}
+
+		public static HashSet<string> InferVirtualBundleChildTagsFromManifest(string VirtualBundleTag, BuildPatchToolBase.ManifestMetadataOutput ManifestContents)
+		{
+			HashSet<string> ChildTags = [];
+			if (ManifestContents != null)
+			{
+				foreach (BuildPatchToolBase.ManifestMetadataOutput.ManifestMetadataOutputFile TaggedFile in ManifestContents.Files)
 				{
-					List<string> VirtualBundleTags = InferVirtualBundleTagsFromManifest(ManifestMetadataOutput);
-					
-					// We expect very few virtual bundles (like 1) so doing this n^2 isn't terrible.
-					List<string> NonVirtualTags = new List<string>();
-					foreach (string Tag in ManifestMetadataOutput.FileTagList)
+					if(TaggedFile.Tags.Exists(Tag => Tag.Equals(VirtualBundleTag, StringComparison.OrdinalIgnoreCase)))
 					{
-						if (!VirtualBundleTags.Contains(Tag))
+						// There should be only one 'base' bundle tag for each file in a virtual bundle, but multiple may occurr if a child bundle has subsets
+						List<string> NonVirtualBundleTags = TaggedFile.Tags.FindAll(Tag => !Tag.Contains("virtual", StringComparison.OrdinalIgnoreCase));
+						foreach (string ChildTag in NonVirtualBundleTags)
+						{
+							ChildTags.Add(ChildTag);
+						}
+
+						if(NonVirtualBundleTags.Count == 0)
+						{
+							Console.WriteLine($"Couldn't find the base tag when inferring virtual bundle children for file {TaggedFile.Filename}");
+						}
+						else if(NonVirtualBundleTags.Count > 1)
 						{
-							NonVirtualTags.Add(Tag);
+							Console.WriteLine($"Found multiple base tags when inferring virtual bundle children for file {TaggedFile.Filename}: Base: {NonVirtualBundleTags[0]}, Found {NonVirtualBundleTags.Count}");
 						}
 					}
-					return NonVirtualTags;
 				}
 			}
-			return new List<string>();
+			return ChildTags;
 		}
 	}
 }
