Separating the OnlineTests into PS4 and PS5 components: added OnlineServicesPS4Tests and OnlineServicesPS5Tests projects.

#Synced-CL 52469003

diff --git a/Engine/Source/Programs/AutomationTool/LowLevelTests/RunLowLevelTests.cs b/Engine/Source/Programs/AutomationTool/LowLevelTests/RunLowLevelTests.cs
--- a/Engine/Source/Programs/AutomationTool/LowLevelTests/RunLowLevelTests.cs
+++ b/Engine/Source/Programs/AutomationTool/LowLevelTests/RunLowLevelTests.cs
@@ -759,11 +759,12 @@ namespace LowLevelTests
 				}
 
 				// Development executable does not contain configuration or platform name
+				// Note: The executable name may legitimately contain a platform name as part of the test app name
+				// (e.g., OnlineServicesPS4Tests). The Equals check is sufficient to identify Development builds.
 				Log.Verbose("Config type: {0}", InConfiguration);
 				if (InConfiguration == UnrealTargetConfiguration.Development)
 				{
-					if (BuildExecutableName.Equals(InTestApp, StringComparison.OrdinalIgnoreCase)
-						&& !BuildExecutableName.Contains(InPlatform.ToString(), StringComparison.OrdinalIgnoreCase))
+					if (BuildExecutableName.Equals(InTestApp, StringComparison.OrdinalIgnoreCase))
 					{
 						Log.VeryVerbose("Output Development Executable: {0}", Path.GetRelativePath(InBuildPath, Executable));
 						return Path.GetRelativePath(InBuildPath, Executable);
