Fixed external environment variable changes during turnkey

#Synced-CL 52461511

diff --git a/Engine/Source/Programs/AutomationTool/Turnkey/TurnkeyUtils.cs b/Engine/Source/Programs/AutomationTool/Turnkey/TurnkeyUtils.cs
--- a/Engine/Source/Programs/AutomationTool/Turnkey/TurnkeyUtils.cs
+++ b/Engine/Source/Programs/AutomationTool/Turnkey/TurnkeyUtils.cs
@@ -649,64 +649,105 @@ namespace Turnkey
 
 #region Env vars
 
-		private static IDictionary[] SavedEnvVars = new System.Collections.IDictionary[2];
-		private static Dictionary<string, string> EnvVarsToSaveToBatchFile = new Dictionary<string, string>();
+		private static Dictionary<string, string> SavedEnvVars;
+		private static Dictionary<string, string> EnvVarsToSaveToBatchFile = new();
+		
+		private static Dictionary<string, string> GetCombinedEnvironmentVariables()
+		{
+			// Environment variables are case-insensitive on Windows
+			Dictionary<string, string> Result = new(StringComparer.OrdinalIgnoreCase);
+
+			// 1. Load Machine Variables First
+			foreach (DictionaryEntry Pair in Environment.GetEnvironmentVariables(EnvironmentVariableTarget.Machine))
+			{
+				Result[(string)Pair.Key] = (string)Pair.Value;
+			}
+
+			// 2. Load User Variables Second (Applying Overwrite/Append rules)
+			foreach (DictionaryEntry Pair in Environment.GetEnvironmentVariables(EnvironmentVariableTarget.User))
+			{
+				string Key = (string)Pair.Key;
+				string UserValue = (string)Pair.Value;
+
+				// Apply the special concatenation rule for PATH
+				if (Key.Equals("PATH", StringComparison.OrdinalIgnoreCase))
+				{
+					string CombinedPath = Result.TryGetValue(Key, out string Value) ? Value : string.Empty;
+                
+					if (!string.IsNullOrEmpty(UserValue))
+					{
+						if (!string.IsNullOrEmpty(CombinedPath) && !CombinedPath.EndsWith(";"))
+						{
+							CombinedPath += ";";
+						}
+						
+						CombinedPath += UserValue;
+					}
+					
+					Result[Key] = CombinedPath;
+				}
+				else
+				{
+					// General Rule: User variable overwrites Machine variable
+					Result[Key] = UserValue;
+				}
+			}
+
+			return Result;
+		}
 
 		public static void StartTrackingExternalEnvVarChanges()
 		{
-			SavedEnvVars[0] = Environment.GetEnvironmentVariables(EnvironmentVariableTarget.User);
-			SavedEnvVars[1] = Environment.GetEnvironmentVariables(EnvironmentVariableTarget.Machine);
+			SavedEnvVars = GetCombinedEnvironmentVariables();
 		}
 
 		public static void EndTrackingExternalEnvVarChanges()
 		{
-			System.Collections.IDictionary[] NewEnvVars =
-			{
-				Environment.GetEnvironmentVariables(EnvironmentVariableTarget.User),
-				Environment.GetEnvironmentVariables(EnvironmentVariableTarget.Machine),
-			};
+			Dictionary<string, string> NewEnvVars = GetCombinedEnvironmentVariables();
 
-			TurnkeyUtils.Log("Scanning for envvar changes...");
+			Log("Scanning for envvar changes...");
 
-			// look for differences
-			for (int Index = 0; Index < NewEnvVars.Length; Index++)
+			// Find Additions and Modifications
+			foreach (KeyValuePair<string, string> Pair in NewEnvVars)
 			{
-				IDictionary NewSet = NewEnvVars[Index];
-				IDictionary PreviousSet = SavedEnvVars[Index];
+				if (!SavedEnvVars.ContainsKey(Pair.Key) || SavedEnvVars[Pair.Key] != Pair.Value)
+				{
+					Log("  Updating process env var {0} to {1}", Pair.Key, Pair.Value);
+					Environment.SetEnvironmentVariable(Pair.Key, Pair.Value, EnvironmentVariableTarget.Process);
+					EnvVarsToSaveToBatchFile[Pair.Key] = Pair.Value;
+				}
+			}
 
-				foreach (DictionaryEntry Pair in NewSet)
+			// Find Deletions
+			foreach (KeyValuePair<string, string> Pair in SavedEnvVars)
+			{
+				if (!NewEnvVars.ContainsKey(Pair.Key))
 				{
-					Object PrevValue = PreviousSet[Pair.Key];
-					string NewKey = Pair.Key as string;
-					string NewValue= Pair.Value as string;
+					Log("  Deleting process env var {0}", Pair.Key);
+					Environment.SetEnvironmentVariable(Pair.Key, null, EnvironmentVariableTarget.Process);
 
-					// if we have a new or changed value, apply it to the process
-					if (PrevValue == null || string.Compare(PrevValue as string, NewValue) != 0)
+					if (!EnvVarsToSaveToBatchFile.Remove(Pair.Key))
 					{
-						TurnkeyUtils.Log("  Updating process env var {0} to {1}", Pair.Key, Pair.Value);
-						Environment.SetEnvironmentVariable(NewKey, NewValue, EnvironmentVariableTarget.Process);
-						// remember to save to batch file
-						EnvVarsToSaveToBatchFile[NewKey] = NewValue;
+						EnvVarsToSaveToBatchFile[Pair.Key] = string.Empty;
 					}
 				}
 			}
 
-			TurnkeyUtils.Log("... done! ");
-			if (EnvVarsToSaveToBatchFile.Count > 0)
-			{
-				StringBuilder BatchContents = new StringBuilder();
-				foreach (var Pair in EnvVarsToSaveToBatchFile)
-				{
-					BatchContents.AppendLine("set {0}={1}", Pair.Key, Pair.Value);
-					TurnkeyUtils.Log("  Recording variable to set by caller {0}={1}", Pair.Key, Pair.Value);
-				}
+			SavedEnvVars = NewEnvVars;
+			Log("... done! ");
 
-				// write out a batch file for the caller of this to call to update vars, including all changes from previous commands
-				string BatchPath = ExpandVariables("$(EngineDir)/Intermediate/Turnkey/PostTurnkeyVariables.bat");
-				TurnkeyUtils.Log("  Writing updated envvars to {0}", BatchPath);
-				Directory.CreateDirectory(Path.GetDirectoryName(BatchPath));
-				File.WriteAllText(BatchPath, BatchContents.ToString());
+			StringBuilder BatchContents = new();
+			foreach (KeyValuePair<string, string> Pair in EnvVarsToSaveToBatchFile)
+			{
+				BatchContents.AppendLine("set {0}={1}", Pair.Key, Pair.Value);
+				Log("  Recording variable to set by caller {0}={1}", Pair.Key, Pair.Value);
 			}
+
+			// write out a batch file for the caller of this to call to update vars, including all changes from previous commands
+			string BatchPath = ExpandVariables("$(EngineDir)/Intermediate/Turnkey/PostTurnkeyVariables.bat");
+			Log("  Writing updated envvars to {0}", BatchPath);
+			Directory.CreateDirectory(Path.GetDirectoryName(BatchPath));
+			File.WriteAllText(BatchPath, BatchContents.ToString());
 		}
 
 #endregion
