add optional param to ziptask to skip the default tagset add
#pf 69d514666e459e121728abf0

diff --git a/Engine/Source/Programs/AutomationTool/BuildGraph/Tasks/ZipTask.cs b/Engine/Source/Programs/AutomationTool/BuildGraph/Tasks/ZipTask.cs
--- a/Engine/Source/Programs/AutomationTool/BuildGraph/Tasks/ZipTask.cs
+++ b/Engine/Source/Programs/AutomationTool/BuildGraph/Tasks/ZipTask.cs
@@ -38,6 +38,12 @@ namespace AutomationTool.Tasks
 		[TaskParameter]
 		public FileReference ZipFile { get; set; }
 
+		/// <summary>
+		/// Skips adding the resulting zip to the default tagset.
+		/// </summary>
+		[TaskParameter(Optional = true, ValidationType = TaskParameterValidationType.Default)]
+		public bool SkipDefaultTag { get; set; }
+
 		/// <summary>
 		/// Tag to be applied to the created zip file.
 		/// </summary>
@@ -105,8 +111,11 @@ namespace AutomationTool.Tasks
 				FindOrAddTagSet(tagNameToFileSet, tagName).Add(_parameters.ZipFile);
 			}
 
-			// Add the archive to the set of build products
-			buildProducts.Add(_parameters.ZipFile);
+			if(!_parameters.SkipDefaultTag)
+			{
+				// Add the archive to the set of build products
+				buildProducts.Add(_parameters.ZipFile);
+			}
 			return Task.CompletedTask;
 		}
 
