Adding parsing for android signal crashes

#Synced-CL 52419526

diff --git a/Engine/Source/Programs/AutomationTool/Gauntlet/Platform/Android/Gauntlet.TargetDeviceAndroid.cs b/Engine/Source/Programs/AutomationTool/Gauntlet/Platform/Android/Gauntlet.TargetDeviceAndroid.cs
--- a/Engine/Source/Programs/AutomationTool/Gauntlet/Platform/Android/Gauntlet.TargetDeviceAndroid.cs
+++ b/Engine/Source/Programs/AutomationTool/Gauntlet/Platform/Android/Gauntlet.TargetDeviceAndroid.cs
@@ -1792,11 +1792,16 @@ namespace Gauntlet
 		{
 			get
 			{
+				CheckGeneratedCrashLog();
 				return LogProcess.Output;
 			}
 		}
 
-		public ILogStreamReader GetLogReader() => LogProcess.GetLogReader();
+		public ILogStreamReader GetLogReader()
+		{
+			CheckGeneratedCrashLog();
+			return LogProcess.GetLogReader();
+		}
 
 		public ILogStreamReader GetLogBufferReader() => LogProcess.GetLogBufferReader();
 
@@ -1819,6 +1824,7 @@ namespace Gauntlet
 		protected AndroidAppInstall Install;
 
 		protected bool bHaveSavedArtifacts;
+		protected bool bWasCheckedForCrash = false;
 
 		private bool ActivityExited = false;
 
@@ -1918,6 +1924,20 @@ namespace Gauntlet
 			}
 		}
 
+		private void CheckGeneratedCrashLog()
+		{
+			// Parse Android native crash dumps from logcat and convert to Unreal crash format
+			if (HasExited && !bWasCheckedForCrash)
+			{
+				bWasCheckedForCrash = true;
+				string CrashLog = AndroidCrashParser.GenerateCrashLog(LogProcess.GetLogReader());
+				if (!string.IsNullOrEmpty(CrashLog))
+				{
+					LogProcess.AppendToOutput(CrashLog, false);
+				}
+			}
+		}
+
 		/// <summary>
 		/// Checks on device whether the activity is running
 		/// </summary>
@@ -1951,6 +1971,121 @@ namespace Gauntlet
 		}
 	}
 
+	/// <summary>
+	/// Helper class to parse Android native crash tombstones and generate Unreal compatible log callstack
+	/// </summary>
+	static class AndroidCrashParser
+	{
+		/// <summary>
+		/// Parse Android native crash tombstone format to Unreal log format
+		/// </summary>
+		public static string GenerateCrashLog(ILogStreamReader LogReader)
+		{
+			try
+			{
+				string SignalInfo;
+				List<string> Callstack = ParseTombstone(LogReader, out SignalInfo);
+				if (Callstack == null || Callstack.Count == 0)
+				{
+					return null;
+				}
+
+				StringBuilder CrashLog = new StringBuilder();
+				CrashLog.Append(string.Format("LogCore: Error: === Fatal Error: ===\n"));
+				CrashLog.Append(string.Format("LogCore: Error: {0}\n", SignalInfo));
+				CrashLog.Append(string.Join("\n", Callstack));
+
+				return CrashLog.ToString();
+			}
+			catch (Exception Ex)
+			{
+				Log.Warning(KnownLogEvents.Gauntlet_DeviceEvent, "Exception parsing Android crash tombstone {Exception}", Ex.Message);
+			}
+
+			return null;
+		}
+
+		private static List<string> ParseTombstone(ILogStreamReader LogReader, out string SignalInfo)
+		{
+			SignalInfo = "Android Native Crash";
+
+			Regex TombstonePattern = new Regex(@"\*\*\* \*\*\* \*\*\* \*\*\* \*\*\* \*\*\* \*\*\* \*\*\*");
+			Regex SignalPattern = new Regex(@"signal (\d+) \((\w+)\)");
+			Regex BacktraceStartPattern = new Regex(@"^\s*backtrace:\s*$", RegexOptions.IgnoreCase);
+			Regex BacktraceLinePattern = new Regex(@"^\s*#(\d+)\s+pc\s+([0-9a-fA-F]+)\s+(.+?)(?:\s+\((.+?)\))?(?:\s+\(BuildId:\s*([0-9a-fA-F]+)\))?$");
+
+			List<string> Callstack = null;
+			bool InCrash = false;
+			bool InBacktrace = false;
+
+			LogReader.SetLineIndex(0);
+			foreach (string Line in LogReader.EnumerateNextLines())
+			{
+				// Look for tombstone start marker
+				if (!InCrash && TombstonePattern.IsMatch(Line))
+				{
+					InCrash = true;
+					Callstack = new List<string>();
+					continue;
+				}
+
+				if (!InCrash) continue;
+
+				// Extract signal information
+				if (!InBacktrace)
+				{
+					Match SignalMatch = SignalPattern.Match(Line);
+					if (SignalMatch.Success)
+					{
+						string SignalNumber = SignalMatch.Groups[1].Value;
+						string SignalName = SignalMatch.Groups[2].Value;
+						SignalInfo = string.Format("Android Native Crash: signal {0} ({1})", SignalNumber, SignalName);
+					}
+
+					// Check for backtrace start
+					if (BacktraceStartPattern.IsMatch(Line))
+					{
+						InBacktrace = true;
+					}
+				}
+				else
+				{
+					// Parse backtrace lines
+					Match BacktraceMatch = BacktraceLinePattern.Match(Line);
+					if (BacktraceMatch.Success)
+					{
+						string Address = BacktraceMatch.Groups[2].Value;
+						string Path = BacktraceMatch.Groups[3].Value;
+						string Function = BacktraceMatch.Groups[4].Success ? BacktraceMatch.Groups[4].Value : "";
+
+						// Extract module name from path
+						string Module = System.IO.Path.GetFileName(Path.Split('!')[0].Trim());
+
+						// Format as Unreal callstack: LogCore: Error: [Callstack] 0xaddress module!function [file]
+						string FormattedFrame;
+						if (!string.IsNullOrEmpty(Function))
+						{
+							FormattedFrame = string.Format("LogCore: Error: [Callstack] 0x{0} {1}!{2} [{3}]", Address, Module, Function, Path);
+						}
+						else
+						{
+							FormattedFrame = string.Format("LogCore: Error: [Callstack] 0x{0} {1} [{2}]", Address, Module, Path);
+						}
+
+						Callstack.Add(FormattedFrame);
+					}
+					else if (!string.IsNullOrWhiteSpace(Line) && !Line.StartsWith(" ") && !Line.StartsWith("\t"))
+					{
+						// End of backtrace - non-indented line found
+						break;
+					}
+				}
+			}
+
+			return Callstack;
+		}
+	}
+
 	public class AdbScreenRecorder : IDisposable
 	{
 		private Process ADBProcess;
