[NuGet] Update Magick.NET-Q16-HDRI-AnyCPU to 14.11.1

#rnx
#Synced-CL 52526864

diff --git a/Engine/Source/Programs/AutomationTool/AutomationTool.csproj b/Engine/Source/Programs/AutomationTool/AutomationTool.csproj
--- a/Engine/Source/Programs/AutomationTool/AutomationTool.csproj
+++ b/Engine/Source/Programs/AutomationTool/AutomationTool.csproj
@@ -84,7 +84,7 @@
   <ItemGroup>
     <PackageReference Include="AWSSDK.S3" Version="3.7.501.14" />
     <PackageReference Include="AWSSDK.SecurityToken" Version="3.7.501.14" />
-    <PackageReference Include="Magick.NET-Q16-HDRI-AnyCPU" Version="14.9.1" />
+    <PackageReference Include="Magick.NET-Q16-HDRI-AnyCPU" Version="14.11.1" />
     <PackageReference Include="Microsoft.Build" Version="17.11.48" ExcludeAssets="runtime" PrivateAssets="all">
       <IncludeAssets>compile; build; native; contentfiles; analyzers; buildtransitive</IncludeAssets>
     </PackageReference>
diff --git a/Engine/Source/Programs/AutomationTool/AutomationUtils/AutomationUtils.Automation.csproj b/Engine/Source/Programs/AutomationTool/AutomationUtils/AutomationUtils.Automation.csproj
--- a/Engine/Source/Programs/AutomationTool/AutomationUtils/AutomationUtils.Automation.csproj
+++ b/Engine/Source/Programs/AutomationTool/AutomationUtils/AutomationUtils.Automation.csproj
@@ -37,7 +37,7 @@
     <ProjectReference Include="..\..\UnrealBuildTool\UnrealBuildTool.csproj" PrivateAssets="All"><Private>true</Private></ProjectReference>
   </ItemGroup>
   <ItemGroup>
-    <PackageReference Include="Magick.NET-Q16-HDRI-AnyCPU" Version="14.9.1" PrivateAssets="all" />
+    <PackageReference Include="Magick.NET-Q16-HDRI-AnyCPU" Version="14.11.1" PrivateAssets="all" />
     <PackageReference Include="System.Drawing.Common" Version="8.0.20" />
     <!--Required by Ionic -->
     <PackageReference Include="System.Security.Permissions" Version="8.0.0" />
diff --git a/Engine/Source/Programs/AutomationTool/Gauntlet/Gauntlet.Automation.csproj b/Engine/Source/Programs/AutomationTool/Gauntlet/Gauntlet.Automation.csproj
--- a/Engine/Source/Programs/AutomationTool/Gauntlet/Gauntlet.Automation.csproj
+++ b/Engine/Source/Programs/AutomationTool/Gauntlet/Gauntlet.Automation.csproj
@@ -76,7 +76,7 @@
   </ItemGroup>
   <ItemGroup>
     <PackageReference Include="Appium.WebDriver" Version="8.0.0" />
-    <PackageReference Include="Magick.NET-Q16-HDRI-AnyCPU" Version="14.9.1" PrivateAssets="all" />
+    <PackageReference Include="Magick.NET-Q16-HDRI-AnyCPU" Version="14.11.1" PrivateAssets="all" />
     <PackageReference Include="Microsoft.IdentityModel.Tokens" Version="6.36.0" />
     <PackageReference Include="MySql.Data" Version="6.10.9" PrivateAssets="all" />
     <PackageReference Include="Newtonsoft.Json" Version="13.0.4" PrivateAssets="all" />
