Add exception message to a localization warning

#Synced-CL 52520240

diff --git a/Engine/Source/Programs/AutomationTool/Localization/LocalizationProvider.cs b/Engine/Source/Programs/AutomationTool/Localization/LocalizationProvider.cs
--- a/Engine/Source/Programs/AutomationTool/Localization/LocalizationProvider.cs
+++ b/Engine/Source/Programs/AutomationTool/Localization/LocalizationProvider.cs
@@ -193,9 +193,10 @@ namespace EpicGames.Localization
 									var LocalizationProviderId = Method.Invoke(null, null) as string;
 									CachedLocalizationProviderTypes.Add(LocalizationProviderId, PotentialLocalizationNodeType);
 								}
-								catch
+								catch(Exception ex)
 								{
 									Logger.LogWarning("Type '{Name}' threw when calling its StaticGetLocalizationProviderId method.", PotentialLocalizationNodeType.FullName);
+									Logger.LogWarning("Exception: {Message}", ex.Message);
 								}
 							}
 							else
