Stream watched log event categories to phase event streams
+ avoid duplicating events streamed to phase to other system

#rnx

#Synced-CL 52538943

diff --git a/Engine/Source/Programs/AutomationTool/Gauntlet/Framework/Base/Gauntlet.TestTracking.cs b/Engine/Source/Programs/AutomationTool/Gauntlet/Framework/Base/Gauntlet.TestTracking.cs
--- a/Engine/Source/Programs/AutomationTool/Gauntlet/Framework/Base/Gauntlet.TestTracking.cs
+++ b/Engine/Source/Programs/AutomationTool/Gauntlet/Framework/Base/Gauntlet.TestTracking.cs
@@ -41,6 +41,7 @@ namespace Gauntlet
 			public double ElapseTimeInSeconds;
 			public DateTime? StartedAt;
 			public TestPhaseOutcome Outcome;
+			public int ResetCount { get; private set; }
 			private ConcurrentQueue<LogEvent> Events;
 
 			/// <summary>
@@ -59,6 +60,7 @@ namespace Gauntlet
 				Outcome = TestPhaseOutcome.Unknown;
 				Events = new ConcurrentQueue<LogEvent>();
 				Data = InData;
+				ResetCount = 0;
 			}
 
 			/// <summary>
@@ -135,6 +137,32 @@ namespace Gauntlet
 				return default(T);
 			}
 
+			/// <summary>
+			/// List events associated with this phase
+			/// </summary>
+			/// <returns></returns>
+			public IEnumerable<LogEvent> GetEvents()
+			{
+				return Events;
+			}
+
+			/// <summary>
+			/// Reset this phase outcome and optionally its events 
+			/// </summary>
+			/// <param name="bKeepEvents"></param>
+			public void Reset(bool bKeepEvents = true)
+			{
+				ElapseTimeInSeconds = 0;
+				StartedAt = null;
+				Outcome = TestPhaseOutcome.Unknown;
+				Data = null;
+				if (!bKeepEvents)
+				{
+					Events.Clear();
+				}
+				ResetCount++;
+			}
+
 			/// <summary>
 			/// Add the information tracked by the phase to the test report
 			/// </summary>
@@ -169,6 +197,7 @@ namespace Gauntlet
 				if (Data is IPhaseDataHandler DataHandler)
 				{
 					DataHandler.OnAddToReport(ReportedPhase);
+					Outcome = ReportedPhase.outcome;
 				}
 
 				return ReportedPhase;
diff --git a/Engine/Source/Programs/AutomationTool/Gauntlet/Framework/Gauntlet.Utils.cs b/Engine/Source/Programs/AutomationTool/Gauntlet/Framework/Gauntlet.Utils.cs
--- a/Engine/Source/Programs/AutomationTool/Gauntlet/Framework/Gauntlet.Utils.cs
+++ b/Engine/Source/Programs/AutomationTool/Gauntlet/Framework/Gauntlet.Utils.cs
@@ -536,8 +536,15 @@ namespace Gauntlet
 		/// <param name="Args"></param>
 		public void Interrupt(string Message, params object[] Args)
 		{
-			Instance?._Phases?.Current?.AddEvent(Logging.LogLevel.Critical, Message, Args);
-			Instance?._TestNode?.AddTestEvent(new UnrealTestEvent(EventSeverity.Fatal, Message, Args));
+			Phase? Phase = Instance?._Phases?.Current;
+			if (Phase != null)
+			{
+				Phase.AddEvent(Logging.LogLevel.Critical, Message, Args);
+			}
+			else
+			{
+				Instance?._TestNode?.AddTestEvent(new UnrealTestEvent(EventSeverity.Fatal, Message, Args));
+			}
 			
 			if (_TestNode == null || !_TestNode.LogWarningsAndErrorsAfterSummary)
 			{
@@ -557,8 +564,15 @@ namespace Gauntlet
 		/// <param name="Args"></param>
 		public void Failure(string Message, params object[] Args)
 		{
-			Instance?._Phases?.Current?.AddEvent(Logging.LogLevel.Error, Message, Args);
-			Instance?._TestNode?.AddTestEvent(new UnrealTestEvent(EventSeverity.Error, Message, Args));
+			Phase? Phase = Instance?._Phases?.Current;
+			if (Phase != null)
+			{
+				Phase.AddEvent(Logging.LogLevel.Error, Message, Args);
+			}
+			else
+			{
+				Instance?._TestNode?.AddTestEvent(new UnrealTestEvent(EventSeverity.Error, Message, Args));
+			}
 
 			if (_TestNode == null || !_TestNode.LogWarningsAndErrorsAfterSummary)
 			{
@@ -589,8 +603,15 @@ namespace Gauntlet
 		/// <param name="Args"></param>
 		public void Warning(EventId EventId, string Message, params object[] Args)
 		{
-			Instance?._Phases?.Current?.AddEvent(Logging.LogLevel.Warning, Message, Args);
-			Instance?._TestNode?.AddTestEvent(new UnrealTestEvent(EventSeverity.Warning, Message, Args));
+			Phase? Phase = Instance?._Phases?.Current;
+			if (Phase != null)
+			{
+				Phase.AddEvent(Logging.LogLevel.Warning, Message, Args);
+			}
+			else
+			{
+				Instance?._TestNode?.AddTestEvent(new UnrealTestEvent(EventSeverity.Warning, Message, Args));
+			}
 
 			if (_TestNode == null || !_TestNode.LogWarningsAndErrorsAfterSummary)
 			{
@@ -605,8 +626,15 @@ namespace Gauntlet
 		/// <param name="Args"></param>
 		public void Info(string Message, params object[] Args)
 		{
-			Instance?._Phases?.Current?.AddEvent(Logging.LogLevel.Information, Message, Args);
-			Instance?._TestNode?.AddTestEvent(new UnrealTestEvent(EventSeverity.Info, Message, Args));
+			Phase? Phase = Instance?._Phases?.Current;
+			if (Phase != null)
+			{
+				Phase.AddEvent(Logging.LogLevel.Information, Message, Args);
+			}
+			else
+			{
+				Instance?._TestNode?.AddTestEvent(new UnrealTestEvent(EventSeverity.Info, Message, Args));
+			}
 			Log.Info(KnownLogEvents.Gauntlet_TestEvent, Message, Args);
 		}
 	}
diff --git a/Engine/Source/Programs/AutomationTool/Gauntlet/Framework/RpcFramework/Gauntlet.RpcTimedSections.cs b/Engine/Source/Programs/AutomationTool/Gauntlet/Framework/RpcFramework/Gauntlet.RpcTimedSections.cs
--- a/Engine/Source/Programs/AutomationTool/Gauntlet/Framework/RpcFramework/Gauntlet.RpcTimedSections.cs
+++ b/Engine/Source/Programs/AutomationTool/Gauntlet/Framework/RpcFramework/Gauntlet.RpcTimedSections.cs
@@ -1,7 +1,6 @@
 // Copyright Epic Games, Inc. All Rights Reserved.
 
 #nullable enable
-using Gauntlet.Report;
 using Gauntlet.TestTracking;
 using System;
 using System.Collections.Concurrent;
@@ -10,7 +9,6 @@ using System.Diagnostics;
 using System.Linq;
 using System.Text;
 using static Gauntlet.HordeReport.AutomatedTestSessionData;
-using Logging = Microsoft.Extensions.Logging;
 
 namespace Gauntlet
 {
@@ -91,15 +89,7 @@ namespace Gauntlet
 			/// Called when test phase is added to report
 			/// </summary>
 			/// <param name="ReportedPhase"></param>
-			public void OnAddToReport(TestPhase ReportedPhase)
-			{
-				if (ReportedPhase.outcome == TestPhaseOutcome.Interrupted
-					&& !ReportedPhase.GetStream().GetEvents()
-						.Where(E => E.Level == Logging.LogLevel.Error || E.Level == Logging.LogLevel.Critical).Any())
-				{
-					ReportedPhase.GetStream().AddError("Failed to complete phase. See failing nested phases and Closure for details.");
-				}
-			}
+			public void OnAddToReport(TestPhase ReportedPhase) { }
 		}
 
 		/// <summary>
diff --git a/Engine/Source/Programs/AutomationTool/Gauntlet/Unreal/Base/Gauntlet.UnrealTestConfiguration.cs b/Engine/Source/Programs/AutomationTool/Gauntlet/Unreal/Base/Gauntlet.UnrealTestConfiguration.cs
--- a/Engine/Source/Programs/AutomationTool/Gauntlet/Unreal/Base/Gauntlet.UnrealTestConfiguration.cs
+++ b/Engine/Source/Programs/AutomationTool/Gauntlet/Unreal/Base/Gauntlet.UnrealTestConfiguration.cs
@@ -857,6 +857,7 @@ namespace Gauntlet
 		/// <summary>
 		/// Log channels that should be treated as events for this test. Warnings & Errors in these
 		/// channels will be promoted to test warnings and errors. For LogFoo return "Foo".
+		/// If a phase was started, the events are streamed into it additionally.
 		/// </summary>
 		public List<string> LogCategoriesForEvents { get; protected set; } = new List<string>();
 
diff --git a/Engine/Source/Programs/AutomationTool/Gauntlet/Unreal/Base/Gauntlet.UnrealTestNode.cs b/Engine/Source/Programs/AutomationTool/Gauntlet/Unreal/Base/Gauntlet.UnrealTestNode.cs
--- a/Engine/Source/Programs/AutomationTool/Gauntlet/Unreal/Base/Gauntlet.UnrealTestNode.cs
+++ b/Engine/Source/Programs/AutomationTool/Gauntlet/Unreal/Base/Gauntlet.UnrealTestNode.cs
@@ -16,6 +16,7 @@ using Gauntlet.Utils;
 using static Gauntlet.HordeReport.AutomatedTestSessionData;
 using Gauntlet.TestTracking;
 using Gauntlet.Report;
+using Logging = Microsoft.Extensions.Logging;
 
 namespace Gauntlet
 {
@@ -1640,12 +1641,51 @@ namespace Gauntlet
 		{
 			// Assume all devices were used to run all test cases
 			HashSet<string> DeviceKeys = SessionData.GetDevices().Select(D => D.key).ToHashSet();
+			bool RoleResultsChecked = false;
 			// Go through all phases
 			foreach (Phase CurPhase in Phases)
 			{
 				if (string.IsNullOrEmpty(CurPhase.Name)) continue;
 				TestPhase ReportedPhase = CurPhase.AddToReport(SessionData);
 
+				if (ReportedPhase.outcome == TestPhaseOutcome.Interrupted
+					&& !ReportedPhase.GetStream().GetEvents()
+						.Where(E => E.Level == Logging.LogLevel.Error || E.Level == Logging.LogLevel.Critical).Any())
+				{
+					// Look into Role artifact for extra information
+					if (RoleResults != null)
+					{
+						TestEventStream PhaseStream = ReportedPhase.GetStream();
+						foreach (UnrealRoleResult Item in RoleResults)
+						{
+							if (Item.ProcessResult != UnrealProcessResult.ExitOk)
+							{
+								string RoleType = Item.Artifacts.SessionRole.RoleType.ToString();
+								PhaseStream.AddError($"{RoleType}: {Item.Summary}");
+								if (!RoleResultsChecked)
+								{
+									var FatalEvents = Item.Events.Where(E => E.Severity == EventSeverity.Fatal);
+									if (FatalEvents.Any())
+									{
+										foreach (UnrealTestEvent Event in FatalEvents)
+										{
+											PhaseStream.AddEvent(Event);
+										}
+									}
+									else
+									{
+										foreach (UnrealTestEvent Event in Item.Events)
+										{
+											if (Event.IsError) PhaseStream.AddEvent(Event);
+										}
+									}
+								}
+							}
+						}
+					}
+					RoleResultsChecked = true;
+				}
+
 				if (ReportedPhase.outcome != TestPhaseOutcome.NotRun
 					&& ReportedPhase.outcome != TestPhaseOutcome.Skipped)
 				{
@@ -1664,14 +1704,13 @@ namespace Gauntlet
 			var ClosurePhase = TestReport.AddPhase(PhaseName);
 			ClosurePhase.SetTiming(SessionTime ?? DateTime.UtcNow, TimeElapseSeconds);
 			// Propagate events
-			if (RoleResults != null)
+			if (!GetCachedConfiguration().IgnoreRoleFailures && RoleResults != null)
 			{
 				foreach (UnrealRoleResult Item in RoleResults)
 				{
-					string RoleType = Item.Artifacts.SessionRole.RoleType.ToString();
 					if (Item.ProcessResult != UnrealProcessResult.InitializationFailure && Item.ProcessResult != UnrealProcessResult.LoginFailed)
 					{
-						var RelevantStream = ClosurePhase.GetStream();
+						TestEventStream RelevantStream = ClosurePhase.GetStream();
 						foreach (UnrealTestEvent Event in Item.Events)
 						{
 							RelevantStream.AddEvent(Event);
@@ -1679,22 +1718,26 @@ namespace Gauntlet
 						// Add Process result
 						if (Item.ProcessResult != UnrealProcessResult.ExitOk)
 						{
+							string RoleType = Item.Artifacts.SessionRole.RoleType.ToString();
 							RelevantStream.AddError($"{RoleType}: {Item.Summary}");
 						}
 					}
 				}
 				ClosurePhase.deviceKeys = TestReport.GetDevices().Select(D => D.key).ToHashSet();
 			}
-			var MainStream = ClosurePhase.GetStream();
+			TestEventStream MainStream = ClosurePhase.GetStream();
 			foreach (UnrealTestEvent Event in TestNodeEvents)
 			{
 				MainStream.AddEvent(Event);
 			}
-			// Outcome
-			TestResult Result = GetTestResult();
-			if (Result != TestResult.Passed)
+			if (!Phases.Any())
 			{
-				ClosurePhase.SetOutcome(Result);
+				// Transpose test outcome to Closure phase only if no Phase exists
+				TestResult Result = GetTestResult();
+				if (Result != TestResult.Passed)
+				{
+					ClosurePhase.SetOutcome(Result);
+				}
 			}
 		}
 
@@ -1803,7 +1846,7 @@ namespace Gauntlet
 						{
 							bool bAllowZeroTestResults = Globals.Params.ParseParam("AllowZeroTestResults");
 
-							var Stream = InitPhase.GetStream();
+							TestEventStream Stream = InitPhase.GetStream();
 							if (bAllowZeroTestResults)
 							{
 								Stream.AddWarning("No tests were executed, check test filter!");
@@ -1862,7 +1905,7 @@ namespace Gauntlet
 						foreach (UnrealRoleResult RoleResult in RoleResults.Where(R => R.ProcessResult != UnrealProcessResult.ExitOk && R.ProcessResult != UnrealProcessResult.EncounteredEnsure))
 						{
 							string RoleType = RoleResult.Artifacts.SessionRole.RoleType.ToString();
-							var Stream = InitPhase.GetStream();
+							TestEventStream Stream = InitPhase.GetStream();
 							Stream.AddError("{Role}: {Summary}", RoleType, RoleResult.Summary);
 							foreach (UnrealTestEvent Item in RoleResult.Events)
 							{
@@ -2296,8 +2339,15 @@ namespace Gauntlet
 						return;
 					}
 
-					foreach (string TestLine in Parser.GetLogFromShortNameChannels(HeartbeatLogCategories))
+					bool bPromoteToPhase = Phases.Current != null && CachedConfig.LogCategoriesForEvents.Count > 0;
+
+					foreach (UnrealLog.LogEntry Event in Parser.GetEventsFromChannels(HeartbeatLogCategories, false))
 					{
+						string TestLine = Event.ToString();
+						if (bPromoteToPhase && Event.Level >= UnrealLog.LogLevel.Warning)
+						{
+							Phases.Current?.AddEvent(Event.LogLevel, "{App}: {Message}", AppPrefix, Event.ToStringWithoutLevel());
+						}
 						Log.Info("{App}: {Message}", AppPrefix, TestLine);
 
 						if (bUpdateHeartbeatTime)
@@ -2677,7 +2727,14 @@ namespace Gauntlet
 
 				foreach (var Warning in PrintedWarningList)
 				{
-					Log.Warning(KnownLogEvents.Gauntlet_TestEvent, " * {Message}", Warning);
+					if (GetCachedConfiguration().IgnoreRoleFailures)
+					{
+						Log.Info(KnownLogEvents.Gauntlet_TestEvent, " * {Message}", Warning);
+					}
+					else
+					{
+						Log.Warning(KnownLogEvents.Gauntlet_TestEvent, " * {Message}", Warning);
+					}
 				}
 				if (!string.IsNullOrEmpty(TrimStatement))
 				{
@@ -2763,6 +2820,36 @@ namespace Gauntlet
 			return MB.ToString();
 		}
 
+		/// <summary>
+		/// Log Phase Outcomes and associated error events if any
+		/// </summary>
+		protected virtual void LogPhasesSummary()
+		{
+			foreach(Phase Phase in Phases)
+			{
+				TestPhaseOutcome Outcome = Phase.Outcome;
+				string Message = " ### [{Outcome}] Phase: {Phase}";
+				if (Phase.ResetCount > 0)
+				{
+					Message += $" (retried={Phase.ResetCount})";
+				}
+
+				if (Outcome == TestPhaseOutcome.Failed || Outcome == TestPhaseOutcome.Interrupted)
+				{
+					Log.Error(KnownLogEvents.Gauntlet_TestEvent, Message, Outcome.ToString(), Phase.Name);
+					foreach(LogEvent Error in Phase.GetEvents().Where(E => E.Level == Logging.LogLevel.Critical || E.Level == Logging.LogLevel.Error))
+					{
+						Log.Error(Error.Id, $"  * {(Error.Format ?? Error.Message)}", ((Dictionary<string, object>)Error.Properties)?.Values.ToArray());
+					}
+				}
+				else
+				{
+					Log.Info(Message, Outcome.ToString(), Phase.Name);
+				}
+			}
+			Log.Info("");
+		}
+
 		/// <summary>
 		/// Returns the current 0-based Pass count for this node.
 		/// </summary>
@@ -2955,12 +3042,15 @@ namespace Gauntlet
 				Warnings += RoleResult.LogSummary.Warnings.Count();
 			}
 
+			int FailingPhases = Phases.Where(P => P.Outcome == TestPhaseOutcome.Interrupted || P.Outcome == TestPhaseOutcome.Failed).Count();
+
 			Log.Info("");
 			Log.Info(string.Join("\n", new string[] {
 				string.Format(" * Main Context: {0}", GetMainRoleContextString()),
 				FatalErrors > 0 ? string.Format(" * FatalErrors: {0}", FatalErrors) : null,
 				Ensures > 0 ? string.Format(" * Ensures: {0}", Ensures) : null,
 				Errors > 0 ? string.Format(" * Log Errors: {0}", Errors) : null,
+				FailingPhases > 0 ? string.Format(" * Failing Phases: {0}", FailingPhases) : null,
 				Warnings > 0 ? string.Format(" * Log Warnings: {0}", Warnings) : null,
 				string.Format(" * Result: {0}", GetTestResult())
 			}.Where(L => !string.IsNullOrEmpty(L))));
@@ -3003,7 +3093,7 @@ namespace Gauntlet
 				// add Summary
 				string HorizontalLine = " " + new string('-', 80);
 				Log.Info(HorizontalLine);
-				Log.Info(" # Test Summary: " + this.Name);
+				Log.Info(" # Test Summary: {Name}", this.Name);
 				Log.Info(HorizontalLine);
 				string HeaderSummary = GetTestSummaryHeader();
 				if (!string.IsNullOrEmpty(HeaderSummary))
@@ -3015,6 +3105,14 @@ namespace Gauntlet
 					LogTestSummaryHeader();
 				}
 
+				if (Phases.Any())
+				{
+					Log.Info(HorizontalLine);
+					Log.Info(" # Phase(s): {Name}", this.Name);
+					Log.Info(HorizontalLine);
+					LogPhasesSummary();
+				}
+
 				Log.Info(HorizontalLine);
 				Log.Info(" # Role(s): {Name}", this.Name);
 				Log.Info(HorizontalLine);
diff --git a/Engine/Source/Programs/AutomationTool/Gauntlet/Unreal/Utils/Gauntlet.UnrealLogParser.cs b/Engine/Source/Programs/AutomationTool/Gauntlet/Unreal/Utils/Gauntlet.UnrealLogParser.cs
--- a/Engine/Source/Programs/AutomationTool/Gauntlet/Unreal/Utils/Gauntlet.UnrealLogParser.cs
+++ b/Engine/Source/Programs/AutomationTool/Gauntlet/Unreal/Utils/Gauntlet.UnrealLogParser.cs
@@ -61,6 +61,29 @@ namespace Gauntlet
 			/// </summary>
 			public LogLevel Level { get; private set; }
 
+			/// <summary>
+			/// Convert Unreal Log level to .Net LogLevel
+			/// </summary>
+			public Logging.LogLevel LogLevel
+			{
+				get
+				{
+					switch (Level)
+					{
+						case UnrealLog.LogLevel.VeryVerbose:
+						case UnrealLog.LogLevel.Verbose:
+							return Logging.LogLevel.Debug;
+						case UnrealLog.LogLevel.Warning:
+							return Logging.LogLevel.Warning;
+						case UnrealLog.LogLevel.Error:
+							return Logging.LogLevel.Error;
+						case UnrealLog.LogLevel.Fatal:
+							return Logging.LogLevel.Critical;
+					}
+					return Logging.LogLevel.Information;
+				}
+			}
+
 			/// <summary>
 			/// The message string from the entry
 			/// </summary>
@@ -79,7 +102,7 @@ namespace Gauntlet
 			{
 				// LogFoo: Display: Some Message
 				// Match how Unreal does not display the level for 'Log' level messages
-				if (Level == LogLevel.Log)
+				if (Level == UnrealLog.LogLevel.Log)
 				{
 					return string.Format("{0}: {1}", LongChannelName, Message);
 				}
@@ -105,6 +128,15 @@ namespace Gauntlet
 				}
 			}
 
+			/// <summary>
+			/// Format the entry without Log level but with log channel
+			/// </summary>
+			/// <returns></returns>
+			public string ToStringWithoutLevel()
+			{
+				return string.Format("{0}: {1}", LongChannelName, Message);
+			}
+
 			/// <summary>
 			/// Constructor that requires all info
 			/// </summary>
