Add logs to debug LocalizationProviderId problems

#Synced-CL 52520240

diff --git a/Engine/Source/Programs/AutomationTool/Localization/LocalizationProvider.cs b/Engine/Source/Programs/AutomationTool/Localization/LocalizationProvider.cs
--- a/Engine/Source/Programs/AutomationTool/Localization/LocalizationProvider.cs
+++ b/Engine/Source/Programs/AutomationTool/Localization/LocalizationProvider.cs
@@ -188,15 +188,31 @@ namespace EpicGames.Localization
 							MethodInfo Method = PotentialLocalizationNodeType.GetMethod("StaticGetLocalizationProviderId");
 							if (Method != null)
 							{
+								string LocalizationProviderId = null;
+
 								try
 								{
-									var LocalizationProviderId = Method.Invoke(null, null) as string;
-									CachedLocalizationProviderTypes.Add(LocalizationProviderId, PotentialLocalizationNodeType);
+									LocalizationProviderId = Method.Invoke(null, null) as string;
 								}
 								catch(Exception ex)
 								{
-									Logger.LogWarning("Type '{Name}' threw when calling its StaticGetLocalizationProviderId method.", PotentialLocalizationNodeType.FullName);
-									Logger.LogWarning("Exception: {Message}", ex.Message);
+									Logger.LogWarning("Type '{Name}' threw when calling its StaticGetLocalizationProviderId method. Exception: {Message}", PotentialLocalizationNodeType.FullName, ex.Message);
+								}
+
+								if (!String.IsNullOrEmpty(LocalizationProviderId))
+								{
+									if(CachedLocalizationProviderTypes.TryAdd(LocalizationProviderId, PotentialLocalizationNodeType))
+									{
+										Logger.LogInformation("Registered '{Name}' in the localization provider cache using id '{Id}'", PotentialLocalizationNodeType.FullName, LocalizationProviderId);
+									}
+									else
+									{
+										Logger.LogWarning("Could not add '{Name}' in the localization provider cache.", PotentialLocalizationNodeType.FullName);
+									}
+								}
+								else
+								{
+									Logger.LogWarning("{Name}.StaticGetLocalizationProviderId method returned an empty ID.", PotentialLocalizationNodeType.FullName);
 								}
 							}
 							else
