[NuGet] Update Microsoft.IdentityModel.* to 8.17.0

#rnx
#Synced-CL 52554515

diff --git a/Engine/Source/Programs/AutomationTool/Apple/Apple.Automation.csproj b/Engine/Source/Programs/AutomationTool/Apple/Apple.Automation.csproj
--- a/Engine/Source/Programs/AutomationTool/Apple/Apple.Automation.csproj
+++ b/Engine/Source/Programs/AutomationTool/Apple/Apple.Automation.csproj
@@ -18,7 +18,7 @@
   </ItemGroup>
   <ItemGroup>
 
-	<PackageReference Include="Microsoft.IdentityModel.Tokens" Version="6.36.0" />
+	<PackageReference Include="Microsoft.IdentityModel.Tokens" Version="8.17.0" />
 	<PackageReference Include="System.IdentityModel.Tokens.Jwt" Version="6.36.0" />
 	<Reference Include="BouncyCastle.Crypto, Version=1.8.3.0, Culture=neutral, PublicKeyToken=0e99375e54769942">
 		<HintPath>..\..\..\..\Binaries\ThirdParty\IOS\BouncyCastle.Crypto.dll</HintPath>
diff --git a/Engine/Source/Programs/AutomationTool/AutomationUtils/CompatibilityShims.cs b/Engine/Source/Programs/AutomationTool/AutomationUtils/CompatibilityShims.cs
new file mode 100644
--- /dev/null
+++ b/Engine/Source/Programs/AutomationTool/AutomationUtils/CompatibilityShims.cs
@@ -0,0 +1,36 @@
+// Copyright Epic Games, Inc. All Rights Reserved.
+
+#nullable enable
+
+using System;
+using System.Collections.Generic;
+using System.Linq;
+
+namespace Microsoft.IdentityModel.Tokens
+{
+	/// <summary>
+	/// Compatibility shim restoring <c>IsNullOrEmpty</c> extension methods that were made internal
+	/// in <c>Microsoft.IdentityModel.Tokens</c> v8.0.0 (see issue #2651).
+	/// </summary>
+	public static class CollectionUtilitiesShim
+	{
+		/// <summary>
+		/// Returns <see langword="true"/> if <paramref name="value"/> is <see langword="null"/> or empty.
+		/// </summary>
+		/// <param name="value">The string to check.</param>
+		public static bool IsNullOrEmpty(this string? value)
+		{
+			return String.IsNullOrEmpty(value);
+		}
+
+		/// <summary>
+		/// Returns <see langword="true"/> if <paramref name="enumerable"/> is <see langword="null"/> or contains no elements.
+		/// </summary>
+		/// <typeparam name="T">The element type.</typeparam>
+		/// <param name="enumerable">The sequence to check.</param>
+		public static bool IsNullOrEmpty<T>(this IEnumerable<T>? enumerable)
+		{
+			return enumerable == null || !enumerable.Any();
+		}
+	}
+}
diff --git a/Engine/Source/Programs/AutomationTool/Gauntlet/Gauntlet.Automation.csproj b/Engine/Source/Programs/AutomationTool/Gauntlet/Gauntlet.Automation.csproj
--- a/Engine/Source/Programs/AutomationTool/Gauntlet/Gauntlet.Automation.csproj
+++ b/Engine/Source/Programs/AutomationTool/Gauntlet/Gauntlet.Automation.csproj
@@ -77,7 +77,7 @@
   <ItemGroup>
     <PackageReference Include="Appium.WebDriver" Version="8.0.0" />
     <PackageReference Include="Magick.NET-Q16-HDRI-AnyCPU" Version="14.11.1" PrivateAssets="all" />
-    <PackageReference Include="Microsoft.IdentityModel.Tokens" Version="6.36.0" />
+    <PackageReference Include="Microsoft.IdentityModel.Tokens" Version="8.17.0" />
     <PackageReference Include="MySql.Data" Version="6.10.9" PrivateAssets="all" />
     <PackageReference Include="Newtonsoft.Json" Version="13.0.4" PrivateAssets="all" />
   </ItemGroup>
diff --git a/Engine/Source/Programs/UnrealCloudDDC/Jupiter.Common/Jupiter.Common.csproj b/Engine/Source/Programs/UnrealCloudDDC/Jupiter.Common/Jupiter.Common.csproj
--- a/Engine/Source/Programs/UnrealCloudDDC/Jupiter.Common/Jupiter.Common.csproj
+++ b/Engine/Source/Programs/UnrealCloudDDC/Jupiter.Common/Jupiter.Common.csproj
@@ -45,10 +45,10 @@
 
     <PackageReference Include="Okta.AspNetCore" Version="4.6.1" />
     <!-- Adding explicit version of identity model to workaround deletion of a internal symbol used by Okta.AspNetCore -->
-    <PackageReference Include="Microsoft.IdentityModel.Tokens" Version="8.0.1" />
-    <PackageReference Include="Microsoft.IdentityModel.Logging" Version="8.0.1" />
-    <PackageReference Include="Microsoft.IdentityModel.Protocols" Version="8.0.1" />
-    <PackageReference Include="Microsoft.IdentityModel.JsonWebTokens" Version="8.0.1" />
+    <PackageReference Include="Microsoft.IdentityModel.Tokens" Version="8.17.0" />
+    <PackageReference Include="Microsoft.IdentityModel.Logging" Version="8.17.0" />
+    <PackageReference Include="Microsoft.IdentityModel.Protocols" Version="8.17.0" />
+    <PackageReference Include="Microsoft.IdentityModel.JsonWebTokens" Version="8.17.0" />
 
     <!--Support reading yaml config files -->
     <PackageReference Include="NetEscapades.Configuration.Yaml" Version="3.1.0" />
