[NuGet] Update Microsoft.Build to 17.11.48, Microsoft.Build.Locator to 1.11.2

#rnx
#Synced-CL 52554515

diff --git a/Engine/Source/Programs/AutomationTool/AutomationTool.csproj b/Engine/Source/Programs/AutomationTool/AutomationTool.csproj
--- a/Engine/Source/Programs/AutomationTool/AutomationTool.csproj
+++ b/Engine/Source/Programs/AutomationTool/AutomationTool.csproj
@@ -85,10 +85,10 @@
     <PackageReference Include="AWSSDK.S3" Version="3.7.501.14" />
     <PackageReference Include="AWSSDK.SecurityToken" Version="3.7.501.14" />
     <PackageReference Include="Magick.NET-Q16-HDRI-AnyCPU" Version="14.11.1" />
-    <PackageReference Include="Microsoft.Build" Version="17.11.48" ExcludeAssets="runtime" PrivateAssets="all">
+    <PackageReference Include="Microsoft.Build" Version="18.4.0" ExcludeAssets="runtime" PrivateAssets="all">
       <IncludeAssets>compile; build; native; contentfiles; analyzers; buildtransitive</IncludeAssets>
     </PackageReference>
-    <PackageReference Include="Microsoft.Build.Locator" Version="1.9.1" />
+    <PackageReference Include="Microsoft.Build.Locator" Version="1.11.2" />
     <!-- Added as a temporary workaround to running from visual studio where the system version of dotnet is run. 
          Adding a loader context resolves many of the issues, but it still fails when MSBuild tries to write the json -->
     <PackageReference Include="Newtonsoft.Json" Version="13.0.4" />
diff --git a/Engine/Source/Programs/Shared/EpicGames.MsBuild/EpicGames.MsBuild.csproj b/Engine/Source/Programs/Shared/EpicGames.MsBuild/EpicGames.MsBuild.csproj
--- a/Engine/Source/Programs/Shared/EpicGames.MsBuild/EpicGames.MsBuild.csproj
+++ b/Engine/Source/Programs/Shared/EpicGames.MsBuild/EpicGames.MsBuild.csproj
@@ -1,7 +1,7 @@
 ﻿<Project Sdk="Microsoft.NET.Sdk">
   <ItemGroup>
-    <PackageReference Include="Microsoft.Build" Version="17.11.48" ExcludeAssets="runtime" />
-    <PackageReference Include="Microsoft.Build.Locator" Version="1.9.1" />
+    <PackageReference Include="Microsoft.Build" Version="18.4.0" ExcludeAssets="runtime" />
+    <PackageReference Include="Microsoft.Build.Locator" Version="1.11.2" />
   </ItemGroup>
 
   <ItemGroup>
