[NuGet] Update System.IO.FileSystem.AccessControl to 5.0.0, System.IO.Hashing to 10.0.5

#rnx
#Synced-CL 52554515

diff --git a/Engine/Source/Programs/AutomationTool/Turnkey/Turnkey.Automation.csproj b/Engine/Source/Programs/AutomationTool/Turnkey/Turnkey.Automation.csproj
--- a/Engine/Source/Programs/AutomationTool/Turnkey/Turnkey.Automation.csproj
+++ b/Engine/Source/Programs/AutomationTool/Turnkey/Turnkey.Automation.csproj
@@ -36,6 +36,6 @@
     <PackageReference Include="Google.Apis.Drive.v3" Version="1.33.1.1232" PrivateAssets="all" />
     <!-- Unused, but required to resolve version conflicts from transitive dependencies. -->
     <PackageReference Include="Newtonsoft.Json" Version="13.0.4" PrivateAssets="all" />
-    <PackageReference Include="System.IO.FileSystem.AccessControl" Version="4.7.0" />
+    <PackageReference Include="System.IO.FileSystem.AccessControl" Version="5.0.0" />
   </ItemGroup>
 </Project>
diff --git a/Engine/Source/Programs/Shared/EpicGames.Serialization/EpicGames.Serialization.csproj b/Engine/Source/Programs/Shared/EpicGames.Serialization/EpicGames.Serialization.csproj
--- a/Engine/Source/Programs/Shared/EpicGames.Serialization/EpicGames.Serialization.csproj
+++ b/Engine/Source/Programs/Shared/EpicGames.Serialization/EpicGames.Serialization.csproj
@@ -1,6 +1,6 @@
 ﻿<Project Sdk="Microsoft.NET.Sdk">
   <ItemGroup>
-    <PackageReference Include="System.IO.Hashing" Version="8.0.0" />
+    <PackageReference Include="System.IO.Hashing" Version="10.0.5" />
   </ItemGroup>
 
   <ItemGroup>
