[NuGet] Update System.Drawing.Common to 10.0.5

#rnx
#Synced-CL 52554515

diff --git a/Engine/Source/Programs/AutomationTool/AutomationUtils/AutomationUtils.Automation.csproj b/Engine/Source/Programs/AutomationTool/AutomationUtils/AutomationUtils.Automation.csproj
--- a/Engine/Source/Programs/AutomationTool/AutomationUtils/AutomationUtils.Automation.csproj
+++ b/Engine/Source/Programs/AutomationTool/AutomationUtils/AutomationUtils.Automation.csproj
@@ -38,7 +38,7 @@
   </ItemGroup>
   <ItemGroup>
     <PackageReference Include="Magick.NET-Q16-HDRI-AnyCPU" Version="14.11.1" PrivateAssets="all" />
-    <PackageReference Include="System.Drawing.Common" Version="8.0.20" />
+    <PackageReference Include="System.Drawing.Common" Version="10.0.5" />
     <!--Required by Ionic -->
     <PackageReference Include="System.Security.Permissions" Version="8.0.0" />
 
diff --git a/Engine/Source/Programs/UnrealToolbox/UnrealToolbox.csproj b/Engine/Source/Programs/UnrealToolbox/UnrealToolbox.csproj
--- a/Engine/Source/Programs/UnrealToolbox/UnrealToolbox.csproj
+++ b/Engine/Source/Programs/UnrealToolbox/UnrealToolbox.csproj
@@ -86,7 +86,7 @@
     <PackageReference Include="FluentAvaloniaUI" Version="2.1.0" />
     <PackageReference Include="CommunityToolkit.Mvvm" Version="8.3.2" />
     <PackageReference Include="DesktopNotifications.Avalonia" Version="1.3.1" />
-    <PackageReference Include="System.Drawing.Common" Version="8.0.20" />
+    <PackageReference Include="System.Drawing.Common" Version="10.0.5" />
   </ItemGroup>
 
 </Project>
\ No newline at end of file
