[NuGet] Update Portable.BouncyCastle to 1.9.0

#rnx
#Synced-CL 52554515

diff --git a/Engine/Source/Programs/AutomationTool/Apple/Apple.Automation.csproj b/Engine/Source/Programs/AutomationTool/Apple/Apple.Automation.csproj
--- a/Engine/Source/Programs/AutomationTool/Apple/Apple.Automation.csproj
+++ b/Engine/Source/Programs/AutomationTool/Apple/Apple.Automation.csproj
@@ -23,7 +23,7 @@
 	<Reference Include="BouncyCastle.Crypto, Version=1.8.3.0, Culture=neutral, PublicKeyToken=0e99375e54769942">
 		<HintPath>..\..\..\..\Binaries\ThirdParty\IOS\BouncyCastle.Crypto.dll</HintPath>
 	</Reference>
-	<PackageReference Include="Portable.BouncyCastle" Version="1.8.10" />
+	<PackageReference Include="Portable.BouncyCastle" Version="1.9.0" />
 
     <ProjectReference Include="..\..\Shared\EpicGames.Core\EpicGames.Core.csproj" PrivateAssets="All"><Private>false</Private></ProjectReference>
     <ProjectReference Include="..\..\UnrealBuildTool\UnrealBuildTool.csproj" PrivateAssets="All"><Private>false</Private></ProjectReference>
diff --git a/Engine/Source/Programs/AutomationTool/IOS/IOS.Automation.csproj b/Engine/Source/Programs/AutomationTool/IOS/IOS.Automation.csproj
--- a/Engine/Source/Programs/AutomationTool/IOS/IOS.Automation.csproj
+++ b/Engine/Source/Programs/AutomationTool/IOS/IOS.Automation.csproj
@@ -36,6 +36,6 @@
   </ItemGroup>
   <ItemGroup>
     <PackageReference Include="System.Resources.Extensions" Version="10.0.5" PrivateAssets="all" />
-    <PackageReference Include="Portable.BouncyCastle" Version="1.8.10" PrivateAssets="all" />
+    <PackageReference Include="Portable.BouncyCastle" Version="1.9.0" PrivateAssets="all" />
   </ItemGroup>
 </Project>
\ No newline at end of file
