Add assembly information to localization provider logs to understand some errors

#Synced-CL 52555821
#Synced-CL 52555821

diff --git a/Engine/Source/Programs/AutomationTool/Localization/LocalizationProvider.cs b/Engine/Source/Programs/AutomationTool/Localization/LocalizationProvider.cs
--- a/Engine/Source/Programs/AutomationTool/Localization/LocalizationProvider.cs
+++ b/Engine/Source/Programs/AutomationTool/Localization/LocalizationProvider.cs
@@ -196,28 +196,28 @@ namespace EpicGames.Localization
 								}
 								catch(Exception ex)
 								{
-									Logger.LogWarning("Type '{Name}' threw when calling its StaticGetLocalizationProviderId method. Exception: {Message}", PotentialLocalizationNodeType.FullName, ex.Message);
+									Logger.LogWarning("Type '{Name}' threw when calling its StaticGetLocalizationProviderId method. Exception: {Message}. (Assembly={Assembly}, AssemblyQualifiedName={AssemblyQualifiedName})", PotentialLocalizationNodeType.FullName, ex.Message, PotentialLocalizationNodeType.Assembly.FullName, PotentialLocalizationNodeType.AssemblyQualifiedName);
 								}
 
 								if (!String.IsNullOrEmpty(LocalizationProviderId))
 								{
 									if(CachedLocalizationProviderTypes.TryAdd(LocalizationProviderId, PotentialLocalizationNodeType))
 									{
-										Logger.LogInformation("Registered '{Name}' in the localization provider cache using id '{Id}'", PotentialLocalizationNodeType.FullName, LocalizationProviderId);
+										Logger.LogInformation("Registered '{Name}' in the localization provider cache using id '{Id}'. (Assembly={Assembly}, AssemblyQualifiedName={AssemblyQualifiedName})", PotentialLocalizationNodeType.FullName, LocalizationProviderId, PotentialLocalizationNodeType.Assembly.FullName, PotentialLocalizationNodeType.AssemblyQualifiedName);
 									}
 									else
 									{
-										Logger.LogWarning("Could not add '{Name}' in the localization provider cache.", PotentialLocalizationNodeType.FullName);
+										Logger.LogWarning("Could not add '{Name}' in the localization provider cache using id '{Id}'. (Assembly={Assembly}, AssemblyQualifiedName={AssemblyQualifiedName})", PotentialLocalizationNodeType.FullName, LocalizationProviderId, PotentialLocalizationNodeType.Assembly.FullName, PotentialLocalizationNodeType.AssemblyQualifiedName);
 									}
 								}
 								else
 								{
-									Logger.LogWarning("{Name}.StaticGetLocalizationProviderId method returned an empty ID.", PotentialLocalizationNodeType.FullName);
+									Logger.LogWarning("{Name}.StaticGetLocalizationProviderId method returned an empty ID. (Assembly={Assembly}, AssemblyQualifiedName={AssemblyQualifiedName})", PotentialLocalizationNodeType.FullName, PotentialLocalizationNodeType.Assembly.FullName, PotentialLocalizationNodeType.AssemblyQualifiedName);
 								}
 							}
 							else
 							{
-								Logger.LogWarning("Type '{Name}' derives from LocalizationProvider but is missing its StaticGetLocalizationProviderId method.", PotentialLocalizationNodeType.FullName);
+								Logger.LogWarning("Type '{Name}' derives from LocalizationProvider but is missing its StaticGetLocalizationProviderId method. (Assembly={Assembly}, AssemblyQualifiedName={AssemblyQualifiedName})", PotentialLocalizationNodeType.FullName, PotentialLocalizationNodeType.Assembly.FullName, PotentialLocalizationNodeType.AssemblyQualifiedName);
 							}
 						}
 					}
