Fix launching iOS on Mac from Editor when path to .app has spaces. Shell quote as appropriate.

#Synced-CL 52308170

diff --git a/Engine/Source/Programs/AutomationTool/IOS/IOSPlatform.Automation.cs b/Engine/Source/Programs/AutomationTool/IOS/IOSPlatform.Automation.cs
--- a/Engine/Source/Programs/AutomationTool/IOS/IOSPlatform.Automation.cs
+++ b/Engine/Source/Programs/AutomationTool/IOS/IOSPlatform.Automation.cs
@@ -1308,7 +1308,7 @@ public class IOSPlatform : ApplePlatform
 				string AppName = Path.GetFileNameWithoutExtension(ClientApp);
 				LaunchTool = IOSLaunchTool.Open;
 				Program = "open";
-				Arguments = IOSonMacDir + "/" + AppName + ".app --args";
+				Arguments = "\"" + IOSonMacDir + "/" + AppName + ".app\" --args";
 			}
 			else
 			{
