Removed "Hit Enter" text for Android Studio install from turnkey

#android
#Synced-CL 52546546

diff --git a/Engine/Source/Programs/AutomationTool/Android/AndroidPlatform.Automation.cs b/Engine/Source/Programs/AutomationTool/Android/AndroidPlatform.Automation.cs
--- a/Engine/Source/Programs/AutomationTool/Android/AndroidPlatform.Automation.cs
+++ b/Engine/Source/Programs/AutomationTool/Android/AndroidPlatform.Automation.cs
@@ -463,7 +463,7 @@ public class AndroidPlatform : Platform
 
 		if (!bHaveAndroidStudio)
 		{
-			TurnkeyContext.PauseForUser("Android Studio was not found on this machine. Press Enter to download and install Android Studio which is required to use Android.");
+			TurnkeyContext.PauseForUser("Android Studio was not found on this machine, which is required to use Android. Download and install Android Studio?");
 
 			// get AS installer
 			string OutputPath = TurnkeyContext.RetrieveFileSource("AndroidStudio");
