Fix to mac/linux turnkey cmdline-tools auto install missing a gui command-line window

#android
#Synced-CL 52654288

diff --git a/Engine/Source/Programs/AutomationTool/Android/AndroidPlatform.Automation.cs b/Engine/Source/Programs/AutomationTool/Android/AndroidPlatform.Automation.cs
--- a/Engine/Source/Programs/AutomationTool/Android/AndroidPlatform.Automation.cs
+++ b/Engine/Source/Programs/AutomationTool/Android/AndroidPlatform.Automation.cs
@@ -213,25 +213,67 @@ public class AndroidPlatform : Platform
 				if (JavaHome == null)
 				{
 					string StudioExe = GetAndroidStudioExe();
-					string StudioDir = OperatingSystem.IsMacOS() ? Path.Combine(StudioExe, "Contents") : Path.GetDirectoryName(Path.GetDirectoryName(StudioExe))!;
-					Environment.SetEnvironmentVariable("JAVA_HOME", Path.Combine(StudioDir, "jbr"));
+					JavaHome = OperatingSystem.IsMacOS() ? Path.Combine(StudioExe, "Contents", "jbr", "Contents", "Home") : Path.Combine(Path.GetDirectoryName(Path.GetDirectoryName(StudioExe))!, "jbr");
 				}
-
-				try
+				
+				int ExitCode;
+				if (!OperatingSystem.IsWindows())
 				{
-					// This may bring up a license acceptance message that needs the user to respond, so we make a new window
-					int ExitCode = TurnkeyContext.RunExternalCommand(SdkManager, $"--sdk_root=\"{SdkDir}\" \"cmdline-tools;latest\"", false, bUnattended, true);
-					if (ExitCode == 0)
+					string BashCmd = $"export JAVA_HOME=\"{JavaHome}\" && \"{SdkManager}\" --sdk_root=\"{SdkDir}\" \"cmdline-tools;latest\"";
+					string Command;
+					string Arguments;
+					if (OperatingSystem.IsMacOS())
 					{
-						return true;
+						Command = "osascript";
+						Arguments = 
+							"-e \"tell application \\\"Terminal\\\"\" " +
+							"-e \"activate\" " +
+							$"-e \"do script \\\"{BashCmd.Replace("\"", "\\\\\\\"")}; exit\\\"\" " +
+							"-e \"set theTab to the result\" " +
+							"-e \"repeat\" " +
+							"-e \"delay 0.5\" " +
+							"-e \"if not busy of theTab then exit repeat\" " +
+							"-e \"end repeat\" " +
+							"-e \"close (every window whose selected tab is theTab)\" " +
+							"-e \"end tell\"";
 					}
-				}
-				finally
-				{
-					if (JavaHome == null)
+					else
 					{
-						Environment.SetEnvironmentVariable("JAVA_HOME", null);
+						Process Which;
+						try
+						{
+							Which = Process.Start(new ProcessStartInfo { FileName = "which", Arguments = "gnome-terminal", RedirectStandardOutput = true, UseShellExecute = false });
+							Which?.WaitForExit();
+						}
+						catch (Exception)
+						{
+							Which = null;
+						}
+
+						if (Which is { ExitCode: 0 })
+						{
+							Command = "gnome-terminal";
+							Arguments = $"--wait -- bash -c \"{BashCmd.Replace("\"", "\\\"")}\"";
+						}
+						else
+						{
+							Command = "x-terminal-emulator";
+							Arguments = $"-e bash -c \"{BashCmd.Replace("\"", "\\\"")}\"";
+						}
 					}
+					
+					ExitCode = TurnkeyContext.RunExternalCommand(Command, Arguments, false, bUnattended, false);
+				}
+				else
+				{
+					string BatchCmd = $"set \"JAVA_HOME={JavaHome}\" && \"{SdkManager}\" --sdk_root=\"{SdkDir}\" \"cmdline-tools;latest\"";
+					string Arguments = $"/c {BatchCmd}";
+					ExitCode = TurnkeyContext.RunExternalCommand("cmd.exe", Arguments, false, bUnattended, true);
+				}
+
+				if (ExitCode == 0)
+				{
+					return true;
 				}
 			}
 
