CookedEditor.Automation: Make checking for file extensions consistently case-insensitive when adding UFS files to a list.

#14589

#Synced-CL 52673681

diff --git a/Engine/Source/Programs/AutomationTool/CookedEditor/CookedEditor.Automation.cs b/Engine/Source/Programs/AutomationTool/CookedEditor/CookedEditor.Automation.cs
--- a/Engine/Source/Programs/AutomationTool/CookedEditor/CookedEditor.Automation.cs
+++ b/Engine/Source/Programs/AutomationTool/CookedEditor/CookedEditor.Automation.cs
@@ -302,7 +302,7 @@ public class ModifyStageContext
 	{
 		// look in SC and UFSFiles
 		FileList.AddRange(SC.FilesToStage.UFSFiles.Keys.Where(x => x.Name.EndsWith(Extension, StringComparison.OrdinalIgnoreCase)).Select(y => DeploymentContext.UnmakeRelativeStagedReference(SC, y)));
-		FileList.AddRange(UFSFilesToStage.Where(x => x.FullName.EndsWith(Extension, StringComparison.InvariantCultureIgnoreCase)));
+		FileList.AddRange(UFSFilesToStage.Where(x => x.FullName.EndsWith(Extension, StringComparison.OrdinalIgnoreCase)));
 	}
 
 	private void AddUFSFilesToList(List<FileReference> FileList, string Extension, DeploymentContext SC)
@@ -310,12 +310,12 @@ public class ModifyStageContext
 		// look in SC and UFSFiles
 		foreach (var Pair in SC.FilesToStage.UFSFiles)
 		{
-			if (Pair.Key.Name.EndsWith(Extension))
+			if (Pair.Key.Name.EndsWith(Extension, StringComparison.OrdinalIgnoreCase))
 			{
 				FileList.Add(Pair.Value);
 			}
 		}
-		FileList.AddRange(UFSFilesToStage.Where(x => x.FullName.EndsWith(Extension, StringComparison.InvariantCultureIgnoreCase)));
+		FileList.AddRange(UFSFilesToStage.Where(x => x.FullName.EndsWith(Extension, StringComparison.OrdinalIgnoreCase)));
 	}
 
 	private void UncookMaps(DeploymentContext SC)
