Android.Automation: Only add files to a single overflow obb file.

#13680

#Synced-CL 52673681

diff --git a/Engine/Source/Programs/AutomationTool/Android/AndroidPlatform.Automation.cs b/Engine/Source/Programs/AutomationTool/Android/AndroidPlatform.Automation.cs
--- a/Engine/Source/Programs/AutomationTool/Android/AndroidPlatform.Automation.cs
+++ b/Engine/Source/Programs/AutomationTool/Android/AndroidPlatform.Automation.cs
@@ -1398,7 +1398,7 @@ public class AndroidPlatform : Platform
 						else if (AllowOverflowOBBLimit > 0)
 						{
 							// find an overflow we can fit this file
-							bool bFoundOverflow = false;				
+							bool bFoundOverflow = false;
 							foreach (OverflowFileInfo OverflowRef in OverflowInfos)
 							{
 								if (OverflowRef.OverflowObbSize + FileRequirements < OBBSizeAllowed)
@@ -1406,6 +1406,7 @@ public class AndroidPlatform : Platform
 									OverflowRef.FilesForOverflow.Add(FileRef);
 									OverflowRef.OverflowObbSize += FileRequirements;
 									bFoundOverflow = true;
+									break;
 								}
 							}
 
