Automation: Fix some automation static analysis warnings.

#Synced-CL 52673681

diff --git a/Engine/Source/Programs/AutomationTool/AutomationUtils/BuildCommand.cs b/Engine/Source/Programs/AutomationTool/AutomationUtils/BuildCommand.cs
--- a/Engine/Source/Programs/AutomationTool/AutomationUtils/BuildCommand.cs
+++ b/Engine/Source/Programs/AutomationTool/AutomationUtils/BuildCommand.cs
@@ -74,7 +74,7 @@ namespace AutomationTool
 		}
 
 		/// <summary>
-		/// Parses the command's Params list for a property set in the format set:<Param>=<Value> and reads its value. 
+		/// Parses the command's Params list for a property set in the format set:&lt;Param&gt;=&lt;Value&gt; and reads its value. 
 		/// Ex. ParseSetPropertyValue("projectName", "") will return "Lyra" for commandline "BuildGraph -set:ProjectName=Lyra"
 		/// </summary>
 		/// <param name="propertyName">PropertyName to read its value.</param>
@@ -310,6 +310,10 @@ namespace AutomationTool
 			return projectFullPath;
 		}
 
+		/// <summary>
+		/// Attempts to parse the project param. Returns null if no project is found, and throws if the specified project does not exist on disk.
+		/// </summary>
+		/// <returns>FileReference to uproject</returns>
 		public FileReference? ParseProjectParam()
 		{
 			FileReference? projectFullPath = null;
diff --git a/Engine/Source/Programs/AutomationTool/AutomationUtils/DeviceReservation.cs b/Engine/Source/Programs/AutomationTool/AutomationUtils/DeviceReservation.cs
--- a/Engine/Source/Programs/AutomationTool/AutomationUtils/DeviceReservation.cs
+++ b/Engine/Source/Programs/AutomationTool/AutomationUtils/DeviceReservation.cs
@@ -48,7 +48,7 @@ namespace AutomationTool.DeviceReservation
 
 		/// <summary>
 		/// Whether the reservation requires an installation
-		/// <summary>
+		/// </summary>
 		public bool? InstallRequired => ActiveReservation?.InstallRequired;
 
 		public IReadOnlyList<Device> Devices
diff --git a/Engine/Source/Programs/AutomationTool/AutomationUtils/ZenUtils.cs b/Engine/Source/Programs/AutomationTool/AutomationUtils/ZenUtils.cs
--- a/Engine/Source/Programs/AutomationTool/AutomationUtils/ZenUtils.cs
+++ b/Engine/Source/Programs/AutomationTool/AutomationUtils/ZenUtils.cs
@@ -124,7 +124,7 @@ namespace AutomationUtils
 				using(FileStream Stream = File.Open(LockFilePath, FileMode.Open, FileAccess.Read, FileShare.ReadWrite | FileShare.Delete))
 				{
 					CbObjectData = new byte[Stream.Length];
-					Stream.Read(CbObjectData);
+					Stream.ReadExactly(CbObjectData);
 				}
 			}
 			catch(IOException)
