Appium Auto Accept Alerts
Fixing screenshot upload issues
Mitigating other small WDA failures
#Synced-CL 52636240

diff --git a/Engine/Source/Programs/AutomationTool/Gauntlet/Platform/IOS/Gauntlet.IOSAppium.cs b/Engine/Source/Programs/AutomationTool/Gauntlet/Platform/IOS/Gauntlet.IOSAppium.cs
--- a/Engine/Source/Programs/AutomationTool/Gauntlet/Platform/IOS/Gauntlet.IOSAppium.cs
+++ b/Engine/Source/Programs/AutomationTool/Gauntlet/Platform/IOS/Gauntlet.IOSAppium.cs
@@ -70,11 +70,10 @@ namespace Gauntlet
 			/// </summary>
 			public string WdaBundleId { get; set; }
 
-
 			/// <summary>
-			/// Launch Timeout for WebDriverAgent app in ms. Default: 60000ms
+			/// Launch Timeout for WebDriverAgent app in ms. Default: 120000ms (increased for VM environments)
 			/// </summary>
-			public int WdaLaunchTimeoutMs { get; set; } = 60000;
+			public int WdaLaunchTimeoutMs { get; set; } = 120000;
 
 			/// <summary>
 			/// How often to take screenshots in seconds. Disabled if <= 0
@@ -99,8 +98,7 @@ namespace Gauntlet
 				Capabilities.AutomationName = AutomationName;
 				Capabilities.AddAdditionalAppiumOption("udid", UUID);
 				Capabilities.AddAdditionalAppiumOption("bundleId", PackageName);
-				// Use autoDismissAlerts instead of autoAcceptAlerts to click Cancel/Dismiss buttons
-				Capabilities.AddAdditionalAppiumOption("autoDismissAlerts", true);
+				Capabilities.AddAdditionalAppiumOption("autoAcceptAlerts", true);
 				Capabilities.AddAdditionalAppiumOption("xcodeOrgId", OrgId);
 				Capabilities.AddAdditionalAppiumOption("xcodeSigningId", SigningId);
 				Capabilities.AddAdditionalAppiumOption("autoLaunch", false);
@@ -223,6 +221,14 @@ namespace Gauntlet
 
 				}
 				Stop();
+
+				// Ensure monitor thread is stopped (fallback if SaveArtifacts didn't call it)
+				if (MonitorThread != null)
+				{
+					Log.Verbose("Stopping monitor thread during disposal (screenshots will not be saved)");
+					StopMonitorLoop(); // No output directory - screenshots discarded
+				}
+
 				Disposed = true;
 			}
 		}
@@ -292,6 +298,23 @@ namespace Gauntlet
 							Server = null;
 						}
 
+						// Refresh device tunnel to ensure it's active BEFORE starting Appium (critical for VM environments)
+						// This helps prevent "Unknown device" errors when Appium tries to connect
+						if (TargetDeviceIOS.UseDeviceCtl)
+						{
+							try
+							{
+								Log.Info("Refreshing device tunnel for {0} before Appium server starts...", UUID);
+								IProcessResult pairResult = TargetDeviceIOS.ExecuteDevicectlCommand("manage pair", UUID, 30, AdditionalOptions: ERunOptions.NoLoggingOfRunCommand);
+								Log.Verbose("Waiting 5 seconds for device tunnel to stabilize...");
+								Thread.Sleep(5000); // Wait for tunnel to stabilize before Appium queries devices
+							}
+							catch (Exception ex)
+							{
+								Log.Verbose("Device pair refresh warning: {0}", ex.Message);
+							}
+						}
+
 						// Start Appium Server
 						Log.Info("Starting Appium server for device {0} (attempt {1}/{2})", UUID, attemptCount, MaxRetries);
 						AppiumServiceBuilder ServerBuilder = new AppiumServiceBuilder()
@@ -372,10 +395,8 @@ namespace Gauntlet
 				Server = null;
 			}
 
-			if (MonitorThread != null)
-			{
-				StopMonitorLoop();
-			}
+			// Note: MonitorThread cleanup is handled by StopMonitorLoop() which should be called
+			// with an output directory parameter to save screenshots. See SaveArtifacts() in TargetDeviceIOS.
 
 			// TODO: Terminate WDA app - parse devicectl for pid, then kill?
 		}
@@ -712,14 +733,30 @@ namespace Gauntlet
 			}
 
 			string FilePath = Path.Combine(OutputDirectory, ScreenshotName);
-			try
-			{
-				(Driver as ITakesScreenshot).GetScreenshot().SaveAsFile(FilePath);
-				Log.Verbose("Saved screenshot to: {ScreenshotFile}", FilePath);
-			}
-			catch (Exception Ex)
+
+			// Retry logic for screenshot capture (helps with socket hang up issues on VMs)
+			const int maxRetries = 3;
+			for (int attempt = 0; attempt < maxRetries; attempt++)
 			{
-				Log.Info("Encountered an {ExceptionType} when trying to capture a screenshot: {ExceptionMessage}", Ex.GetType().Name, Ex.Message);
+				try
+				{
+					(Driver as ITakesScreenshot).GetScreenshot().SaveAsFile(FilePath);
+					Log.Verbose("Saved screenshot to: {ScreenshotFile}", FilePath);
+					break; // Success - exit retry loop
+				}
+				catch (Exception Ex)
+				{
+					if (Ex.Message.Contains("socket hang up") && attempt < maxRetries - 1)
+					{
+						Log.Verbose("Screenshot attempt {0} failed with socket hang up, retrying...", attempt + 1);
+						Thread.Sleep(1000); // Wait 1 second before retry
+					}
+					else
+					{
+						Log.Info("Encountered an {ExceptionType} when trying to capture a screenshot: {ExceptionMessage}", Ex.GetType().Name, Ex.Message);
+						break; // Non-retryable error or final attempt
+					}
+				}
 			}
 		}
 
