Replaced some chmod with File.SetUnixFileMode

#android
#Synced-CL 52700724

diff --git a/Engine/Source/Programs/AutomationTool/Android/AndroidPlatform.Automation.cs b/Engine/Source/Programs/AutomationTool/Android/AndroidPlatform.Automation.cs
--- a/Engine/Source/Programs/AutomationTool/Android/AndroidPlatform.Automation.cs
+++ b/Engine/Source/Programs/AutomationTool/Android/AndroidPlatform.Automation.cs
@@ -1619,7 +1619,7 @@ public class AndroidPlatform : Platform
 			GetPlatformInstallOptions(SC, out bNeedsPCInstall, out bNeedsMacInstall, out bNeedsLinuxInstall);
 
 			//helper delegate to prevent code duplication but allow us access to all the local variables we need
-			var CreateInstallFilesAction = new Action<UnrealTargetPlatform>(Target =>
+			void CreateInstallFilesAction(UnrealTargetPlatform Target)
 			{
 				bool bIsPC = (Target == UnrealTargetPlatform.Win64);
 				string LineEnding = bIsPC ? "\r\n" : "\n";
@@ -1663,39 +1663,39 @@ public class AndroidPlatform : Platform
 				{
 					if (!Deploy.bNoObbs && bHaveAPK)
 					{
-						CommandUtils.FixUnixFilePermissions(BatchName);
+						File.SetUnixFileMode(BatchName, File.GetUnixFileMode(BatchName) | UnixFileMode.UserExecute | UnixFileMode.GroupExecute | UnixFileMode.OtherExecute);
 					}
 					if (bHaveAPK || bHaveUniversal)
 					{
-						CommandUtils.FixUnixFilePermissions(UninstallBatchName);
+						File.SetUnixFileMode(UninstallBatchName, File.GetUnixFileMode(UninstallBatchName) | UnixFileMode.UserExecute | UnixFileMode.GroupExecute | UnixFileMode.OtherExecute);
 					}
 					if (bHaveUniversal)
 					{
-						CommandUtils.FixUnixFilePermissions(UniversalBatchName);
+						File.SetUnixFileMode(UniversalBatchName, File.GetUnixFileMode(UniversalBatchName) | UnixFileMode.UserExecute | UnixFileMode.GroupExecute | UnixFileMode.OtherExecute);
 					}
 					// Deprecated check of bBuildWithHiddenSymbolVisibility. Remove in 5.8+
 					if (bBuildWithHiddenSymbolVisibility || bSaveSymbols)
 					{
-						CommandUtils.FixUnixFilePermissions(SymbolizeBatchName);
+						File.SetUnixFileMode(SymbolizeBatchName, File.GetUnixFileMode(SymbolizeBatchName) | UnixFileMode.UserExecute | UnixFileMode.GroupExecute | UnixFileMode.OtherExecute);
 					}
 					//if(File.Exists(NoInstallBatchName)) 
 					//{
-					//    CommandUtils.FixUnixFilePermissions(NoInstallBatchName);
+					//    File.SetUnixFileMode(NoInstallBatchName, File.GetUnixFileMode(NoInstallBatchName) | UnixFileMode.UserExecute | UnixFileMode.GroupExecute | UnixFileMode.OtherExecute);
 					//}
 				}
-			});
+			}
 
 			if (bNeedsPCInstall)
 			{
-				CreateInstallFilesAction.Invoke(UnrealTargetPlatform.Win64);
+				CreateInstallFilesAction(UnrealTargetPlatform.Win64);
 			}
 			if (bNeedsMacInstall)
 			{
-				CreateInstallFilesAction.Invoke(UnrealTargetPlatform.Mac);
+				CreateInstallFilesAction(UnrealTargetPlatform.Mac);
 			}
 			if (bNeedsLinuxInstall)
 			{
-				CreateInstallFilesAction.Invoke(UnrealTargetPlatform.Linux);
+				CreateInstallFilesAction(UnrealTargetPlatform.Linux);
 			}
 
 			// If we aren't packaging data in the APK then lets write out a bat file to also let us test without the OBB
diff --git a/Engine/Source/Programs/UnrealBuildTool/Platform/Android/UEDeployAndroid.cs b/Engine/Source/Programs/UnrealBuildTool/Platform/Android/UEDeployAndroid.cs
--- a/Engine/Source/Programs/UnrealBuildTool/Platform/Android/UEDeployAndroid.cs
+++ b/Engine/Source/Programs/UnrealBuildTool/Platform/Android/UEDeployAndroid.cs
@@ -5262,7 +5262,8 @@ popd
 					if (!OperatingSystem.IsWindows())
 					{
 						// fix permissions for Mac/Linux
-						RunCommandLineProgramWithException(UnrealBuildGradlePath, "/bin/sh", String.Format("-c 'chmod 0755 \"{0}\"'", GradleScriptPath.Replace("'", "'\"'\"'", StringComparison.Ordinal)), Logger, "Fix gradlew permissions");
+						Logger.LogInformation("Fix gradlew permissions");
+						File.SetUnixFileMode(GradleScriptPath, File.GetUnixFileMode(GradleScriptPath) | UnixFileMode.UserExecute | UnixFileMode.GroupExecute | UnixFileMode.OtherExecute);
 					}
 					else
 					{
@@ -5717,7 +5718,8 @@ popd
 					if (!OperatingSystem.IsWindows())
 					{
 						// fix permissions for Mac/Linux
-						RunCommandLineProgramWithException(UnrealBuildGradlePath, "/bin/sh", String.Format("-c 'chmod 0755 \"{0}\"'", GradleScriptPath.Replace("'", "'\"'\"'", StringComparison.Ordinal)), Logger, "Fix gradlew permissions");
+						Logger.LogInformation("Fix gradlew permissions");
+						File.SetUnixFileMode(GradleScriptPath, File.GetUnixFileMode(GradleScriptPath) | UnixFileMode.UserExecute | UnixFileMode.GroupExecute | UnixFileMode.OtherExecute);
 					}
 					else
 					{
@@ -5779,7 +5781,8 @@ popd
 						if (!OperatingSystem.IsWindows())
 						{
 							// fix permissions for Mac/Linux
-							RunCommandLineProgramWithException(UnrealBuildGradlePath, "/bin/sh", String.Format("-c 'chmod 0755 \"{0}\"'", GradleScriptPath.Replace("'", "'\"'\"'", StringComparison.Ordinal)), Logger, "Fix gradlew permissions");
+							Logger.LogInformation("Fix gradlew permissions");
+							File.SetUnixFileMode(GradleScriptPath, File.GetUnixFileMode(GradleScriptPath) | UnixFileMode.UserExecute | UnixFileMode.GroupExecute | UnixFileMode.OtherExecute);
 						}
 						else
 						{
