Update GetListenAddressAndPort to work on mac agents
Get IPAddress from Megazord
Add IPAddress lookup for ios in RPC Smoke
#Synced-CL 52690152

#Synced-CL 52690152

diff --git a/Engine/Source/Programs/AutomationTool/Gauntlet/Framework/RpcFramework/Gauntlet.RpcExecutor.cs b/Engine/Source/Programs/AutomationTool/Gauntlet/Framework/RpcFramework/Gauntlet.RpcExecutor.cs
--- a/Engine/Source/Programs/AutomationTool/Gauntlet/Framework/RpcFramework/Gauntlet.RpcExecutor.cs
+++ b/Engine/Source/Programs/AutomationTool/Gauntlet/Framework/RpcFramework/Gauntlet.RpcExecutor.cs
@@ -10,6 +10,7 @@ using System.IO;
 using System.Linq;
 using System.Net;
 using System.Net.Http;
+using System.Net.NetworkInformation;
 using System.Net.Sockets;
 using System.Runtime.CompilerServices;
 using System.Text;
@@ -125,7 +126,7 @@ namespace Gauntlet
 					if (ExecutorInstance == null)
 					{
 						ExecutorInstance = new RpcExecutor();
-						ExecutorClient.Timeout = Globals.Params.ParseParam("NoRPCHttpTimeout") ? Timeout.InfiniteTimeSpan 
+						ExecutorClient.Timeout = Globals.Params.ParseParam("NoRPCHttpTimeout") ? Timeout.InfiniteTimeSpan
 							: TimeSpan.FromMilliseconds(Globals.Params.ParseValue("RPCHttpTimeoutInMilliseconds", 100 * 1000));
 					}
 					return ExecutorInstance;
@@ -162,15 +163,29 @@ namespace Gauntlet
 		public static string GetListenAddressAndPort()
 		{
 			string LocalIP = "127.0.0.1";
-			IPHostEntry HostEntry = Dns.GetHostEntry(Dns.GetHostName());
-			foreach (var TargetIp in HostEntry.AddressList)
+			try
 			{
-				if (TargetIp.AddressFamily == AddressFamily.InterNetwork)
+				IEnumerable<NetworkInterface> NetworkInterfaces = NetworkInterface.GetAllNetworkInterfaces()
+					.Where(ni => ni.OperationalStatus == OperationalStatus.Up &&
+								ni.NetworkInterfaceType != NetworkInterfaceType.Loopback);
+
+				foreach (var ni in NetworkInterfaces)
 				{
-					LocalIP = TargetIp.ToString();
-					break;
+					IPInterfaceProperties IpProperties = ni.GetIPProperties();
+					UnicastIPAddressInformation Ipv4Address = IpProperties.UnicastAddresses
+						.FirstOrDefault(ip => ip.Address.AddressFamily == AddressFamily.InterNetwork);
+
+					if (Ipv4Address != null)
+					{
+						LocalIP = Ipv4Address.Address.ToString();
+						break;
+					}
 				}
 			}
+			catch (Exception ex)
+			{
+				Log.Verbose($"Failed to get network IP, using loopback: {ex.Message}");
+			}
 			return String.Format("{0}:{1}", LocalIP, GauntletRPCListenPort);
 		}
 
@@ -178,13 +193,13 @@ namespace Gauntlet
 		protected ConcurrentDictionary<GauntletRpcDelegateRegistryFilter, GauntletIncomingRpcMessageDelegate> MessageDelegates = [];
 		protected ConcurrentDictionary<RouteMapping, GauntletHttpRouteDelegate> HttpRoutes = [];
 
-		protected object ReceivedTerminationCallLock = new ();
+		protected object ReceivedTerminationCallLock = new();
 		private bool _bHasReceivedTerminationCall = false;
 		public bool bHasReceivedTerminationCall
 		{
 			get
 			{
-				lock(ReceivedTerminationCallLock)
+				lock (ReceivedTerminationCallLock)
 				{
 					return _bHasReceivedTerminationCall;
 				}
@@ -241,7 +256,7 @@ namespace Gauntlet
 				else
 				{
 					RpcTarget Target = ActiveTargetRegistry.GetTarget(IncomingMessage.SenderId);
-					if( null == Target )
+					if (null == Target)
 					{
 						Log.Warning($"No RPC target registered for '{IncomingMessage.SenderId}' in active registry; routing to central queue to prevent message loss (Category: '{IncomingMessage.Category}').");
 						IncomingMessageQueue.Enqueue(IncomingMessage);
@@ -330,8 +345,8 @@ namespace Gauntlet
 				}
 			}
 
-			try 
-			{ 
+			try
+			{
 				using (HttpListenerResponse ListenerResponse = ListenerContext.Response)
 				{
 					ListenerResponse.Headers.Set("Content-Type", "text/plain");
@@ -365,7 +380,7 @@ namespace Gauntlet
 		{
 			List<GauntletIncomingRpcMessage> NewMessages = [];
 			int TargetsMessageCount = 0;
-			if(null != ActiveTargetRegistry)
+			if (null != ActiveTargetRegistry)
 			{
 				foreach (RpcTarget Target in ActiveTargetRegistry.RpcTargets.Values)
 				{
@@ -400,9 +415,9 @@ namespace Gauntlet
 			{
 				if (bClearMessageQueue)
 				{
-					for(int MessageIndex = 0; MessageIndex < InitialCount; ++MessageIndex)
+					for (int MessageIndex = 0; MessageIndex < InitialCount; ++MessageIndex)
 					{
-						if(Queue.TryDequeue(out GauntletIncomingRpcMessage NewMessage))
+						if (Queue.TryDequeue(out GauntletIncomingRpcMessage NewMessage))
 						{
 							Messages.Add(new GauntletIncomingRpcMessage(NewMessage));
 						}
@@ -531,12 +546,12 @@ namespace Gauntlet
 				{
 					++MatchesFound;
 				}
-				else if(bSaveUnmatchedMessages)
+				else if (bSaveUnmatchedMessages)
 				{
 					UnmatchedMessages.Add(Message);
 				}
 			}
-			
+
 			foreach (var Message in UnmatchedMessages)
 			{
 				MessageQueue.Enqueue(Message);
@@ -645,7 +660,7 @@ namespace Gauntlet
 		{
 			List<GauntletRpcDelegateRegistryFilter> NewFilters = new List<GauntletRpcDelegateRegistryFilter>();
 
-			foreach(RpcTarget Target in InTargets)
+			foreach (RpcTarget Target in InTargets)
 			{
 				NewFilters.Add(new GauntletRpcDelegateRegistryFilter(Target.TargetName, TargetMessageCategory, TargetMessagePayload, false));
 			}
@@ -738,7 +753,7 @@ namespace Gauntlet
 				}
 				return true;
 			}
-		} 
+		}
 
 		/// <summary>
 		/// Delegate type for incoming messages from Unreal processes.
@@ -760,8 +775,8 @@ namespace Gauntlet
 				Log.Warning("Cannot register Delegate with no filter at all.");
 				return;
 			}
-			
-			if(MessageDelegates.TryAdd(new GauntletRpcDelegateRegistryFilter(SenderFilter, CategoryFilter, PayloadFilter, bDeregisterAfterTrigger), FunctionToRun))
+
+			if (MessageDelegates.TryAdd(new GauntletRpcDelegateRegistryFilter(SenderFilter, CategoryFilter, PayloadFilter, bDeregisterAfterTrigger), FunctionToRun))
 			{
 				Log.Info("New delegate registered: " + SenderFilter + " | " + CategoryFilter + " | " + PayloadFilter);
 			}
@@ -788,7 +803,7 @@ namespace Gauntlet
 				{
 					if (FunctionToRun == MessageDelegates[MessageFilter])
 					{
-						if(MessageDelegates.TryRemove(MessageFilter, out var Unused))
+						if (MessageDelegates.TryRemove(MessageFilter, out var Unused))
 						{
 							bDeregisteredDelegate = true;
 							Log.Verbose("Delegate removed for filter: Sender: " + SenderFilter + " | Category: " + CategoryFilter + " | Payload: " + PayloadFilter);
@@ -873,7 +888,7 @@ namespace Gauntlet
 			else
 			{
 				Log.Info(string.Format("Registering endpoint {0} | {1} with new delegate.", Path.Path, Path.Verb));
-				if(!HttpRoutes.TryAdd(Path, RouteDelegate))
+				if (!HttpRoutes.TryAdd(Path, RouteDelegate))
 				{
 					Log.Error($"Unexpectedly encountered existing RouteMapping for {Path}");
 				}
@@ -921,14 +936,14 @@ namespace Gauntlet
 
 					if (!InArgs.TryGetValue(ArgDesc.Name, out ArgValue))
 					{
-						if ( !ArgDesc.Optional )
+						if (!ArgDesc.Optional)
 						{
 							Log.Warning(string.Format("RPC {0} lists mandatory argument {1} of type {2} which was not provided. Please either repair argument list or update RPC Definition.",
 								RpcName, ArgDesc.Name, ArgDesc.Type));
 						}
 						continue;
 					}
-					
+
 					if (ArgValue == null)
 					{
 						Log.Warning(string.Format("RPC {0} argument {1} of expected type {2} is null. If null is a possible value please make value optional in RPC instead.",
@@ -1012,57 +1027,59 @@ namespace Gauntlet
 				return null;
 			}
 		}
-		
-		public static HttpResponseMessage CallRPCWithUrlParams(RpcTarget InTarget, string RpcName, Dictionary<string, object> InArgs, int Timeout = 5000,  string InContentType = "application/json")
-        {
         	Log.VeryVerbose($"Calling RPC {RpcName}");
-        	string RpcUrl = InTarget.CreateRpcUrl(RpcName);
-        	RpcDefinition RpcEntry = InTarget.GetRpc(RpcName);
-        	if (string.IsNullOrEmpty(RpcUrl) || RpcEntry == null)
-        	{
         		Log.Verbose($"Failed Calling RPC {RpcName}");
-        		return null;
-        	}
-            
-            //format the arguments for the url, sanitizing to remove invalid characters e.g. & = and /
-            //only special characters allowed are _ - and . 
-            string UrlArgs = string.Empty;
-            string SanitizingRegex = "[^a-zA-Z0-9._-]";
-            foreach (KeyValuePair<string, object> Argument in InArgs)
-            {
-	            string Key = Regex.Replace(Argument.Key, SanitizingRegex, string.Empty);
-	            if (Key != Argument.Key)
-	            {
-					Log.Verbose($"Argument Key {Argument.Key} sanitized to remove invalid characters. New Value: {Key}");    
-	            }
-
-	            string ArgumentValue = $"{Argument.Value}";
-	            string Value = Regex.Replace(ArgumentValue, SanitizingRegex, string.Empty);
-	            if (Value != ArgumentValue)
-	            {
-		            Log.Verbose($"Argument Value {ArgumentValue} sanitized to remove invalid characters. New Value: {Value}");   
-	            }
-	            
-	            UrlArgs += $"{Key}={Value}";
-	            if (!Argument.Equals(InArgs.Last()))
-	            {
-		            UrlArgs += "&";
-	            }
-            }
-            RpcUrl = String.Format($"{RpcUrl}?{UrlArgs}");
-            HttpMethod MethodToUse = new HttpMethod(RpcEntry.Verb);
-        	HttpRequestMessage RpcRequest = new HttpRequestMessage(MethodToUse, RpcUrl);
+
+		public static HttpResponseMessage CallRPCWithUrlParams(RpcTarget InTarget, string RpcName, Dictionary<string, object> InArgs, int Timeout = 5000, string InContentType = "application/json")
+		{
+			Log.VeryVerbose("Calling RPC" + RpcName);
+			string RpcUrl = InTarget.CreateRpcUrl(RpcName);
+			RpcDefinition RpcEntry = InTarget.GetRpc(RpcName);
+			if (string.IsNullOrEmpty(RpcUrl) || RpcEntry == null)
+			{
+				Log.Verbose("Failed Calling RPC" + RpcName);
+				return null;
+			}
+
+			//format the arguments for the url, sanitizing to remove invalid characters e.g. & = and /
+			//only special characters allowed are _ - and . 
+			string UrlArgs = string.Empty;
+			string SanitizingRegex = "[^a-zA-Z0-9._-]";
+			foreach (KeyValuePair<string, object> Argument in InArgs)
+			{
+				string Key = Regex.Replace(Argument.Key, SanitizingRegex, string.Empty);
+				if (Key != Argument.Key)
+				{
+					Log.Verbose($"Argument Key {Argument.Key} sanitized to remove invalid characters. New Value: {Key}");
+				}
+
+				string ArgumentValue = $"{Argument.Value}";
+				string Value = Regex.Replace(ArgumentValue, SanitizingRegex, string.Empty);
+				if (Value != ArgumentValue)
+				{
+					Log.Verbose($"Argument Value {ArgumentValue} sanitized to remove invalid characters. New Value: {Value}");
+				}
+
+				UrlArgs += $"{Key}={Value}";
+				if (!Argument.Equals(InArgs.Last()))
+				{
+					UrlArgs += "&";
+				}
+			}
+			RpcUrl = String.Format($"{RpcUrl}?{UrlArgs}");
+			HttpMethod MethodToUse = new HttpMethod(RpcEntry.Verb);
+			HttpRequestMessage RpcRequest = new HttpRequestMessage(MethodToUse, RpcUrl);
 			RpcRequest.Headers.Connection.Add("keep-alive");
 			RpcRequest.Headers.Add("Keep-Alive", "600");
 
 			// We might timeout calling CallRpc on a client that's closing - we do not want this to be a fail;
 			try
-        	{
+			{
 				HttpResponseMessage RpcResponse = ExecutorClient.Send(RpcRequest);
-        		return RpcResponse;
-        	}
-        	catch (Exception Error)
-        	{
+				return RpcResponse;
+			}
+			catch (Exception Error)
+			{
 				string Message = $"CallRpc for: {RpcName} failed with following exception: {Error}";
 
 				Log.SuspendECErrorParsing();
@@ -1075,8 +1092,8 @@ namespace Gauntlet
 					Log.Info(Message);
 				}
 				Log.ResumeECErrorParsing();
-        		return null;
-        	}
-        }
+				return null;
+			}
+		}
 	}
-}
\ No newline at end of file
+}
diff --git a/Engine/Source/Programs/AutomationTool/Gauntlet/Platform/IOS/Gauntlet.TargetDeviceIOS.cs b/Engine/Source/Programs/AutomationTool/Gauntlet/Platform/IOS/Gauntlet.TargetDeviceIOS.cs
--- a/Engine/Source/Programs/AutomationTool/Gauntlet/Platform/IOS/Gauntlet.TargetDeviceIOS.cs
+++ b/Engine/Source/Programs/AutomationTool/Gauntlet/Platform/IOS/Gauntlet.TargetDeviceIOS.cs
@@ -57,13 +57,15 @@ namespace Gauntlet
 			public string UUID;
 			public int IOSMajorVersion;
 			public bool bConnected;
+			public string IPAddress;
 			public ITargetDevice AssignedDevice;
 			public bool IsAssigned => AssignedDevice != null;
-			public ConnectedDevice(string UUID, int IOSMajorVersion, bool bConnected)
+			public ConnectedDevice(string UUID, int IOSMajorVersion, bool bConnected, string IPAddress = null)
 			{
 				this.UUID = UUID;
 				this.IOSMajorVersion = IOSMajorVersion;
 				this.bConnected = bConnected;
+				this.IPAddress = IPAddress;
 			}
 		}
 
@@ -103,6 +105,8 @@ namespace Gauntlet
 		#region ITargetDevice
 		public string Name { get; protected set; }
 
+		public string IPAddress { get; protected set; }
+
 		public UnrealTargetPlatform? Platform => UnrealTargetPlatform.IOS;
 
 		public ERunOptions RunOptions { get; set; }
@@ -209,6 +213,23 @@ namespace Gauntlet
 			}
 
 			ConnectedDevices[UUID].AssignedDevice = this;
+			IPAddress = ConnectedDevices[UUID].IPAddress;
+
+			// If IP address is not available from device list, try to query it directly
+			if (string.IsNullOrEmpty(IPAddress) && UseDeviceCtl)
+			{
+				Log.Info("Attempting to query device {0} for WiFi IP address...", UUID);
+				IPAddress = QueryDeviceIPAddress(UUID);
+				if (!string.IsNullOrEmpty(IPAddress))
+				{
+					Log.Info("Found WiFi IP address for device {0}: {1}", UUID, IPAddress);
+					ConnectedDevices[UUID].IPAddress = IPAddress;
+				}
+				else
+				{
+					Log.Verbose("Could not determine WiFi IP address for device {0}", UUID);
+				}
+			}
 
 			if (XCodeMajorVersion < 15)
 			{
@@ -1461,6 +1482,33 @@ namespace Gauntlet
 			return bIsConnected;
 		}
 
+		/// <summary>
+		/// Query the IP address of an iOS device from Megazord API using its UUID
+		/// </summary>
+		/// <param name="DeviceUUID">The device UUID/serial number</param>
+		/// <returns>IP address string if found, null otherwise</returns>
+		private static string QueryDeviceIPAddress(string DeviceUUID)
+		{
+			try
+			{
+				// Call Megazord API asynchronously and wait for result
+				System.Threading.Tasks.Task<string> IPAddressTask = MegazordService.GetDeviceIPAddressAsync(DeviceUUID);
+				IPAddressTask.Wait();
+
+				return IPAddressTask.Result;
+			}
+			catch (System.AggregateException ex)
+			{
+				Log.Verbose("Failed to query Megazord for device IP address: {0}", ex.InnerException?.Message ?? ex.Message);
+				return null;
+			}
+			catch (Exception ex)
+			{
+				Log.Verbose("Failed to query Megazord for device IP address: {0}", ex.Message);
+				return null;
+			}
+		}
+
 		// DEPRECATED
 		public IAppInstall InstallApplication(UnrealAppConfig AppConfig)
 		{
@@ -1669,6 +1717,21 @@ namespace Gauntlet
 			{
 				Log.Error(KnownLogEvents.Gauntlet_DeviceEvent, "Failed to move artifacts out of temp directory:\n{Exception}", Ex.Message);
 			}
+
+			// Copy Appium screenshots to artifact directory
+			if (IOSDevice.AppiumContainer != null && AppiumContainer.Enabled)
+			{
+				try
+				{
+					string ScreenshotsDir = Path.Combine(DestinationArtifactPath, "Screenshots");
+					Log.Info("Copying Appium screenshots to artifact directory: {0}", ScreenshotsDir);
+					IOSDevice.AppiumContainer.StopMonitorLoop(ScreenshotsDir);
+				}
+				catch (Exception Ex)
+				{
+					Log.Warning("Failed to copy Appium screenshots to artifact directory: {0}", Ex.Message);
+				}
+			}
 		}
 	}
 
