Automation: Only set OutputPath if we're not setting an ArtifactsPath. This means that we can easily build, say, all of the Automation projects to a specific location that we control, and should allow us to simplify our current uatobj / uatbin support, swapping it for just specifying an ArtifactsPath inside of UAT.

#Synced-CL 52723323

diff --git a/Engine/Platforms/VisionOS/Source/Programs/AutomationTool/VisionOS.Automation.csproj b/Engine/Platforms/VisionOS/Source/Programs/AutomationTool/VisionOS.Automation.csproj
--- a/Engine/Platforms/VisionOS/Source/Programs/AutomationTool/VisionOS.Automation.csproj
+++ b/Engine/Platforms/VisionOS/Source/Programs/AutomationTool/VisionOS.Automation.csproj
@@ -8,7 +8,7 @@
     <RootNamespace>VisionOS.Automation</RootNamespace>
     <AssemblyName>VisionOS.Automation</AssemblyName>
     <WarningsNotAsErrors>612,618</WarningsNotAsErrors>
-    <OutputPath>..\..\..\Binaries\DotNET\AutomationTool\AutomationScripts\</OutputPath>
+    <OutputPath Condition="'$(ArtifactsPath)' == ''">..\..\..\Binaries\DotNET\AutomationTool\AutomationScripts\</OutputPath>
     <TreatWarningsAsErrors>true</TreatWarningsAsErrors>
   </PropertyGroup>
 
diff --git a/Engine/Plugins/Performance/AutomatedPerfTesting/Build/Scripts/AutomatedPerfTesting.Automation.csproj b/Engine/Plugins/Performance/AutomatedPerfTesting/Build/Scripts/AutomatedPerfTesting.Automation.csproj
--- a/Engine/Plugins/Performance/AutomatedPerfTesting/Build/Scripts/AutomatedPerfTesting.Automation.csproj
+++ b/Engine/Plugins/Performance/AutomatedPerfTesting/Build/Scripts/AutomatedPerfTesting.Automation.csproj
@@ -6,7 +6,7 @@
     <GenerateTargetFrameworkAttribute>false</GenerateTargetFrameworkAttribute>
     <RootNamespace>AutomatedPerfTesting.Automation</RootNamespace>
     <AssemblyName>AutomatedPerfTesting.Automation</AssemblyName>
-    <OutputPath>$(EngineDir)\Binaries\DotNET\AutomationTool\AutomationScripts\AutomatedPerfTesting</OutputPath>
+    <OutputPath Condition="'$(ArtifactsPath)' == ''">$(EngineDir)\Binaries\DotNET\AutomationTool\AutomationScripts\AutomatedPerfTesting</OutputPath>
     <TreatWarningsAsErrors>true</TreatWarningsAsErrors>
     <WarningsNotAsErrors>612,618</WarningsNotAsErrors>
   </PropertyGroup>
diff --git a/Engine/Plugins/VirtualProduction/ICVFXTesting/Build/Scripts/ICVFX.Automation.csproj b/Engine/Plugins/VirtualProduction/ICVFXTesting/Build/Scripts/ICVFX.Automation.csproj
--- a/Engine/Plugins/VirtualProduction/ICVFXTesting/Build/Scripts/ICVFX.Automation.csproj
+++ b/Engine/Plugins/VirtualProduction/ICVFXTesting/Build/Scripts/ICVFX.Automation.csproj
@@ -8,7 +8,7 @@
     <RootNamespace>ICVFX.Automation</RootNamespace>
     <AssemblyName>ICVFX.Automation</AssemblyName>
     <WarningsNotAsErrors>612,618</WarningsNotAsErrors>
-    <OutputPath>..\..\..\..\..\Binaries\DotNET\AutomationTool\AutomationScripts\ICVFX</OutputPath>
+    <OutputPath Condition="'$(ArtifactsPath)' == ''">..\..\..\..\..\Binaries\DotNET\AutomationTool\AutomationScripts\ICVFX</OutputPath>
     <TreatWarningsAsErrors>true</TreatWarningsAsErrors>
   </PropertyGroup>
   <ItemGroup>
diff --git a/Engine/Source/Programs/AutomationTool/Android/Android.Automation.csproj b/Engine/Source/Programs/AutomationTool/Android/Android.Automation.csproj
--- a/Engine/Source/Programs/AutomationTool/Android/Android.Automation.csproj
+++ b/Engine/Source/Programs/AutomationTool/Android/Android.Automation.csproj
@@ -8,7 +8,7 @@
     <RootNamespace>Android.Automation</RootNamespace>
     <AssemblyName>Android.Automation</AssemblyName>
     <WarningsNotAsErrors>612,618</WarningsNotAsErrors>
-    <OutputPath>..\..\..\..\Binaries\DotNET\AutomationTool\AutomationScripts\Platforms\Android</OutputPath>
+    <OutputPath Condition="'$(ArtifactsPath)' == ''">..\..\..\..\Binaries\DotNET\AutomationTool\AutomationScripts\Platforms\Android</OutputPath>
     <TreatWarningsAsErrors>true</TreatWarningsAsErrors>
   </PropertyGroup>
   <ItemGroup>
diff --git a/Engine/Source/Programs/AutomationTool/Apple/Apple.Automation.csproj b/Engine/Source/Programs/AutomationTool/Apple/Apple.Automation.csproj
--- a/Engine/Source/Programs/AutomationTool/Apple/Apple.Automation.csproj
+++ b/Engine/Source/Programs/AutomationTool/Apple/Apple.Automation.csproj
@@ -8,7 +8,7 @@
     <RootNamespace>Apple.Automation</RootNamespace>
     <AssemblyName>Apple.Automation</AssemblyName>
     <WarningsNotAsErrors>612,618</WarningsNotAsErrors>
-    <OutputPath>..\..\..\..\Binaries\DotNET\AutomationTool\AutomationScripts\Platforms\Apple</OutputPath>
+    <OutputPath Condition="'$(ArtifactsPath)' == ''">..\..\..\..\Binaries\DotNET\AutomationTool\AutomationScripts\Platforms\Apple</OutputPath>
     <TreatWarningsAsErrors>true</TreatWarningsAsErrors>
   </PropertyGroup>
   <ItemGroup>
diff --git a/Engine/Source/Programs/AutomationTool/AutomationUtils/AutomationUtils.Automation.csproj b/Engine/Source/Programs/AutomationTool/AutomationUtils/AutomationUtils.Automation.csproj
--- a/Engine/Source/Programs/AutomationTool/AutomationUtils/AutomationUtils.Automation.csproj
+++ b/Engine/Source/Programs/AutomationTool/AutomationUtils/AutomationUtils.Automation.csproj
@@ -11,7 +11,7 @@
     <CopyLocalLockFileAssemblies>true</CopyLocalLockFileAssemblies>
   </PropertyGroup>
   <PropertyGroup Condition="'$(EpicGamesMsBuild)' == 'true'">
-    <OutputPath>..\..\..\..\Binaries\DotNET\AutomationTool\AutomationUtils</OutputPath>
+    <OutputPath Condition="'$(ArtifactsPath)' == ''">..\..\..\..\Binaries\DotNET\AutomationTool\AutomationUtils</OutputPath>
   </PropertyGroup>
 
   <!-- The following will glob all platform extensions' AutomationUtils source
diff --git a/Engine/Source/Programs/AutomationTool/BuildGraph/BuildGraph.Automation.csproj b/Engine/Source/Programs/AutomationTool/BuildGraph/BuildGraph.Automation.csproj
--- a/Engine/Source/Programs/AutomationTool/BuildGraph/BuildGraph.Automation.csproj
+++ b/Engine/Source/Programs/AutomationTool/BuildGraph/BuildGraph.Automation.csproj
@@ -8,7 +8,7 @@
     <RootNamespace>AutomationTool</RootNamespace>
     <AssemblyName>BuildGraph.Automation</AssemblyName>
     <WarningsNotAsErrors>612,618</WarningsNotAsErrors>
-    <OutputPath>..\..\..\..\Binaries\DotNET\AutomationTool\AutomationScripts\BuildGraph</OutputPath>
+    <OutputPath Condition="'$(ArtifactsPath)' == ''">..\..\..\..\Binaries\DotNET\AutomationTool\AutomationScripts\BuildGraph</OutputPath>
     <TreatWarningsAsErrors>true</TreatWarningsAsErrors>
     <GenerateDocumentationFile>True</GenerateDocumentationFile>
   </PropertyGroup>
diff --git a/Engine/Source/Programs/AutomationTool/CookedEditor/CookedEditor.Automation.csproj b/Engine/Source/Programs/AutomationTool/CookedEditor/CookedEditor.Automation.csproj
--- a/Engine/Source/Programs/AutomationTool/CookedEditor/CookedEditor.Automation.csproj
+++ b/Engine/Source/Programs/AutomationTool/CookedEditor/CookedEditor.Automation.csproj
@@ -8,7 +8,7 @@
     <RootNamespace>AutomationTool</RootNamespace>
     <AssemblyName>CookedEditor.Automation</AssemblyName>
     <WarningsNotAsErrors>612,618</WarningsNotAsErrors>
-    <OutputPath>..\..\..\..\Binaries\DotNET\AutomationTool\AutomationScripts\CookedEditor</OutputPath>
+    <OutputPath Condition="'$(ArtifactsPath)' == ''">..\..\..\..\Binaries\DotNET\AutomationTool\AutomationScripts\CookedEditor</OutputPath>
     <TreatWarningsAsErrors>true</TreatWarningsAsErrors>
   </PropertyGroup>
   <ItemGroup>
diff --git a/Engine/Source/Programs/AutomationTool/CrowdinLocalization/CrowdinLocalization.Automation.csproj b/Engine/Source/Programs/AutomationTool/CrowdinLocalization/CrowdinLocalization.Automation.csproj
--- a/Engine/Source/Programs/AutomationTool/CrowdinLocalization/CrowdinLocalization.Automation.csproj
+++ b/Engine/Source/Programs/AutomationTool/CrowdinLocalization/CrowdinLocalization.Automation.csproj
@@ -8,7 +8,7 @@
     <RootNamespace>AutomationTool</RootNamespace>
     <AssemblyName>CrowdinLocalization.Automation</AssemblyName>
     <WarningsNotAsErrors>612,618</WarningsNotAsErrors>
-    <OutputPath>..\..\..\..\Binaries\DotNET\AutomationTool\AutomationScripts\CrowdinLocalization</OutputPath>
+    <OutputPath Condition="'$(ArtifactsPath)' == ''">..\..\..\..\Binaries\DotNET\AutomationTool\AutomationScripts\CrowdinLocalization</OutputPath>
     <NoWarn>$(NoWarn);1701;1702;CS1998;CS0649</NoWarn>
     <TreatWarningsAsErrors>true</TreatWarningsAsErrors>
     <CopyLocalLockFileAssemblies>true</CopyLocalLockFileAssemblies>
diff --git a/Engine/Source/Programs/AutomationTool/Gauntlet/Gauntlet.Automation.csproj b/Engine/Source/Programs/AutomationTool/Gauntlet/Gauntlet.Automation.csproj
--- a/Engine/Source/Programs/AutomationTool/Gauntlet/Gauntlet.Automation.csproj
+++ b/Engine/Source/Programs/AutomationTool/Gauntlet/Gauntlet.Automation.csproj
@@ -8,7 +8,7 @@
     <RootNamespace>Gauntlet.Automation</RootNamespace>
     <AssemblyName>Gauntlet.Automation</AssemblyName>
     <WarningsNotAsErrors>612,618</WarningsNotAsErrors>
-    <OutputPath>..\..\..\..\Binaries\DotNET\AutomationTool\AutomationScripts\Gauntlet</OutputPath>
+    <OutputPath Condition="'$(ArtifactsPath)' == ''">..\..\..\..\Binaries\DotNET\AutomationTool\AutomationScripts\Gauntlet</OutputPath>
     <TreatWarningsAsErrors>true</TreatWarningsAsErrors>
     <CopyLocalLockFileAssemblies>true</CopyLocalLockFileAssemblies>
     <SatelliteResourceLanguages>en</SatelliteResourceLanguages>    <!-- remove non english resource languages -->
diff --git a/Engine/Source/Programs/AutomationTool/IOS/IOS.Automation.csproj b/Engine/Source/Programs/AutomationTool/IOS/IOS.Automation.csproj
--- a/Engine/Source/Programs/AutomationTool/IOS/IOS.Automation.csproj
+++ b/Engine/Source/Programs/AutomationTool/IOS/IOS.Automation.csproj
@@ -8,7 +8,7 @@
     <RootNamespace>IOS.Automation</RootNamespace>
     <AssemblyName>IOS.Automation</AssemblyName>
     <WarningsNotAsErrors>612,618</WarningsNotAsErrors>
-    <OutputPath>..\..\..\..\Binaries\DotNET\AutomationTool\AutomationScripts\Platforms\IOS</OutputPath>
+    <OutputPath Condition="'$(ArtifactsPath)' == ''">..\..\..\..\Binaries\DotNET\AutomationTool\AutomationScripts\Platforms\IOS</OutputPath>
     <TreatWarningsAsErrors>true</TreatWarningsAsErrors>
   </PropertyGroup>
   <ItemGroup>
diff --git a/Engine/Source/Programs/AutomationTool/Insights/Insights.Automation.csproj b/Engine/Source/Programs/AutomationTool/Insights/Insights.Automation.csproj
--- a/Engine/Source/Programs/AutomationTool/Insights/Insights.Automation.csproj
+++ b/Engine/Source/Programs/AutomationTool/Insights/Insights.Automation.csproj
@@ -7,7 +7,7 @@
     <AssemblyName>Insights.Automation</AssemblyName>
     <GenerateAssemblyInfo>false</GenerateAssemblyInfo>
     <GenerateTargetFrameworkAttribute>false</GenerateTargetFrameworkAttribute>
-    <OutputPath>..\..\..\..\Binaries\DotNET\AutomationTool\AutomationScripts\Insights</OutputPath>
+    <OutputPath Condition="'$(ArtifactsPath)' == ''">..\..\..\..\Binaries\DotNET\AutomationTool\AutomationScripts\Insights</OutputPath>
     <TreatWarningsAsErrors>true</TreatWarningsAsErrors>
     <SatelliteResourceLanguages>en</SatelliteResourceLanguages>    <!-- remove non english resource languages -->
   </PropertyGroup>
diff --git a/Engine/Source/Programs/AutomationTool/Linux/Linux.Automation.csproj b/Engine/Source/Programs/AutomationTool/Linux/Linux.Automation.csproj
--- a/Engine/Source/Programs/AutomationTool/Linux/Linux.Automation.csproj
+++ b/Engine/Source/Programs/AutomationTool/Linux/Linux.Automation.csproj
@@ -8,7 +8,7 @@
     <RootNamespace>Linux.Automation</RootNamespace>
     <AssemblyName>Linux.Automation</AssemblyName>
     <WarningsNotAsErrors>612,618</WarningsNotAsErrors>
-    <OutputPath>..\..\..\..\Binaries\DotNET\AutomationTool\AutomationScripts\Platforms\Linux\</OutputPath>
+    <OutputPath Condition="'$(ArtifactsPath)' == ''">..\..\..\..\Binaries\DotNET\AutomationTool\AutomationScripts\Platforms\Linux\</OutputPath>
     <TreatWarningsAsErrors>true</TreatWarningsAsErrors>
   </PropertyGroup>
 
diff --git a/Engine/Source/Programs/AutomationTool/LiveLinkHub/LiveLinkHub.Automation.csproj b/Engine/Source/Programs/AutomationTool/LiveLinkHub/LiveLinkHub.Automation.csproj
--- a/Engine/Source/Programs/AutomationTool/LiveLinkHub/LiveLinkHub.Automation.csproj
+++ b/Engine/Source/Programs/AutomationTool/LiveLinkHub/LiveLinkHub.Automation.csproj
@@ -5,7 +5,7 @@
     <OutputType>Library</OutputType>
     <GenerateAssemblyInfo>false</GenerateAssemblyInfo>
     <GenerateTargetFrameworkAttribute>false</GenerateTargetFrameworkAttribute>
-    <OutputPath>..\..\..\..\Binaries\DotNET\AutomationTool\AutomationScripts\LiveLinkHub</OutputPath>
+    <OutputPath Condition="'$(ArtifactsPath)' == ''">..\..\..\..\Binaries\DotNET\AutomationTool\AutomationScripts\LiveLinkHub</OutputPath>
   </PropertyGroup>
 
   <ItemGroup>
diff --git a/Engine/Source/Programs/AutomationTool/Localization/Localization.Automation.csproj b/Engine/Source/Programs/AutomationTool/Localization/Localization.Automation.csproj
--- a/Engine/Source/Programs/AutomationTool/Localization/Localization.Automation.csproj
+++ b/Engine/Source/Programs/AutomationTool/Localization/Localization.Automation.csproj
@@ -8,7 +8,7 @@
     <RootNamespace>AutomationTool</RootNamespace>
     <AssemblyName>Localization.Automation</AssemblyName>
     <WarningsNotAsErrors>612,618</WarningsNotAsErrors>
-    <OutputPath>..\..\..\..\Binaries\DotNET\AutomationTool\AutomationScripts\Localization</OutputPath>
+    <OutputPath Condition="'$(ArtifactsPath)' == ''">..\..\..\..\Binaries\DotNET\AutomationTool\AutomationScripts\Localization</OutputPath>
     <TreatWarningsAsErrors>true</TreatWarningsAsErrors>
   </PropertyGroup>
 
diff --git a/Engine/Source/Programs/AutomationTool/LowLevelTests/LowLevelTests.Automation.csproj b/Engine/Source/Programs/AutomationTool/LowLevelTests/LowLevelTests.Automation.csproj
--- a/Engine/Source/Programs/AutomationTool/LowLevelTests/LowLevelTests.Automation.csproj
+++ b/Engine/Source/Programs/AutomationTool/LowLevelTests/LowLevelTests.Automation.csproj
@@ -7,7 +7,7 @@
     <AssemblyName>LowLevelTests.Automation</AssemblyName>
     <GenerateAssemblyInfo>false</GenerateAssemblyInfo>
     <GenerateTargetFrameworkAttribute>false</GenerateTargetFrameworkAttribute>
-    <OutputPath>..\..\..\..\Binaries\DotNET\AutomationTool\AutomationScripts\LowLevelTests</OutputPath>
+    <OutputPath Condition="'$(ArtifactsPath)' == ''">..\..\..\..\Binaries\DotNET\AutomationTool\AutomationScripts\LowLevelTests</OutputPath>
     <TreatWarningsAsErrors>true</TreatWarningsAsErrors>
     <SatelliteResourceLanguages>en</SatelliteResourceLanguages>    <!-- remove non english resource languages -->
   </PropertyGroup>
diff --git a/Engine/Source/Programs/AutomationTool/Mac/Mac.Automation.csproj b/Engine/Source/Programs/AutomationTool/Mac/Mac.Automation.csproj
--- a/Engine/Source/Programs/AutomationTool/Mac/Mac.Automation.csproj
+++ b/Engine/Source/Programs/AutomationTool/Mac/Mac.Automation.csproj
@@ -8,7 +8,7 @@
     <RootNamespace>Mac.Automation</RootNamespace>
     <AssemblyName>Mac.Automation</AssemblyName>
     <WarningsNotAsErrors>612,618</WarningsNotAsErrors>
-    <OutputPath>..\..\..\..\Binaries\DotNET\AutomationTool\AutomationScripts\Platforms\Mac\</OutputPath>
+    <OutputPath Condition="'$(ArtifactsPath)' == ''">..\..\..\..\Binaries\DotNET\AutomationTool\AutomationScripts\Platforms\Mac\</OutputPath>
     <TreatWarningsAsErrors>true</TreatWarningsAsErrors>
   </PropertyGroup>
   <ItemGroup>
diff --git a/Engine/Source/Programs/AutomationTool/Microsoft/Microsoft.Automation.csproj b/Engine/Source/Programs/AutomationTool/Microsoft/Microsoft.Automation.csproj
--- a/Engine/Source/Programs/AutomationTool/Microsoft/Microsoft.Automation.csproj
+++ b/Engine/Source/Programs/AutomationTool/Microsoft/Microsoft.Automation.csproj
@@ -8,7 +8,7 @@
     <RootNamespace>Microsoft.Automation</RootNamespace>
     <AssemblyName>Microsoft.Automation</AssemblyName>
     <WarningsNotAsErrors>612,618</WarningsNotAsErrors>
-    <OutputPath>..\..\..\..\Binaries\DotNET\AutomationTool\AutomationScripts\Platforms\Microsoft</OutputPath>
+    <OutputPath Condition="'$(ArtifactsPath)' == ''">..\..\..\..\Binaries\DotNET\AutomationTool\AutomationScripts\Platforms\Microsoft</OutputPath>
     <TreatWarningsAsErrors>true</TreatWarningsAsErrors>
   </PropertyGroup>
   <ItemGroup>
diff --git a/Engine/Source/Programs/AutomationTool/Mutable/Mutable.Automation.csproj b/Engine/Source/Programs/AutomationTool/Mutable/Mutable.Automation.csproj
--- a/Engine/Source/Programs/AutomationTool/Mutable/Mutable.Automation.csproj
+++ b/Engine/Source/Programs/AutomationTool/Mutable/Mutable.Automation.csproj
@@ -7,7 +7,7 @@
     <Nullable>enable</Nullable>
     <GenerateAssemblyInfo>false</GenerateAssemblyInfo>
     <GenerateTargetFrameworkAttribute>false</GenerateTargetFrameworkAttribute>
-    <OutputPath>..\..\..\..\Binaries\DotNET\AutomationTool\AutomationScripts\Mutable\</OutputPath>
+    <OutputPath Condition="'$(ArtifactsPath)' == ''">..\..\..\..\Binaries\DotNET\AutomationTool\AutomationScripts\Mutable\</OutputPath>
   </PropertyGroup>
 
   <ItemGroup>
diff --git a/Engine/Source/Programs/AutomationTool/OneSkyLocalization/OneSkyLocalization.Automation.csproj b/Engine/Source/Programs/AutomationTool/OneSkyLocalization/OneSkyLocalization.Automation.csproj
--- a/Engine/Source/Programs/AutomationTool/OneSkyLocalization/OneSkyLocalization.Automation.csproj
+++ b/Engine/Source/Programs/AutomationTool/OneSkyLocalization/OneSkyLocalization.Automation.csproj
@@ -8,7 +8,7 @@
     <RootNamespace>AutomationTool</RootNamespace>
     <AssemblyName>OneSkyLocalization.Automation</AssemblyName>
     <WarningsNotAsErrors>612,618</WarningsNotAsErrors>
-    <OutputPath>..\..\..\..\Binaries\DotNET\AutomationTool\AutomationScripts\OneSkyLocalization</OutputPath>
+    <OutputPath Condition="'$(ArtifactsPath)' == ''">..\..\..\..\Binaries\DotNET\AutomationTool\AutomationScripts\OneSkyLocalization</OutputPath>
     <TreatWarningsAsErrors>true</TreatWarningsAsErrors>
   </PropertyGroup>
   <ItemGroup>
diff --git a/Engine/Source/Programs/AutomationTool/Scripts/AutomationScripts.Automation.csproj b/Engine/Source/Programs/AutomationTool/Scripts/AutomationScripts.Automation.csproj
--- a/Engine/Source/Programs/AutomationTool/Scripts/AutomationScripts.Automation.csproj
+++ b/Engine/Source/Programs/AutomationTool/Scripts/AutomationScripts.Automation.csproj
@@ -8,7 +8,7 @@
     <RootNamespace>AutomationTool</RootNamespace>
     <AssemblyName>AutomationScripts.Automation</AssemblyName>
     <WarningsNotAsErrors>612,618</WarningsNotAsErrors>
-    <OutputPath>..\..\..\..\Binaries\DotNET\AutomationTool\AutomationScripts\Scripts</OutputPath>
+    <OutputPath Condition="'$(ArtifactsPath)' == ''">..\..\..\..\Binaries\DotNET\AutomationTool\AutomationScripts\Scripts</OutputPath>
     <TreatWarningsAsErrors>true</TreatWarningsAsErrors>
   </PropertyGroup>
   <ItemGroup>
diff --git a/Engine/Source/Programs/AutomationTool/SmartlingLocalization/SmartlingLocalization.Automation.csproj b/Engine/Source/Programs/AutomationTool/SmartlingLocalization/SmartlingLocalization.Automation.csproj
--- a/Engine/Source/Programs/AutomationTool/SmartlingLocalization/SmartlingLocalization.Automation.csproj
+++ b/Engine/Source/Programs/AutomationTool/SmartlingLocalization/SmartlingLocalization.Automation.csproj
@@ -8,7 +8,7 @@
     <RootNamespace>AutomationTool</RootNamespace>
     <AssemblyName>SmartlingLocalization.Automation</AssemblyName>
     <WarningsNotAsErrors>612,618</WarningsNotAsErrors>
-    <OutputPath>..\..\..\..\Binaries\DotNET\AutomationTool\AutomationScripts\SmartlingLocalization</OutputPath>
+    <OutputPath Condition="'$(ArtifactsPath)' == ''">..\..\..\..\Binaries\DotNET\AutomationTool\AutomationScripts\SmartlingLocalization</OutputPath>
     <TreatWarningsAsErrors>true</TreatWarningsAsErrors>
     <CopyLocalLockFileAssemblies>true</CopyLocalLockFileAssemblies>
     <SatelliteResourceLanguages>en</SatelliteResourceLanguages> <!-- remove non english resource languages -->
diff --git a/Engine/Source/Programs/AutomationTool/SteamDeck/SteamDeck.Automation.csproj b/Engine/Source/Programs/AutomationTool/SteamDeck/SteamDeck.Automation.csproj
--- a/Engine/Source/Programs/AutomationTool/SteamDeck/SteamDeck.Automation.csproj
+++ b/Engine/Source/Programs/AutomationTool/SteamDeck/SteamDeck.Automation.csproj
@@ -8,7 +8,7 @@
     <RootNamespace>SteamDeck.Automation</RootNamespace>
     <AssemblyName>SteamDeck.Automation</AssemblyName>
     <WarningsNotAsErrors>612,618</WarningsNotAsErrors>
-    <OutputPath>..\..\..\..\Binaries\DotNET\AutomationTool\AutomationScripts\Platforms\SteamDeck\</OutputPath>
+    <OutputPath Condition="'$(ArtifactsPath)' == ''">..\..\..\..\Binaries\DotNET\AutomationTool\AutomationScripts\Platforms\SteamDeck\</OutputPath>
     <TreatWarningsAsErrors>true</TreatWarningsAsErrors>
   </PropertyGroup>
 
diff --git a/Engine/Source/Programs/AutomationTool/TVOS/TVOS.Automation.csproj b/Engine/Source/Programs/AutomationTool/TVOS/TVOS.Automation.csproj
--- a/Engine/Source/Programs/AutomationTool/TVOS/TVOS.Automation.csproj
+++ b/Engine/Source/Programs/AutomationTool/TVOS/TVOS.Automation.csproj
@@ -8,7 +8,7 @@
     <RootNamespace>TVOS.Automation</RootNamespace>
     <AssemblyName>TVOS.Automation</AssemblyName>
     <WarningsNotAsErrors>612,618</WarningsNotAsErrors>
-    <OutputPath>..\..\..\..\Binaries\DotNET\AutomationTool\AutomationScripts\Platforms\TVOS\</OutputPath>
+    <OutputPath Condition="'$(ArtifactsPath)' == ''">..\..\..\..\Binaries\DotNET\AutomationTool\AutomationScripts\Platforms\TVOS\</OutputPath>
     <TreatWarningsAsErrors>true</TreatWarningsAsErrors>
   </PropertyGroup>
 
diff --git a/Engine/Source/Programs/AutomationTool/Turnkey/Turnkey.Automation.csproj b/Engine/Source/Programs/AutomationTool/Turnkey/Turnkey.Automation.csproj
--- a/Engine/Source/Programs/AutomationTool/Turnkey/Turnkey.Automation.csproj
+++ b/Engine/Source/Programs/AutomationTool/Turnkey/Turnkey.Automation.csproj
@@ -8,7 +8,7 @@
     <RootNamespace>Turnkey.Automation</RootNamespace>
     <AssemblyName>Turnkey.Automation</AssemblyName>
     <WarningsNotAsErrors>612,618</WarningsNotAsErrors>
-    <OutputPath>..\..\..\..\Binaries\DotNET\AutomationTool\AutomationScripts\Turnkey\</OutputPath>
+    <OutputPath Condition="'$(ArtifactsPath)' == ''">..\..\..\..\Binaries\DotNET\AutomationTool\AutomationScripts\Turnkey\</OutputPath>
     <TreatWarningsAsErrors>true</TreatWarningsAsErrors>
     <CopyLocalLockFileAssemblies>true</CopyLocalLockFileAssemblies>
     <SatelliteResourceLanguages>en</SatelliteResourceLanguages>    <!-- remove non english resource languages -->
diff --git a/Engine/Source/Programs/AutomationTool/Win/Win.Automation.csproj b/Engine/Source/Programs/AutomationTool/Win/Win.Automation.csproj
--- a/Engine/Source/Programs/AutomationTool/Win/Win.Automation.csproj
+++ b/Engine/Source/Programs/AutomationTool/Win/Win.Automation.csproj
@@ -8,7 +8,7 @@
     <RootNamespace>Win.Automation</RootNamespace>
     <AssemblyName>Win.Automation</AssemblyName>
     <WarningsNotAsErrors>612,618</WarningsNotAsErrors>
-    <OutputPath>..\..\..\..\Binaries\DotNET\AutomationTool\AutomationScripts\Platforms\Windows</OutputPath>
+    <OutputPath Condition="'$(ArtifactsPath)' == ''">..\..\..\..\Binaries\DotNET\AutomationTool\AutomationScripts\Platforms\Windows</OutputPath>
     <TreatWarningsAsErrors>true</TreatWarningsAsErrors>
   </PropertyGroup>
 
diff --git a/Engine/Source/Programs/AutomationTool/XLocLocalization/XLocLocalization.Automation.csproj b/Engine/Source/Programs/AutomationTool/XLocLocalization/XLocLocalization.Automation.csproj
--- a/Engine/Source/Programs/AutomationTool/XLocLocalization/XLocLocalization.Automation.csproj
+++ b/Engine/Source/Programs/AutomationTool/XLocLocalization/XLocLocalization.Automation.csproj
@@ -8,7 +8,7 @@
     <RootNamespace>AutomationTool</RootNamespace>
     <AssemblyName>XLocLocalization.Automation</AssemblyName>
     <WarningsNotAsErrors>612,618</WarningsNotAsErrors>
-    <OutputPath>..\..\..\..\Binaries\DotNET\AutomationTool\AutomationScripts\XLocLocalization</OutputPath>
+    <OutputPath Condition="'$(ArtifactsPath)' == ''">..\..\..\..\Binaries\DotNET\AutomationTool\AutomationScripts\XLocLocalization</OutputPath>
     <TreatWarningsAsErrors>true</TreatWarningsAsErrors>
     <CopyLocalLockFileAssemblies>true</CopyLocalLockFileAssemblies>
     <SatelliteResourceLanguages>en</SatelliteResourceLanguages> <!-- remove non english resource languages -->
