Update GetListenAddressAndPort to work on mac agents
Get IPAddress from Megazord
Add IPAddress lookup for ios in RPC Smoke
#Synced-CL 52690152

#Synced-CL 52690152

diff --git a/Engine/Source/Programs/AutomationTool/Gauntlet/Framework/RpcFramework/Gauntlet.RpcExecutor.cs b/Engine/Source/Programs/AutomationTool/Gauntlet/Framework/RpcFramework/Gauntlet.RpcExecutor.cs
--- a/Engine/Source/Programs/AutomationTool/Gauntlet/Framework/RpcFramework/Gauntlet.RpcExecutor.cs
+++ b/Engine/Source/Programs/AutomationTool/Gauntlet/Framework/RpcFramework/Gauntlet.RpcExecutor.cs
@@ -10,6 +10,7 @@ using System.IO;
 using System.Linq;
 using System.Net;
 using System.Net.Http;
+using System.Net.NetworkInformation;
 using System.Net.Sockets;
 using System.Runtime.CompilerServices;
 using System.Text;
@@ -162,15 +163,29 @@ namespace Gauntlet
 		public static string GetListenAddressAndPort()
 		{
 			string LocalIP = "127.0.0.1";
-			IPHostEntry HostEntry = Dns.GetHostEntry(Dns.GetHostName());
-			foreach (var TargetIp in HostEntry.AddressList)
+			try
 			{
-				if (TargetIp.AddressFamily == AddressFamily.InterNetwork)
+				IEnumerable<NetworkInterface> NetworkInterfaces = NetworkInterface.GetAllNetworkInterfaces()
+					.Where(ni => ni.OperationalStatus == OperationalStatus.Up &&
+								ni.NetworkInterfaceType != NetworkInterfaceType.Loopback);
+
+				foreach (var ni in NetworkInterfaces)
 				{
-					LocalIP = TargetIp.ToString();
-					break;
+					IPInterfaceProperties IpProperties = ni.GetIPProperties();
+					UnicastIPAddressInformation Ipv4Address = IpProperties.UnicastAddresses
+						.FirstOrDefault(ip => ip.Address.AddressFamily == AddressFamily.InterNetwork);
+
+					if (Ipv4Address != null)
+					{
+						LocalIP = Ipv4Address.Address.ToString();
+						break;
+					}
 				}
 			}
+			catch (Exception ex)
+			{
+				Log.Verbose($"Failed to get network IP, using loopback: {ex.Message}");
+			}
 			return String.Format("{0}:{1}", LocalIP, GauntletRPCListenPort);
 		}
 
diff --git a/Engine/Source/Programs/AutomationTool/Gauntlet/Platform/IOS/Gauntlet.TargetDeviceIOS.cs b/Engine/Source/Programs/AutomationTool/Gauntlet/Platform/IOS/Gauntlet.TargetDeviceIOS.cs
--- a/Engine/Source/Programs/AutomationTool/Gauntlet/Platform/IOS/Gauntlet.TargetDeviceIOS.cs
+++ b/Engine/Source/Programs/AutomationTool/Gauntlet/Platform/IOS/Gauntlet.TargetDeviceIOS.cs
@@ -57,13 +57,15 @@ namespace Gauntlet
 			public string UUID;
 			public int IOSMajorVersion;
 			public bool bConnected;
+			public string IPAddress;
 			public ITargetDevice AssignedDevice;
 			public bool IsAssigned => AssignedDevice != null;
-			public ConnectedDevice(string UUID, int IOSMajorVersion, bool bConnected)
+			public ConnectedDevice(string UUID, int IOSMajorVersion, bool bConnected, string IPAddress = null)
 			{
 				this.UUID = UUID;
 				this.IOSMajorVersion = IOSMajorVersion;
 				this.bConnected = bConnected;
+				this.IPAddress = IPAddress;
 			}
 		}
 
@@ -103,6 +105,8 @@ namespace Gauntlet
 		#region ITargetDevice
 		public string Name { get; protected set; }
 
+		public string IPAddress { get; protected set; }
+
 		public UnrealTargetPlatform? Platform => UnrealTargetPlatform.IOS;
 
 		public ERunOptions RunOptions { get; set; }
@@ -209,6 +213,23 @@ namespace Gauntlet
 			}
 
 			ConnectedDevices[UUID].AssignedDevice = this;
+			IPAddress = ConnectedDevices[UUID].IPAddress;
+
+			// If IP address is not available from device list, try to query it directly
+			if (string.IsNullOrEmpty(IPAddress) && UseDeviceCtl)
+			{
+				Log.Info("Attempting to query device {0} for WiFi IP address...", UUID);
+				IPAddress = QueryDeviceIPAddress(UUID);
+				if (!string.IsNullOrEmpty(IPAddress))
+				{
+					Log.Info("Found WiFi IP address for device {0}: {1}", UUID, IPAddress);
+					ConnectedDevices[UUID].IPAddress = IPAddress;
+				}
+				else
+				{
+					Log.Verbose("Could not determine WiFi IP address for device {0}", UUID);
+				}
+			}
 
 			if (XCodeMajorVersion < 15)
 			{
@@ -1461,6 +1482,33 @@ namespace Gauntlet
 			return bIsConnected;
 		}
 
+		/// <summary>
+		/// Query the IP address of an iOS device from Megazord API using its UUID
+		/// </summary>
+		/// <param name="DeviceUUID">The device UUID/serial number</param>
+		/// <returns>IP address string if found, null otherwise</returns>
+		private static string QueryDeviceIPAddress(string DeviceUUID)
+		{
+			try
+			{
+				// Call Megazord API asynchronously and wait for result
+				System.Threading.Tasks.Task<string> IPAddressTask = MegazordService.GetDeviceIPAddressAsync(DeviceUUID);
+				IPAddressTask.Wait();
+
+				return IPAddressTask.Result;
+			}
+			catch (System.AggregateException ex)
+			{
+				Log.Verbose("Failed to query Megazord for device IP address: {0}", ex.InnerException?.Message ?? ex.Message);
+				return null;
+			}
+			catch (Exception ex)
+			{
+				Log.Verbose("Failed to query Megazord for device IP address: {0}", ex.Message);
+				return null;
+			}
+		}
+
 		// DEPRECATED
 		public IAppInstall InstallApplication(UnrealAppConfig AppConfig)
 		{
@@ -1669,6 +1717,21 @@ namespace Gauntlet
 			{
 				Log.Error(KnownLogEvents.Gauntlet_DeviceEvent, "Failed to move artifacts out of temp directory:\n{Exception}", Ex.Message);
 			}
+
+			// Copy Appium screenshots to artifact directory
+			if (IOSDevice.AppiumContainer != null && AppiumContainer.Enabled)
+			{
+				try
+				{
+					string ScreenshotsDir = Path.Combine(DestinationArtifactPath, "Screenshots");
+					Log.Info("Copying Appium screenshots to artifact directory: {0}", ScreenshotsDir);
+					IOSDevice.AppiumContainer.StopMonitorLoop(ScreenshotsDir);
+				}
+				catch (Exception Ex)
+				{
+					Log.Warning("Failed to copy Appium screenshots to artifact directory: {0}", Ex.Message);
+				}
+			}
 		}
 	}
 
